* Adding the initial UMG Editor and runtime modules.  Currently it's all a lot of placeholders.  You can enable the UMG editor by using the -umg command.

diff --git a/Engine/Source/Editor/LevelEditor/Private/SLevelEditor.cpp b/Engine/Source/Editor/LevelEditor/Private/SLevelEditor.cpp
--- a/Engine/Source/Editor/LevelEditor/Private/SLevelEditor.cpp
+++ b/Engine/Source/Editor/LevelEditor/Private/SLevelEditor.cpp
@@ -32,6 +32,7 @@
 #include "Editor/WorkspaceMenuStructure/Public/WorkspaceMenuStructureModule.h"
 #include "Editor/Sequencer/Public/ISequencerModule.h"
 #include "Editor/StatsViewer/Public/StatsViewerModule.h"
+#include "Editor/UMGEditor/Public/UMGEditorModule.h"
 #include "EditorModes.h"
 #include "STutorialWrapper.h"
 #include "IDocumentation.h"
@@ -157,6 +158,11 @@ void SLevelEditor::Construct( const SLevelEditor::FArguments& InArgs)
 
 	FEditorDelegates::MapChange.AddSP(this, &SLevelEditor::HandleEditorMapChange);
 	HandleEditorMapChange(MapChangeEventFlags::NewMap);
+
+	if ( FParse::Param(FCommandLine::Get(), TEXT("umg")) )
+	{
+		FModuleManager::LoadModuleChecked<IUMGEditorModule>("UMGEditor");
+	}
 }
 
 void SLevelEditor::Initialize( const TSharedRef<SDockTab>& OwnerTab, const TSharedRef<SWindow>& OwnerWindow )
diff --git a/Engine/Source/Editor/UMGEditor/Classes/GUIPageFactory.h b/Engine/Source/Editor/UMGEditor/Classes/GUIPageFactory.h
new file mode 100644
--- /dev/null
+++ b/Engine/Source/Editor/UMGEditor/Classes/GUIPageFactory.h
@@ -0,0 +1,16 @@
+// Copyright 1998-2014 Epic Games, Inc. All Rights Reserved.
+
+#pragma once
+
+#include "GUIPageFactory.generated.h"
+
+UCLASS()
+class UGUIPageFactory : public UFactory
+{
+	GENERATED_UCLASS_BODY()
+
+	// UFactory interface
+	virtual UObject* FactoryCreateNew(UClass* Class, UObject* InParent, FName Name, EObjectFlags Flags, UObject* Context, FFeedbackContext* Warn) OVERRIDE;
+	virtual FText UGUIPageFactory::GetDisplayName() const OVERRIDE;
+	// End of UFactory interface
+};
diff --git a/Engine/Source/Editor/UMGEditor/Private/GUIPageAssetTypeActions.cpp b/Engine/Source/Editor/UMGEditor/Private/GUIPageAssetTypeActions.cpp
new file mode 100644
--- /dev/null
+++ b/Engine/Source/Editor/UMGEditor/Private/GUIPageAssetTypeActions.cpp
@@ -0,0 +1,49 @@
+// Copyright 1998-2014 Epic Games, Inc. All Rights Reserved.
+
+#include "UMGEditorPrivatePCH.h"
+
+#include "GUIPageAssetTypeActions.h"
+#include "UMGEditor.h"
+
+#define LOCTEXT_NAMESPACE "FGUIPageAssetTypeActions"
+
+//////////////////////////////////////////////////////////////////////////
+// FGUIPageAssetTypeActions
+
+FText FGUIPageAssetTypeActions::GetName() const
+{
+	return LOCTEXT("FGUIPageAssetTypeActionsName", "GUIPage");
+}
+
+FColor FGUIPageAssetTypeActions::GetTypeColor() const
+{
+	return FColor(0, 255, 255);
+}
+
+UClass* FGUIPageAssetTypeActions::GetSupportedClass() const
+{
+	return UGUIPage::StaticClass();
+}
+
+void FGUIPageAssetTypeActions::OpenAssetEditor(const TArray<UObject*>& InObjects, TSharedPtr<class IToolkitHost> EditWithinLevelEditor)
+{
+	const EToolkitMode::Type Mode = EditWithinLevelEditor.IsValid() ? EToolkitMode::WorldCentric : EToolkitMode::Standalone;
+
+	for (auto ObjIt = InObjects.CreateConstIterator(); ObjIt; ++ObjIt)
+	{
+		if ( UGUIPage* GUIPage = Cast<UGUIPage>(*ObjIt) )
+		{
+			TSharedRef<FUMGEditor> NewGUIPageEditor(new FUMGEditor());
+			NewGUIPageEditor->Initialize(Mode, EditWithinLevelEditor, GUIPage);
+		}
+	}
+}
+
+uint32 FGUIPageAssetTypeActions::GetCategories()
+{
+	return EAssetTypeCategories::Misc;
+}
+
+//////////////////////////////////////////////////////////////////////////
+
+#undef LOCTEXT_NAMESPACE
\ No newline at end of file
diff --git a/Engine/Source/Editor/UMGEditor/Private/GUIPageAssetTypeActions.h b/Engine/Source/Editor/UMGEditor/Private/GUIPageAssetTypeActions.h
new file mode 100644
--- /dev/null
+++ b/Engine/Source/Editor/UMGEditor/Private/GUIPageAssetTypeActions.h
@@ -0,0 +1,17 @@
+// Copyright 1998-2014 Epic Games, Inc. All Rights Reserved.
+
+#pragma once
+
+#include "AssetTypeActions_Base.h"
+
+class FGUIPageAssetTypeActions : public FAssetTypeActions_Base
+{
+public:
+	// IAssetTypeActions interface
+	virtual FText GetName() const OVERRIDE;
+	virtual FColor GetTypeColor() const OVERRIDE;
+	virtual UClass* GetSupportedClass() const OVERRIDE;
+	virtual void OpenAssetEditor(const TArray<UObject*>& InObjects, TSharedPtr<class IToolkitHost> EditWithinLevelEditor = TSharedPtr<IToolkitHost>()) OVERRIDE;
+	virtual uint32 GetCategories() OVERRIDE;
+	// End of IAssetTypeActions interface
+};
\ No newline at end of file
diff --git a/Engine/Source/Editor/UMGEditor/Private/GUIPageFactory.cpp b/Engine/Source/Editor/UMGEditor/Private/GUIPageFactory.cpp
new file mode 100644
--- /dev/null
+++ b/Engine/Source/Editor/UMGEditor/Private/GUIPageFactory.cpp
@@ -0,0 +1,32 @@
+// Copyright 1998-2014 Epic Games, Inc. All Rights Reserved.
+
+#include "UMGEditorPrivatePCH.h"
+
+#include "GUIPageFactory.h"
+
+#define LOCTEXT_NAMESPACE "UGUIPageFactory"
+
+/////////////////////////////////////////////////////
+// UGUIPageFactory
+
+UGUIPageFactory::UGUIPageFactory(const class FPostConstructInitializeProperties& PCIP)
+	: Super(PCIP)
+{
+	bCreateNew = true;
+	bEditAfterNew = true;
+	SupportedClass = UGUIPage::StaticClass();
+}
+
+FText UGUIPageFactory::GetDisplayName() const
+{
+	return LOCTEXT("GUIPageFactoryDescription", "GUI Page");
+}
+
+UObject* UGUIPageFactory::FactoryCreateNew(UClass* Class, UObject* InParent, FName Name, EObjectFlags Flags, UObject* Context, FFeedbackContext* Warn)
+{
+	UGUIPage* NewPage = ConstructObject<UGUIPage>(Class, InParent, Name, Flags);
+
+	return NewPage;
+}
+
+#undef LOCTEXT_NAMESPACE
\ No newline at end of file
diff --git a/Engine/Source/Editor/UMGEditor/Private/SUMGEditorTree.cpp b/Engine/Source/Editor/UMGEditor/Private/SUMGEditorTree.cpp
new file mode 100644
--- /dev/null
+++ b/Engine/Source/Editor/UMGEditor/Private/SUMGEditorTree.cpp
@@ -0,0 +1,51 @@
+// Copyright 1998-2014 Epic Games, Inc. All Rights Reserved.
+
+#include "UMGEditorPrivatePCH.h"
+
+#include "SUMGEditorTree.h"
+#include "UMGEditor.h"
+#include "UMGEditorViewportClient.h"
+#include "UMGEditorActions.h"
+
+#include "PreviewScene.h"
+#include "SceneViewport.h"
+
+void SUMGEditorTree::Construct(const FArguments& InArgs, TWeakPtr<IUMGEditor> UMGEditor, UGUIPage* ObjectToEdit)
+{
+	UMGEditorPtr = UMGEditor;
+	Page = ObjectToEdit;
+
+	FCoreDelegates::OnObjectPropertyChanged.Add( FCoreDelegates::FOnObjectPropertyChanged::FDelegate::CreateRaw(this, &SUMGEditorTree::OnObjectPropertyChanged) );
+
+	ChildSlot
+	[
+		SNew(SVerticalBox)
+		+SVerticalBox::Slot()
+		[
+			SAssignNew(PageTreeView, STreeView<TSharedPtr<FString>>)
+			.ItemHeight(20.0f)
+			//.OnGenerateRow(this, &SDeviceProcesses::HandleProcessTreeViewGenerateRow)
+			//.OnGetChildren(this, &SDeviceProcesses::HandleProcessTreeViewGetChildren)
+			//.SelectionMode(ESelectionMode::Multi)
+			//.TreeItemsSource(&ProcessList)
+		]
+	];
+}
+
+SUMGEditorTree::~SUMGEditorTree()
+{
+	FCoreDelegates::OnObjectPropertyChanged.Remove( FCoreDelegates::FOnObjectPropertyChanged::FDelegate::CreateRaw(this, &SUMGEditorTree::OnObjectPropertyChanged) );
+}
+
+void SUMGEditorTree::AddReferencedObjects( FReferenceCollector& Collector )
+{
+	Collector.AddReferencedObject(Page);
+}
+
+void SUMGEditorTree::OnObjectPropertyChanged(UObject* ObjectBeingModified)
+{
+	if ( !ensure(ObjectBeingModified) )
+	{
+		return;
+	}
+}
diff --git a/Engine/Source/Editor/UMGEditor/Private/SUMGEditorTree.h b/Engine/Source/Editor/UMGEditor/Private/SUMGEditorTree.h
new file mode 100644
--- /dev/null
+++ b/Engine/Source/Editor/UMGEditor/Private/SUMGEditorTree.h
@@ -0,0 +1,39 @@
+// Copyright 1998-2014 Epic Games, Inc. All Rights Reserved.
+
+#pragma once
+
+#include "SCompoundWidget.h"
+
+//TODO rename SUMGEditorHierarchy
+class SUMGEditorTree : public SCompoundWidget, public FGCObject
+{
+public:
+	SLATE_BEGIN_ARGS( SUMGEditorTree ){}
+	SLATE_END_ARGS()
+
+	void Construct(const FArguments& InArgs, TWeakPtr<IUMGEditor> UMGEditor, UGUIPage* ObjectToEdit);
+	~SUMGEditorTree();
+	
+	// FGCObject interface
+	virtual void AddReferencedObjects( FReferenceCollector& Collector ) OVERRIDE;
+	// End of FGCObject interface
+
+	/** Set the parent tab of the viewport for determining visibility */
+	void SetParentTab( TSharedRef<SDockTab> InParentTab ) { ParentTab = InParentTab; }
+
+private:
+	/** Callback for updating preview socket meshes if the gui page has been modified. */
+	void OnObjectPropertyChanged(UObject* ObjectBeingModified);
+private:
+	
+	/** The parent tab where this viewport resides */
+	TWeakPtr<SDockTab> ParentTab;
+
+	/** Pointer back to the UMG editor tool that owns us */
+	TWeakPtr<IUMGEditor> UMGEditorPtr;
+
+	/** gui page being edited */
+	UGUIPage* Page;
+
+	TSharedPtr< STreeView< TSharedPtr< FString > > > PageTreeView;
+};
diff --git a/Engine/Source/Editor/UMGEditor/Private/SUMGEditorViewport.cpp b/Engine/Source/Editor/UMGEditor/Private/SUMGEditorViewport.cpp
new file mode 100644
--- /dev/null
+++ b/Engine/Source/Editor/UMGEditor/Private/SUMGEditorViewport.cpp
@@ -0,0 +1,104 @@
+// Copyright 1998-2014 Epic Games, Inc. All Rights Reserved.
+
+#include "UMGEditorPrivatePCH.h"
+
+#include "SUMGEditorViewport.h"
+#include "UMGEditor.h"
+#include "UMGEditorViewportClient.h"
+#include "UMGEditorActions.h"
+
+#include "PreviewScene.h"
+#include "SceneViewport.h"
+
+void SUMGEditorViewport::Construct(const FArguments& InArgs, TWeakPtr<IUMGEditor> UMGEditor, UGUIPage* ObjectToEdit)
+{
+	UMGEditorPtr = UMGEditor;
+	Page = ObjectToEdit;
+
+	CurrentViewMode = VMI_Lit;
+
+	SEditorViewport::Construct( SEditorViewport::FArguments() );
+
+	SetPreviewPage(Page);
+
+	FCoreDelegates::OnObjectPropertyChanged.Add( FCoreDelegates::FOnObjectPropertyChanged::FDelegate::CreateRaw(this, &SUMGEditorViewport::OnObjectPropertyChanged) );
+}
+
+SUMGEditorViewport::~SUMGEditorViewport()
+{
+	FCoreDelegates::OnObjectPropertyChanged.Remove( FCoreDelegates::FOnObjectPropertyChanged::FDelegate::CreateRaw(this, &SUMGEditorViewport::OnObjectPropertyChanged) );
+	if (EditorViewportClient.IsValid())
+	{
+		EditorViewportClient->Viewport = NULL;
+	}
+}
+
+void SUMGEditorViewport::AddReferencedObjects( FReferenceCollector& Collector )
+{
+	Collector.AddReferencedObject(Page);
+}
+
+void SUMGEditorViewport::RefreshViewport()
+{
+	// Invalidate the viewport's display.
+	SceneViewport->Invalidate();
+}
+
+void SUMGEditorViewport::OnObjectPropertyChanged(UObject* ObjectBeingModified)
+{
+	if ( !ensure(ObjectBeingModified) )
+	{
+		return;
+	}
+}
+
+void SUMGEditorViewport::SetPreviewPage(UGUIPage* InPage)
+{
+	EditorViewportClient->SetPreviewPage(InPage);
+}
+
+bool SUMGEditorViewport::IsVisible() const
+{
+	return ViewportWidget.IsValid() && (!ParentTab.IsValid() || ParentTab.Pin()->IsForeground());
+}
+
+FUMGEditorViewportClient& SUMGEditorViewport::GetViewportClient()
+{
+	return *EditorViewportClient;
+}
+
+TSharedRef<FEditorViewportClient> SUMGEditorViewport::MakeEditorViewportClient()
+{
+	EditorViewportClient = MakeShareable( new FUMGEditorViewportClient(UMGEditorPtr, PreviewScene, Page) );
+
+	EditorViewportClient->bSetListenerPosition = false;
+
+	EditorViewportClient->SetRealtime( false );
+	EditorViewportClient->VisibilityDelegate.BindSP( this, &SUMGEditorViewport::IsVisible );
+
+	return EditorViewportClient.ToSharedRef();
+}
+
+TSharedPtr<SWidget> SUMGEditorViewport::MakeViewportToolbar()
+{
+	return NULL; 
+}
+
+EVisibility SUMGEditorViewport::OnGetViewportContentVisibility() const
+{
+	return IsVisible() ? EVisibility::Visible : EVisibility::Collapsed;
+}
+
+void SUMGEditorViewport::BindCommands()
+{
+	SEditorViewport::BindCommands();
+
+	//const FUMGEditorCommands& Commands = FUMGEditorCommands::Get();
+
+	//TODO Commands
+}
+
+void SUMGEditorViewport::OnFocusViewportToSelection()
+{
+	//EditorViewportClient->FocusViewportOnBox(...);
+}
\ No newline at end of file
diff --git a/Engine/Source/Editor/UMGEditor/Private/SUMGEditorViewport.h b/Engine/Source/Editor/UMGEditor/Private/SUMGEditorViewport.h
new file mode 100644
--- /dev/null
+++ b/Engine/Source/Editor/UMGEditor/Private/SUMGEditorViewport.h
@@ -0,0 +1,69 @@
+// Copyright 1998-2014 Epic Games, Inc. All Rights Reserved.
+
+#pragma once
+
+#include "PreviewScene.h"
+#include "SEditorViewport.h"
+
+/**
+ * UMG Editor Preview viewport widget
+ */
+class SUMGEditorViewport : public SEditorViewport, public FGCObject
+{
+public:
+	SLATE_BEGIN_ARGS( SUMGEditorViewport ){}
+	SLATE_END_ARGS()
+
+	void Construct(const FArguments& InArgs, TWeakPtr<IUMGEditor> UMGEditor, UGUIPage* ObjectToEdit);
+	~SUMGEditorViewport();
+	
+	// FGCObject interface
+	virtual void AddReferencedObjects( FReferenceCollector& Collector ) OVERRIDE;
+	// End of FGCObject interface
+
+	void RefreshViewport();
+
+	void SetPreviewPage(UGUIPage* InPage);
+
+	/** @return The editor viewport client */
+	class FUMGEditorViewportClient& GetViewportClient();
+
+	/** Set the parent tab of the viewport for determining visibility */
+	void SetParentTab( TSharedRef<SDockTab> InParentTab ) { ParentTab = InParentTab; }
+
+protected:
+	/** SEditorViewport interface */
+	virtual TSharedRef<FEditorViewportClient> MakeEditorViewportClient() OVERRIDE;
+	virtual EVisibility OnGetViewportContentVisibility() const OVERRIDE;
+	virtual void BindCommands() OVERRIDE;
+	virtual void OnFocusViewportToSelection() OVERRIDE;
+	virtual TSharedPtr<SWidget> MakeViewportToolbar() OVERRIDE;
+private:
+	/** Determines the visibility of the viewport. */
+	bool IsVisible() const;
+
+	/** Callback for toggling the realtime preview flag. */
+	void SetRealtimePreview();
+
+	/** Callback for updating preview socket meshes if the gui page or socket has been modified. */
+	void OnObjectPropertyChanged(UObject* ObjectBeingModified);
+private:
+	
+	/** The parent tab where this viewport resides */
+	TWeakPtr<SDockTab> ParentTab;
+
+	/** Pointer back to the UMG editor tool that owns us */
+	TWeakPtr<IUMGEditor> UMGEditorPtr;
+
+	/** The scene for this viewport. */
+	FPreviewScene PreviewScene;
+
+	/** Level viewport client */
+	TSharedPtr<class FUMGEditorViewportClient> EditorViewportClient;
+
+	/** gui page being edited */
+	UGUIPage* Page;
+
+	/** The currently selected view mode. */
+	EViewModeIndex CurrentViewMode;
+};
diff --git a/Engine/Source/Editor/UMGEditor/Private/UMGEditor.cpp b/Engine/Source/Editor/UMGEditor/Private/UMGEditor.cpp
new file mode 100644
--- /dev/null
+++ b/Engine/Source/Editor/UMGEditor/Private/UMGEditor.cpp
@@ -0,0 +1,213 @@
+// Copyright 1998-2014 Epic Games, Inc. All Rights Reserved.
+
+#include "UMGEditorPrivatePCH.h"
+
+#include "UMGEditor.h"
+#include "UMGEditorModule.h"
+
+#include "Editor/WorkspaceMenuStructure/Public/WorkspaceMenuStructureModule.h"
+#include "Editor/PropertyEditor/Public/PropertyEditorModule.h"
+#include "Editor/PropertyEditor/Public/IDetailsView.h"
+
+#define LOCTEXT_NAMESPACE "UMGEditor"
+
+static const FName UMGEditor_HierarchyTab("UMGEditor_HierarchyTab");
+static const FName UMGEditor_PreviewViewportTab("UMGEditor_PreviewViewport");
+static const FName UMGEditor_DetailsTab("UMGEditor_Details");
+static const FName UMGEditor_CurveEditorTab("UMGEditor_CurveEditor");
+
+void FUMGEditor::RegisterTabSpawners(const TSharedRef<class FTabManager>& TabManager)
+{
+	FAssetEditorToolkit::RegisterTabSpawners(TabManager);
+
+	const IWorkspaceMenuStructure& MenuStructure = WorkspaceMenu::GetMenuStructure();
+
+	TabManager->RegisterTabSpawner(UMGEditor_HierarchyTab, FOnSpawnTab::CreateSP(this, &FUMGEditor::SpawnTab_Hierarchy))
+		.SetDisplayName(LOCTEXT("Hierarchy", "Hierarchy"))
+		.SetGroup(MenuStructure.GetAssetEditorCategory());
+
+	TabManager->RegisterTabSpawner(UMGEditor_PreviewViewportTab, FOnSpawnTab::CreateSP(this, &FUMGEditor::SpawnTab_Viewport))
+		.SetDisplayName(LOCTEXT("Viewport", "Viewport"))
+		.SetGroup(MenuStructure.GetAssetEditorCategory());
+
+	TabManager->RegisterTabSpawner(UMGEditor_DetailsTab, FOnSpawnTab::CreateSP(this, &FUMGEditor::SpawnTab_Details))
+		.SetDisplayName(LOCTEXT("Details", "Details"))
+		.SetGroup(MenuStructure.GetAssetEditorCategory());
+
+	//TabManager->RegisterTabSpawner(UMGEditor_CurveEditorTab, FOnSpawnTab::CreateSP(this, &FUMGEditor::SpawnTab, UMGEditor_CurveEditorTab))
+	//	.SetDisplayName(NSLOCTEXT("UMGEditor", "SummonCurveEditor", "CurveEditor"))
+	//	.SetGroup(MenuStructure.GetAssetEditorCategory());
+}
+
+void FUMGEditor::UnregisterTabSpawners(const TSharedRef<class FTabManager>& TabManager)
+{
+	FAssetEditorToolkit::UnregisterTabSpawners(TabManager);
+
+	TabManager->UnregisterTabSpawner(UMGEditor_HierarchyTab);
+	TabManager->UnregisterTabSpawner(UMGEditor_PreviewViewportTab);
+	TabManager->UnregisterTabSpawner(UMGEditor_DetailsTab);
+	//TabManager->UnregisterTabSpawner(UMGEditor_CurveEditorTab);
+}
+
+TSharedRef<SDockTab> FUMGEditor::SpawnTab_Hierarchy(const FSpawnTabArgs& Args)
+{
+	check(Args.GetTabId() == UMGEditor_HierarchyTab);
+
+	TSharedRef<SDockTab> SpawnedTab =
+		SNew(SDockTab)
+		.Label(LOCTEXT("Hierarchy", "Hierarchy"))
+		[
+			Hierarchy.ToSharedRef()
+		];
+
+	Hierarchy->SetParentTab(SpawnedTab);
+
+	return SpawnedTab;
+}
+
+TSharedRef<SDockTab> FUMGEditor::SpawnTab_Viewport(const FSpawnTabArgs& Args)
+{
+	check(Args.GetTabId() == UMGEditor_PreviewViewportTab);
+
+	TSharedRef<SDockTab> SpawnedTab =
+		SNew(SDockTab)
+		.Label(LOCTEXT("Viewport", "Viewport"))
+		[
+			Viewport.ToSharedRef()
+		];
+
+	Viewport->SetParentTab(SpawnedTab);
+
+	return SpawnedTab;
+}
+
+TSharedRef<SDockTab> FUMGEditor::SpawnTab_Details(const FSpawnTabArgs& Args)
+{
+	check(Args.GetTabId() == UMGEditor_DetailsTab);
+
+	TSharedRef<SDockTab> SpawnedTab = SNew(SDockTab)
+		.Icon(FEditorStyle::GetBrush("UMGEditor.Tabs.Properties"))
+		.Label(LOCTEXT("Details", "Details"))
+		[
+			Details.ToSharedRef()
+		];
+
+	return SpawnedTab;
+}
+
+FUMGEditor::~FUMGEditor()
+{
+	GEditor->UnregisterForUndo(this);
+}
+
+void FUMGEditor::Initialize(const EToolkitMode::Type Mode, const TSharedPtr< class IToolkitHost >& InitToolkitHost, UObject* ObjectToEdit)
+{
+	Page = CastChecked<UGUIPage>(ObjectToEdit);
+
+	// Support undo/redo
+	Page->SetFlags(RF_Transactional);
+
+	GEditor->RegisterForUndo(this);
+
+	BuildDetailsWidget();
+
+	Viewport = SNew(SUMGEditorViewport, SharedThis(this), Page);
+	Hierarchy = SNew(SUMGEditorTree, SharedThis(this), Page);
+
+	const TSharedRef<FTabManager::FLayout> StandaloneDefaultLayout = FTabManager::NewLayout( "UMGEditor_Layout_v1" )
+	->AddArea(
+		FTabManager::NewPrimaryArea()
+		->SetOrientation(Orient_Vertical)
+		->Split(
+			FTabManager::NewStack()
+			->SetSizeCoefficient(0.1f)
+			->AddTab( GetToolbarTabId(), ETabState::OpenedTab )
+		)
+		->Split
+		(
+			FTabManager::NewSplitter()
+			->SetSizeCoefficient(1.0f)
+			->SetOrientation(Orient_Horizontal)
+			->Split
+			(
+				FTabManager::NewSplitter()
+				->SetSizeCoefficient(1.0f)
+				->SetOrientation(Orient_Horizontal)
+				->Split
+				(
+					FTabManager::NewStack()
+					->SetSizeCoefficient(0.25f)
+					->AddTab( UMGEditor_HierarchyTab, ETabState::OpenedTab )
+				)
+				->Split
+				(
+					FTabManager::NewStack()
+					->AddTab(UMGEditor_PreviewViewportTab, ETabState::OpenedTab)
+				)
+				->Split
+				(
+					FTabManager::NewStack()
+					->SetSizeCoefficient(0.33f)
+					->AddTab(UMGEditor_DetailsTab, ETabState::OpenedTab)
+				)
+			)
+		)
+	);
+
+	const bool bCreateDefaultStandaloneMenu = true;
+	const bool bCreateDefaultToolbar = true;
+	FAssetEditorToolkit::InitAssetEditor(Mode, InitToolkitHost, UMGEditorAppIdentifier, StandaloneDefaultLayout, bCreateDefaultStandaloneMenu, bCreateDefaultToolbar, ObjectToEdit);
+}
+
+void FUMGEditor::BuildDetailsWidget()
+{
+	FDetailsViewArgs Args;
+	Args.bHideSelectionTip = true;
+	Args.NotifyHook = this;
+
+	FPropertyEditorModule& PropertyModule = FModuleManager::LoadModuleChecked<FPropertyEditorModule>("PropertyEditor");
+	Details = PropertyModule.CreateDetailView(Args);
+	Details->SetObject(Page);
+}
+
+FName FUMGEditor::GetToolkitFName() const
+{
+	return FName("UMGEditor");
+}
+
+FText FUMGEditor::GetBaseToolkitName() const
+{
+	return NSLOCTEXT("UMGEditor", "AppLabel", "UMGEditor");
+}
+
+FString FUMGEditor::GetWorldCentricTabPrefix() const
+{
+	return NSLOCTEXT("UMGEditor", "WorldCentricTabPrefix", "UMGEditor").ToString();
+}
+
+FLinearColor FUMGEditor::GetWorldCentricTabColorScale() const
+{
+	return FLinearColor(0.3f, 0.2f, 0.5f, 0.5f);
+}
+
+void FUMGEditor::Tick(float DeltaTime)
+{
+
+}
+
+bool FUMGEditor::IsTickable() const
+{
+	return true;
+}
+
+TStatId FUMGEditor::GetStatId() const
+{
+	RETURN_QUICK_DECLARE_CYCLE_STAT(FUMGEditor, STATGROUP_Tickables);
+}
+
+void FUMGEditor::AddReferencedObjects(FReferenceCollector& Collector)
+{
+
+}
+
+#undef LOCTEXT_NAMESPACE
\ No newline at end of file
diff --git a/Engine/Source/Editor/UMGEditor/Private/UMGEditor.h b/Engine/Source/Editor/UMGEditor/Private/UMGEditor.h
new file mode 100644
--- /dev/null
+++ b/Engine/Source/Editor/UMGEditor/Private/UMGEditor.h
@@ -0,0 +1,52 @@
+// Copyright 1998-2014 Epic Games, Inc. All Rights Reserved.
+
+#pragma once
+
+#include "Toolkits/AssetEditorToolkit.h"
+#include "IDistCurveEditor.h"
+#include "IUMGEditor.h"
+
+DECLARE_LOG_CATEGORY_EXTERN(LogUMGEditor, Log, All);
+
+class FUMGEditor : public IUMGEditor, public FGCObject, public FTickableEditorObject, public FNotifyHook, public FCurveEdNotifyInterface, public FEditorUndoClient
+{
+public:
+	virtual void RegisterTabSpawners(const TSharedRef<class FTabManager>& TabManager) OVERRIDE;
+	virtual void UnregisterTabSpawners(const TSharedRef<class FTabManager>& TabManager) OVERRIDE;
+	
+	TSharedRef<SDockTab> SpawnTab_Hierarchy(const FSpawnTabArgs& Args);
+	TSharedRef<SDockTab> SpawnTab_Viewport(const FSpawnTabArgs& Args);
+	TSharedRef<SDockTab> SpawnTab_Details(const FSpawnTabArgs& Args);
+
+	/** Destructor */
+	virtual ~FUMGEditor();
+
+	/** Edits the specified ParticleSystem object */
+	void Initialize(const EToolkitMode::Type Mode, const TSharedPtr< class IToolkitHost >& InitToolkitHost, UObject* ObjectToEdit);
+	
+
+	// FAssetEditorToolkit
+	virtual FName GetToolkitFName() const OVERRIDE;
+	virtual FText GetBaseToolkitName() const OVERRIDE;
+	virtual FString GetWorldCentricTabPrefix() const OVERRIDE;
+
+	//IToolkit
+	virtual FLinearColor GetWorldCentricTabColorScale() const OVERRIDE;
+
+	// FTickableEditorObject
+	virtual void Tick(float DeltaTime) OVERRIDE;
+	virtual bool IsTickable() const OVERRIDE;
+	virtual TStatId GetStatId() const OVERRIDE;
+
+	// FGCObject
+	virtual void AddReferencedObjects(FReferenceCollector& Collector) OVERRIDE;
+
+private:
+	void BuildDetailsWidget();
+
+	UGUIPage* Page;
+
+	TSharedPtr<class IDetailsView> Details;
+	TSharedPtr<class SUMGEditorViewport> Viewport;
+	TSharedPtr<class SUMGEditorTree> Hierarchy;
+};
diff --git a/Engine/Source/Editor/UMGEditor/Private/UMGEditorActions.cpp b/Engine/Source/Editor/UMGEditor/Private/UMGEditorActions.cpp
new file mode 100644
--- /dev/null
+++ b/Engine/Source/Editor/UMGEditor/Private/UMGEditorActions.cpp
@@ -0,0 +1,13 @@
+// Copyright 1998-2014 Epic Games, Inc. All Rights Reserved.
+
+#include "UMGEditorPrivatePCH.h"
+
+#include "UMGEditor.h"
+#include "UMGEditorActions.h"
+#include "SUMGEditorViewport.h"
+
+void FUMGEditorCommands::RegisterCommands()
+{
+	
+}
+
diff --git a/Engine/Source/Editor/UMGEditor/Private/UMGEditorActions.h b/Engine/Source/Editor/UMGEditor/Private/UMGEditorActions.h
new file mode 100644
--- /dev/null
+++ b/Engine/Source/Editor/UMGEditor/Private/UMGEditorActions.h
@@ -0,0 +1,33 @@
+// Copyright 1998-2014 Epic Games, Inc. All Rights Reserved.
+
+#pragma once
+
+/**
+ * Unreal UMG editor actions
+ */
+class FUMGEditorCommands : public TCommands<FUMGEditorCommands>
+{
+
+public:
+	FUMGEditorCommands() : TCommands<FUMGEditorCommands>
+	(
+		"UMGEditor", // Context name for fast lookup
+		NSLOCTEXT("Contexts", "UMGEditor", "UMG Editor"), // Localized context name for displaying
+		"EditorViewport",  // Parent
+		FEditorStyle::GetStyleSetName() // Icon Style Set
+	)
+	{
+	}
+	
+	/**
+	 * UMG Editor Commands
+	 */
+	
+
+	/**
+	 * Initialize commands
+	 */
+	virtual void RegisterCommands() OVERRIDE;
+
+public:
+};
diff --git a/Engine/Source/Editor/UMGEditor/Private/UMGEditorModule.cpp b/Engine/Source/Editor/UMGEditor/Private/UMGEditorModule.cpp
new file mode 100644
--- /dev/null
+++ b/Engine/Source/Editor/UMGEditor/Private/UMGEditorModule.cpp
@@ -0,0 +1,107 @@
+// Copyright 1998-2014 Epic Games, Inc. All Rights Reserved.
+
+#include "UMGEditorPrivatePCH.h"
+
+#include "UMGEditorModule.h"
+#include "ModuleManager.h"
+#include "UMGEditor.h"
+#include "AssetToolsModule.h"
+#include "GUIPageAssetTypeActions.h"
+
+#include "UMGEditor.generated.inl"
+
+const FName UMGEditorAppIdentifier = FName(TEXT("UMGEditorApp"));
+
+class FUMGEditorModule : public IUMGEditorModule
+{
+public:
+	/** Constructor, set up console commands and variables **/
+	FUMGEditorModule()
+	{
+	}
+
+	/** Called right after the module DLL has been loaded and the module object has been created */
+	virtual void StartupModule() OVERRIDE
+	{
+		MenuExtensibilityManager = MakeShareable(new FExtensibilityManager());
+		ToolBarExtensibilityManager = MakeShareable(new FExtensibilityManager());
+
+		// Register asset types
+		IAssetTools& AssetTools = FModuleManager::LoadModuleChecked<FAssetToolsModule>("AssetTools").Get();
+		RegisterAssetTypeAction(AssetTools, MakeShareable(new FGUIPageAssetTypeActions()));
+	}
+
+	/** Called before the module is unloaded, right before the module object is destroyed. */
+	virtual void ShutdownModule() OVERRIDE
+	{
+		MenuExtensibilityManager.Reset();
+		ToolBarExtensibilityManager.Reset();
+
+		if ( UObjectInitialized() )
+		{
+		}
+
+		// Unregister all the asset types that we registered
+		if ( FModuleManager::Get().IsModuleLoaded("AssetTools") )
+		{
+			IAssetTools& AssetTools = FModuleManager::GetModuleChecked<FAssetToolsModule>("AssetTools").Get();
+			for ( int32 Index = 0; Index < CreatedAssetTypeActions.Num(); ++Index )
+			{
+				AssetTools.UnregisterAssetTypeActions(CreatedAssetTypeActions[Index].ToSharedRef());
+			}
+		}
+		CreatedAssetTypeActions.Empty();
+	}
+
+	virtual TSharedRef<IUMGEditor> CreateUMGEditor(const EToolkitMode::Type Mode, const TSharedPtr< IToolkitHost >& InitToolkitHost, UParticleSystem* ParticleSystem) OVERRIDE
+	{
+		TSharedRef<FUMGEditor> NewUMGEditor(new FUMGEditor());
+		NewUMGEditor->Initialize(Mode, InitToolkitHost, ParticleSystem);
+		UMGEditorToolkits.Add(&*NewUMGEditor);
+		return NewUMGEditor;
+	}
+
+	virtual void UMGEditorClosed(FUMGEditor* UMGEditorInstance) OVERRIDE
+	{
+		UMGEditorToolkits.Remove(UMGEditorInstance);
+	}
+
+	virtual void RefreshUMGEditor(UParticleSystem* ParticleSystem) OVERRIDE
+	{
+		for (int32 Idx = 0; Idx < UMGEditorToolkits.Num(); ++Idx)
+		{
+			//if (UMGEditorToolkits[Idx]->GetParticleSystem() == ParticleSystem)
+			//{
+			//	UMGEditorToolkits[Idx]->ForceUpdate();
+			//}
+		}
+	}
+
+	virtual void ConvertModulesToSeeded(UParticleSystem* ParticleSystem) OVERRIDE
+	{
+		//FUMGEditor::ConvertAllModulesToSeeded( ParticleSystem );
+	}
+
+	/** Gets the extensibility managers for outside entities to extend gui page editor's menus and toolbars */
+	virtual TSharedPtr<FExtensibilityManager> GetMenuExtensibilityManager() {return MenuExtensibilityManager;}
+	virtual TSharedPtr<FExtensibilityManager> GetToolBarExtensibilityManager() {return ToolBarExtensibilityManager;}
+
+private:
+	void RegisterAssetTypeAction(IAssetTools& AssetTools, TSharedRef<IAssetTypeActions> Action)
+	{
+		AssetTools.RegisterAssetTypeActions(Action);
+		CreatedAssetTypeActions.Add(Action);
+	}
+
+private:
+	TSharedPtr<FExtensibilityManager> MenuExtensibilityManager;
+	TSharedPtr<FExtensibilityManager> ToolBarExtensibilityManager;
+
+	/** All created asset type actions.  Cached here so that we can unregister it during shutdown. */
+	TArray< TSharedPtr<IAssetTypeActions> > CreatedAssetTypeActions;
+
+	/** List of open UMGEditor toolkits */
+	TArray<FUMGEditor*> UMGEditorToolkits;
+};
+
+IMPLEMENT_MODULE(FUMGEditorModule, UMGEditor);
diff --git a/Engine/Source/Editor/UMGEditor/Private/UMGEditorPrivatePCH.h b/Engine/Source/Editor/UMGEditor/Private/UMGEditorPrivatePCH.h
new file mode 100644
--- /dev/null
+++ b/Engine/Source/Editor/UMGEditor/Private/UMGEditorPrivatePCH.h
@@ -0,0 +1,8 @@
+// Copyright 1998-2014 Epic Games, Inc. All Rights Reserved.
+
+#pragma once
+
+#include "UnrealEd.h"
+#include "UMGClasses.h"
+#include "IUMGModule.h"
+#include "UMGEditorClasses.h"
diff --git a/Engine/Source/Editor/UMGEditor/Private/UMGEditorViewportClient.cpp b/Engine/Source/Editor/UMGEditor/Private/UMGEditorViewportClient.cpp
new file mode 100644
--- /dev/null
+++ b/Engine/Source/Editor/UMGEditor/Private/UMGEditorViewportClient.cpp
@@ -0,0 +1,106 @@
+// Copyright 1998-2014 Epic Games, Inc. All Rights Reserved.
+
+#include "UMGEditorPrivatePCH.h"
+
+#include "UMGEditorActions.h"
+#include "UMGEditorViewportClient.h"
+#include "UMGEditor.h"
+
+#include "PreviewScene.h"
+#include "SceneViewport.h"
+
+#define LOCTEXT_NAMESPACE "FUMGEditorViewportClient"
+
+namespace {
+	static const float GridSize = 2048.0f;
+	static const int32 CellSize = 16;
+}
+
+FUMGEditorViewportClient::FUMGEditorViewportClient(TWeakPtr<IUMGEditor> InUMGEditor, FPreviewScene& InPreviewScene, UGUIPage* InPreviewPage)
+	: FEditorViewportClient(&InPreviewScene)
+	, UMGEditorPtr(InUMGEditor)
+{
+	// Setup defaults for the common draw helper.
+	DrawHelper.bDrawPivot = false;
+	DrawHelper.bDrawWorldBox = false;
+	DrawHelper.bDrawKillZ = false;
+	DrawHelper.bDrawGrid = true;
+	DrawHelper.GridColorAxis = FColor(160,160,160);
+	DrawHelper.GridColorMajor = FColor(144,144,144);
+	DrawHelper.GridColorMinor = FColor(128,128,128);
+	DrawHelper.PerspectiveGridSize = GridSize;
+	DrawHelper.NumCells = DrawHelper.PerspectiveGridSize / (CellSize * 2);
+
+	SetViewMode(VMI_Lit);
+
+	EngineShowFlags.DisableAdvancedFeatures();
+	EngineShowFlags.SetSnap(0);
+	EngineShowFlags.CompositeEditorPrimitives = true;
+	OverrideNearClipPlane(1.0f);
+	bUsingOrbitCamera = true;
+
+	SetPreviewPage(InPreviewPage);
+}
+
+void FUMGEditorViewportClient::SetPreviewPage(UGUIPage* InPreviewPage)
+{
+	PreviewPage = InPreviewPage;
+}
+
+void FUMGEditorViewportClient::Tick(float DeltaSeconds)
+{
+	FEditorViewportClient::Tick(DeltaSeconds);
+
+	// Tick the preview scene world.
+	if (!GIntraFrameDebuggingGameThread)
+	{
+		PreviewScene->GetWorld()->Tick(LEVELTICK_All, DeltaSeconds);
+	}
+}
+
+bool FUMGEditorViewportClient::ShouldOrbitCamera() const
+{
+	return false;
+}
+
+void FUMGEditorViewportClient::Draw(const FSceneView* View,FPrimitiveDrawInterface* PDI)
+{
+	FEditorViewportClient::Draw(View, PDI);
+}
+
+void FUMGEditorViewportClient::DrawCanvas( FViewport& InViewport, FSceneView& View, FCanvas& Canvas )
+{
+	FEditorViewportClient::DrawCanvas(InViewport, View, Canvas);
+}
+
+bool FUMGEditorViewportClient::InputKey(FViewport* Viewport, int32 ControllerId, FKey Key, EInputEvent Event,float AmountDepressed,bool Gamepad)
+{
+	bool bHandled = false;
+
+	if( !bHandled )
+	{
+		bHandled = FEditorViewportClient::InputKey( Viewport, ControllerId, Key, Event, AmountDepressed, false );
+	}
+
+	// Handle viewport screenshot.
+	bHandled |= InputTakeScreenshot( Viewport, Key, Event );
+
+	return bHandled;
+}
+
+void FUMGEditorViewportClient::ProcessClick(class FSceneView& InView, class HHitProxy* HitProxy, FKey Key, EInputEvent Event, uint32 HitX, uint32 HitY)
+{
+	Invalidate();
+}
+
+void FUMGEditorViewportClient::PerspectiveCameraMoved()
+{
+	FEditorViewportClient::PerspectiveCameraMoved();
+}
+
+void FUMGEditorViewportClient::ResetCamera()
+{
+	//FocusViewportOnBox( UMGComponent->Bounds.GetBox() );
+	Invalidate();
+}
+#undef LOCTEXT_NAMESPACE 
diff --git a/Engine/Source/Editor/UMGEditor/Private/UMGEditorViewportClient.h b/Engine/Source/Editor/UMGEditor/Private/UMGEditorViewportClient.h
new file mode 100644
--- /dev/null
+++ b/Engine/Source/Editor/UMGEditor/Private/UMGEditorViewportClient.h
@@ -0,0 +1,37 @@
+// Copyright 1998-2014 Epic Games, Inc. All Rights Reserved.
+
+#pragma once
+
+#include "EditorViewportClient.h"
+
+/** Viewport Client for the preview viewport */
+class FUMGEditorViewportClient : public FEditorViewportClient, public TSharedFromThis<FUMGEditorViewportClient>
+{
+public:
+	FUMGEditorViewportClient(TWeakPtr<IUMGEditor> InUMGEditor, FPreviewScene& InPreviewScene, UGUIPage* InPreviewGUIPage);
+
+	// FEditorViewportClient interface
+	virtual bool InputKey(FViewport* Viewport, int32 ControllerId, FKey Key, EInputEvent Event, float AmountDepressed = 1.f, bool bGamepad=false) OVERRIDE;
+	virtual void ProcessClick(class FSceneView& View, class HHitProxy* HitProxy, FKey Key, EInputEvent Event, uint32 HitX, uint32 HitY) OVERRIDE;
+	virtual void Tick(float DeltaSeconds) OVERRIDE;
+	virtual void Draw(const FSceneView* View,FPrimitiveDrawInterface* PDI) OVERRIDE;
+	virtual void DrawCanvas( FViewport& InViewport, FSceneView& View, FCanvas& Canvas ) OVERRIDE;
+	virtual bool ShouldOrbitCamera() const OVERRIDE;
+	//FEditorViewportClient interface
+
+	void ResetCamera();
+
+	void SetPreviewPage(UGUIPage* InPreviewPage);
+
+protected:
+	// FEditorViewportClient interface
+	virtual void PerspectiveCameraMoved() OVERRIDE;
+	// FEditorViewportClient interface
+
+private:
+	
+	/** Pointer back to the UMG editor tool that owns us */
+	TWeakPtr<IUMGEditor> UMGEditorPtr;
+
+	UGUIPage* PreviewPage;
+};
\ No newline at end of file
diff --git a/Engine/Source/Editor/UMGEditor/Public/IUMGEditor.h b/Engine/Source/Editor/UMGEditor/Public/IUMGEditor.h
new file mode 100644
--- /dev/null
+++ b/Engine/Source/Editor/UMGEditor/Public/IUMGEditor.h
@@ -0,0 +1,10 @@
+// Copyright 1998-2014 Epic Games, Inc. All Rights Reserved.
+
+#pragma once
+
+#include "Toolkits/IToolkitHost.h"
+
+class IUMGEditor : public FAssetEditorToolkit
+{
+
+};
diff --git a/Engine/Source/Editor/UMGEditor/Public/UMGEditorModule.h b/Engine/Source/Editor/UMGEditor/Public/UMGEditorModule.h
new file mode 100644
--- /dev/null
+++ b/Engine/Source/Editor/UMGEditor/Public/UMGEditorModule.h
@@ -0,0 +1,29 @@
+// Copyright 1998-2014 Epic Games, Inc. All Rights Reserved.
+
+#pragma once
+
+#include "UnrealEd.h"
+#include "ModuleInterface.h"
+#include "Toolkits/AssetEditorToolkit.h"
+#include "Toolkits/IToolkit.h"
+#include "IUMGEditor.h"
+
+extern const FName UMGEditorAppIdentifier;
+
+class FUMGEditor;
+
+class IUMGEditorModule : public IModuleInterface, public IHasMenuExtensibility, public IHasToolBarExtensibility
+{
+public:
+	/** Creates a new UMGEditor instance */
+	virtual TSharedRef<IUMGEditor> CreateUMGEditor(const EToolkitMode::Type Mode, const TSharedPtr< IToolkitHost >& InitToolkitHost, UParticleSystem* ParticleSystem) = 0;
+
+	/** Removes the specified instance from the list of open UMGEditor toolkits */
+	virtual void UMGEditorClosed(FUMGEditor* UMGEditorInstance) = 0;
+
+	/** Refreshes the toolkit inspecting the specified particle system */
+	virtual void RefreshUMGEditor(UParticleSystem* ParticleSystem) = 0;
+
+	/** Converts all the modules in the specified particle system to seeded modules */
+	virtual void ConvertModulesToSeeded(UParticleSystem* ParticleSystem) = 0;
+};
diff --git a/Engine/Source/Editor/UMGEditor/UMGEditor.Build.cs b/Engine/Source/Editor/UMGEditor/UMGEditor.Build.cs
new file mode 100644
--- /dev/null
+++ b/Engine/Source/Editor/UMGEditor/UMGEditor.Build.cs
@@ -0,0 +1,56 @@
+// Copyright 1998-2014 Epic Games, Inc. All Rights Reserved.
+
+using UnrealBuildTool;
+
+public class UMGEditor : ModuleRules
+{
+	public UMGEditor(TargetInfo Target)
+	{
+        PrivateIncludePaths.Add("Editor/UMGEditor/Private");	// For PCH includes (because they don't work with relative paths, yet)
+
+        PrivateIncludePathModuleNames.AddRange(
+            new string[] {
+                "AssetTools",
+                "UMG",
+            });
+
+		PublicIncludePaths.AddRange(
+			new string[] {
+				"Editor/DistCurveEditor/Public",
+				"Editor/UnrealEd/Public",
+			}
+			);
+
+		PrivateDependencyModuleNames.AddRange(
+			new string[] {
+				"Core",
+				"CoreUObject",
+				"Slate",
+				"Engine",
+                "InputCore",
+				"AssetTools",
+				"UnrealEd", // for FAssetEditorManager
+				"KismetWidgets",
+				"GraphEditor",
+				"Kismet",  // for FWorkflowCentricApplication
+				"PropertyEditor",
+				"RenderCore",
+				"LevelEditor", // for EdModes to get a toolkit host  //@TODO: PAPER: Should be a better way to do this (see the @todo in EdModeTileMap.cpp)
+				"UMG",
+                "DistCurveEditor",
+				"ContentBrowser",
+				"WorkspaceMenuStructure",
+				"EditorStyle"
+			}
+			);
+
+		DynamicallyLoadedModuleNames.AddRange(
+			new string[] {
+				"MainFrame",
+				"PropertyEditor",
+                "AssetTools",
+                "UMG",
+			}
+			);
+	}
+}
diff --git a/Engine/Source/Editor/UnrealEd/UnrealEd.Build.cs b/Engine/Source/Editor/UnrealEd/UnrealEd.Build.cs
--- a/Engine/Source/Editor/UnrealEd/UnrealEd.Build.cs
+++ b/Engine/Source/Editor/UnrealEd/UnrealEd.Build.cs
@@ -114,6 +114,7 @@ public class UnrealEd : ModuleRules
 				"StaticMeshEditor",
 				"TextureEditor",
 				"Cascade",
+                "UMGEditor",
 				"Matinee",
 				"AssetRegistry",
 				"AssetTools",
diff --git a/Engine/Source/Runtime/UMG/Classes/GUIPage.h b/Engine/Source/Runtime/UMG/Classes/GUIPage.h
new file mode 100644
--- /dev/null
+++ b/Engine/Source/Runtime/UMG/Classes/GUIPage.h
@@ -0,0 +1,12 @@
+// Copyright 1998-2014 Epic Games, Inc. All Rights Reserved.
+
+#pragma once
+
+#include "GUIPage.generated.h"
+
+UCLASS()
+class UMG_API UGUIPage : public UObject
+{
+	GENERATED_UCLASS_BODY()
+	
+};
diff --git a/Engine/Source/Runtime/UMG/Private/GUIPage.cpp b/Engine/Source/Runtime/UMG/Private/GUIPage.cpp
new file mode 100644
--- /dev/null
+++ b/Engine/Source/Runtime/UMG/Private/GUIPage.cpp
@@ -0,0 +1,13 @@
+// Copyright 1998-2014 Epic Games, Inc. All Rights Reserved.
+
+#include "UMGPrivatePCH.h"
+
+#include "GUIPage.h"
+
+//////////////////////////////////////////////////////////////////////////
+// UGUIPage
+
+UGUIPage::UGUIPage(const FPostConstructInitializeProperties& PCIP)
+	: Super(PCIP)
+{
+}
\ No newline at end of file
diff --git a/Engine/Source/Runtime/UMG/Private/UMGModule.cpp b/Engine/Source/Runtime/UMG/Private/UMGModule.cpp
new file mode 100644
--- /dev/null
+++ b/Engine/Source/Runtime/UMG/Private/UMGModule.cpp
@@ -0,0 +1,29 @@
+// Copyright 1998-2014 Epic Games, Inc. All Rights Reserved.
+
+#include "UMGPrivatePCH.h"
+
+#include "IUMGModule.h"
+#include "ModuleManager.h"
+
+#include "UMG.generated.inl"
+
+class FUMGModule : public IUMGModule
+{
+public:
+	/** Constructor, set up console commands and variables **/
+	FUMGModule()
+	{
+	}
+
+	/** Called right after the module DLL has been loaded and the module object has been created */
+	virtual void StartupModule() OVERRIDE
+	{
+	}
+
+	/** Called before the module is unloaded, right before the module object is destroyed. */
+	virtual void ShutdownModule() OVERRIDE
+	{
+	}
+};
+
+IMPLEMENT_MODULE(FUMGModule, UMG);
diff --git a/Engine/Source/Runtime/UMG/Private/UMGPrivatePCH.h b/Engine/Source/Runtime/UMG/Private/UMGPrivatePCH.h
new file mode 100644
--- /dev/null
+++ b/Engine/Source/Runtime/UMG/Private/UMGPrivatePCH.h
@@ -0,0 +1,13 @@
+// Copyright 1998-2014 Epic Games, Inc. All Rights Reserved.
+
+#pragma once
+
+// Uncommenting this like causes the FTransform linkage issue
+//#include "Core.h"
+
+#include "Engine.h"
+
+DECLARE_LOG_CATEGORY_EXTERN(LogUMG, Verbose, All);
+
+#include "UMGClasses.h"
+#include "IUMGModule.h"
\ No newline at end of file
diff --git a/Engine/Source/Runtime/UMG/Public/IUMGModule.h b/Engine/Source/Runtime/UMG/Public/IUMGModule.h
new file mode 100644
--- /dev/null
+++ b/Engine/Source/Runtime/UMG/Public/IUMGModule.h
@@ -0,0 +1,11 @@
+// Copyright 1998-2014 Epic Games, Inc. All Rights Reserved.
+
+#pragma once
+
+#include "ModuleInterface.h"
+
+class IUMGModule : public IModuleInterface
+{
+public:
+
+};
diff --git a/Engine/Source/Runtime/UMG/UMG.Build.cs b/Engine/Source/Runtime/UMG/UMG.Build.cs
new file mode 100644
--- /dev/null
+++ b/Engine/Source/Runtime/UMG/UMG.Build.cs
@@ -0,0 +1,26 @@
+// Copyright 1998-2014 Epic Games, Inc. All Rights Reserved.
+
+using UnrealBuildTool;
+
+public class UMG : ModuleRules
+{
+	public UMG(TargetInfo Target)
+	{
+		PublicDependencyModuleNames.AddRange(
+			new string[] {
+				"Core",
+				"CoreUObject",
+				"Engine",
+				"ShaderCore",
+				"RenderCore",
+				"RHI",
+			}
+			);
+
+		PrivateDependencyModuleNames.AddRange(
+			new string[] {
+				"Slate"
+			}
+			);
+	}
+}
