* Fixing non-unity.

diff --git a/Engine/Source/Editor/UMGEditor/Private/SUMGEditorViewport.h b/Engine/Source/Editor/UMGEditor/Private/SUMGEditorViewport.h
--- a/Engine/Source/Editor/UMGEditor/Private/SUMGEditorViewport.h
+++ b/Engine/Source/Editor/UMGEditor/Private/SUMGEditorViewport.h
@@ -4,6 +4,7 @@
 
 #include "PreviewScene.h"
 #include "SEditorViewport.h"
+#include "IUMGEditor.h"
 
 /**
  * UMG Editor Preview viewport widget
