Build Fix.

diff --git a/Engine/Source/Editor/UMGEditor/Classes/GUIPageFactory.h b/Engine/Source/Editor/UMGEditor/Classes/GUIPageFactory.h
--- a/Engine/Source/Editor/UMGEditor/Classes/GUIPageFactory.h
+++ b/Engine/Source/Editor/UMGEditor/Classes/GUIPageFactory.h
@@ -11,6 +11,6 @@ class UGUIPageFactory : public UFactory
 
 	// UFactory interface
 	virtual UObject* FactoryCreateNew(UClass* Class, UObject* InParent, FName Name, EObjectFlags Flags, UObject* Context, FFeedbackContext* Warn) OVERRIDE;
-	virtual FText UGUIPageFactory::GetDisplayName() const OVERRIDE;
+	virtual FText GetDisplayName() const OVERRIDE;
 	// End of UFactory interface
 };
