Fixing non-unity

diff --git a/Engine/Source/Editor/UMGEditor/Private/SUMGEditorTree.h b/Engine/Source/Editor/UMGEditor/Private/SUMGEditorTree.h
--- a/Engine/Source/Editor/UMGEditor/Private/SUMGEditorTree.h
+++ b/Engine/Source/Editor/UMGEditor/Private/SUMGEditorTree.h
@@ -3,6 +3,7 @@
 #pragma once
 
 #include "SCompoundWidget.h"
+#include "IUMGEditor.h"
 
 //TODO rename SUMGEditorHierarchy
 class SUMGEditorTree : public SCompoundWidget, public FGCObject
diff --git a/Engine/Source/Editor/UMGEditor/Private/UMGEditor.cpp b/Engine/Source/Editor/UMGEditor/Private/UMGEditor.cpp
--- a/Engine/Source/Editor/UMGEditor/Private/UMGEditor.cpp
+++ b/Engine/Source/Editor/UMGEditor/Private/UMGEditor.cpp
@@ -4,6 +4,8 @@
 
 #include "UMGEditor.h"
 #include "UMGEditorModule.h"
+#include "SUMGEditorTree.h"
+#include "SUMGEditorViewport.h"
 
 #include "Editor/WorkspaceMenuStructure/Public/WorkspaceMenuStructureModule.h"
 #include "Editor/PropertyEditor/Public/PropertyEditorModule.h"
diff --git a/Engine/Source/Editor/UMGEditor/Private/UMGEditorViewportClient.h b/Engine/Source/Editor/UMGEditor/Private/UMGEditorViewportClient.h
--- a/Engine/Source/Editor/UMGEditor/Private/UMGEditorViewportClient.h
+++ b/Engine/Source/Editor/UMGEditor/Private/UMGEditorViewportClient.h
@@ -3,6 +3,7 @@
 #pragma once
 
 #include "EditorViewportClient.h"
+#include "IUMGEditor.h"
 
 /** Viewport Client for the preview viewport */
 class FUMGEditorViewportClient : public FEditorViewportClient, public TSharedFromThis<FUMGEditorViewportClient>
diff --git a/Engine/Source/Editor/UMGEditor/Public/IUMGEditor.h b/Engine/Source/Editor/UMGEditor/Public/IUMGEditor.h
--- a/Engine/Source/Editor/UMGEditor/Public/IUMGEditor.h
+++ b/Engine/Source/Editor/UMGEditor/Public/IUMGEditor.h
@@ -3,6 +3,7 @@
 #pragma once
 
 #include "Toolkits/IToolkitHost.h"
+#include "AssetEditorToolkit.h"
 
 class IUMGEditor : public FAssetEditorToolkit
 {
