* New asset category in the new menu "User Interface", currently just the slate related assets are visible unless -umg is provided.
* Porting a large swath of the code from SlateComponentWrapper into UMG.
* Creating a derived UBlueprint for AUserWidget. (The base class for a UMG widget).

diff --git a/Engine/Source/Developer/AssetTools/Private/AssetTypeActions/AssetTypeActions_SlateBrush.h b/Engine/Source/Developer/AssetTools/Private/AssetTypeActions/AssetTypeActions_SlateBrush.h
--- a/Engine/Source/Developer/AssetTools/Private/AssetTypeActions/AssetTypeActions_SlateBrush.h
+++ b/Engine/Source/Developer/AssetTools/Private/AssetTypeActions/AssetTypeActions_SlateBrush.h
@@ -9,6 +9,6 @@ public:
 	virtual FText GetName() const OVERRIDE { return NSLOCTEXT("AssetTypeActions", "AssetTypeActions_SlateBrush", "Slate Brush"); }
 	virtual FColor GetTypeColor() const OVERRIDE { return FColor(105, 165, 60); }
 	virtual UClass* GetSupportedClass() const OVERRIDE { return USlateBrushAsset::StaticClass(); }
-	virtual uint32 GetCategories() OVERRIDE { return EAssetTypeCategories::Misc; }
+	virtual uint32 GetCategories() OVERRIDE { return EAssetTypeCategories::UI; }
 	virtual bool HasActions ( const TArray<UObject*>& InObjects ) const OVERRIDE { return false; }
 };
\ No newline at end of file
diff --git a/Engine/Source/Developer/AssetTools/Private/AssetTypeActions/AssetTypeActions_SlateWidgetStyle.h b/Engine/Source/Developer/AssetTools/Private/AssetTypeActions/AssetTypeActions_SlateWidgetStyle.h
--- a/Engine/Source/Developer/AssetTools/Private/AssetTypeActions/AssetTypeActions_SlateWidgetStyle.h
+++ b/Engine/Source/Developer/AssetTools/Private/AssetTypeActions/AssetTypeActions_SlateWidgetStyle.h
@@ -9,7 +9,7 @@ public:
 	virtual FText GetName() const OVERRIDE { return NSLOCTEXT("AssetTypeActions", "AssetTypeActions_SlateStyle", "Slate Widget Style"); }
 	virtual FColor GetTypeColor() const OVERRIDE { return FColor(62, 140, 35); }
 	virtual UClass* GetSupportedClass() const OVERRIDE { return USlateWidgetStyleAsset::StaticClass(); }
-	virtual uint32 GetCategories() OVERRIDE { return EAssetTypeCategories::Misc; }
+	virtual uint32 GetCategories() OVERRIDE { return EAssetTypeCategories::UI; }
 	virtual bool HasActions ( const TArray<UObject*>& InObjects ) const OVERRIDE { return true; }
 	virtual void GetActions( const TArray<UObject*>& InObjects, FMenuBuilder& MenuBuilder ) OVERRIDE;
 
diff --git a/Engine/Source/Developer/AssetTools/Public/IAssetTypeActions.h b/Engine/Source/Developer/AssetTools/Public/IAssetTypeActions.h
--- a/Engine/Source/Developer/AssetTools/Public/IAssetTypeActions.h
+++ b/Engine/Source/Developer/AssetTools/Public/IAssetTypeActions.h
@@ -6,13 +6,14 @@ namespace EAssetTypeCategories
 {
 	enum Type
 	{
-		None					= 0x00,
-		Basic					= 0x01,
-		Animation				= 0x02,
-		MaterialsAndTextures	= 0x04,
-		Sounds					= 0x08,
-		Physics					= 0x10,
-		Misc					= 0x20
+		None					= 0,
+		Basic					= 1 << 0,
+		Animation				= 1 << 1,
+		MaterialsAndTextures	= 1 << 2,
+		Sounds					= 1 << 3,
+		Physics					= 1 << 4,
+		UI						= 1 << 5,
+		Misc					= 1 << 6
 	};
 }
 
diff --git a/Engine/Source/Editor/ContentBrowser/Private/NewAssetContextMenu.cpp b/Engine/Source/Editor/ContentBrowser/Private/NewAssetContextMenu.cpp
--- a/Engine/Source/Editor/ContentBrowser/Private/NewAssetContextMenu.cpp
+++ b/Engine/Source/Editor/ContentBrowser/Private/NewAssetContextMenu.cpp
@@ -155,6 +155,12 @@ void FNewAssetContextMenu::MakeContextMenu(FMenuBuilder& MenuBuilder, const FStr
 			FNewMenuDelegate::CreateStatic(&FNewAssetContextMenu::CreateNewAssetMenuCategory, EAssetTypeCategories::Physics, InPath, InOnNewAssetRequested)
 			);
 
+		MenuBuilder.AddSubMenu(
+			LOCTEXT("MaterialAssetCategory", "User Interface"),
+			FText::GetEmpty(),
+			FNewMenuDelegate::CreateStatic(&FNewAssetContextMenu::CreateNewAssetMenuCategory, EAssetTypeCategories::UI, InPath, InOnNewAssetRequested)
+			);
+
 		MenuBuilder.AddSubMenu(
 			LOCTEXT("MiscellaneousAssetCategory", "Miscellaneous"),
 			FText::GetEmpty(),
diff --git a/Engine/Source/Editor/ContentBrowser/Private/SFilterList.cpp b/Engine/Source/Editor/ContentBrowser/Private/SFilterList.cpp
--- a/Engine/Source/Editor/ContentBrowser/Private/SFilterList.cpp
+++ b/Engine/Source/Editor/ContentBrowser/Private/SFilterList.cpp
@@ -950,6 +950,7 @@ TSharedRef<SWidget> SFilterList::MakeAddFilterMenu(EAssetTypeCategories::Type Me
 	CategoryToMenuMap.Add(EAssetTypeCategories::MaterialsAndTextures, FCategoryMenu( LOCTEXT("MaterialFilter", "Materials & Textures"), LOCTEXT("MaterialFilterTooltip", "Filter by material and texture assets."), "ContentBrowserFilterMaterialAsset", LOCTEXT("MaterialAssetsMenuHeading", "Material Assets") ) );
 	CategoryToMenuMap.Add(EAssetTypeCategories::Sounds, FCategoryMenu( LOCTEXT("SoundFilter", "Sounds"), LOCTEXT("SoundFilterTooltip", "Filter by sound assets."), "ContentBrowserFilterSoundAsset", LOCTEXT("SoundAssetsMenuHeading", "Sound Assets") ) );
 	CategoryToMenuMap.Add(EAssetTypeCategories::Physics, FCategoryMenu( LOCTEXT("PhysicsFilter", "Physics"), LOCTEXT("PhysicsFilterTooltip", "Filter by physics assets."), "ContentBrowserFilterPhysicsAsset", LOCTEXT("PhysicsAssetsMenuHeading", "Physics Assets") ) );
+	CategoryToMenuMap.Add(EAssetTypeCategories::UI, FCategoryMenu(LOCTEXT("UIFilter", "User Interface"), LOCTEXT("UIFilterTooltip", "Filter by UI assets."), "ContentBrowserFilterUIAsset", LOCTEXT("UIAssetsMenuHeading", "User Interface Assets")));
 	CategoryToMenuMap.Add(EAssetTypeCategories::Misc, FCategoryMenu( LOCTEXT("MiscFilter", "Miscellaneous"), LOCTEXT("MiscFilterTooltip", "Filter by miscellaneous assets."), "ContentBrowserFilterMiscAsset", LOCTEXT("MiscAssetsMenuHeading", "Misc Assets") ) );
 
 	// Load the asset tools module to get access to the browser type maps
diff --git a/Engine/Source/Editor/UMGEditor/Classes/GUIPageFactory.h b/Engine/Source/Editor/UMGEditor/Classes/GUIPageFactory.h
deleted file mode 100644
diff --git a/Engine/Source/Editor/UMGEditor/Classes/WidgetBlueprintFactory.h b/Engine/Source/Editor/UMGEditor/Classes/WidgetBlueprintFactory.h
new file mode 100644
--- /dev/null
+++ b/Engine/Source/Editor/UMGEditor/Classes/WidgetBlueprintFactory.h
@@ -0,0 +1,25 @@
+// Copyright 1998-2014 Epic Games, Inc. All Rights Reserved.
+
+#pragma once
+
+#include "WidgetBlueprintFactory.generated.h"
+
+UCLASS(HideCategories=Object, MinimalAPI)
+class UWidgetBlueprintFactory : public UFactory
+{
+	GENERATED_UCLASS_BODY()
+
+	// The type of blueprint that will be created
+	UPROPERTY(EditAnywhere, Category=WidgetBlueprintFactory)
+	TEnumAsByte<enum EBlueprintType> BlueprintType;
+
+	// The parent class of the created blueprint
+	UPROPERTY(EditAnywhere, Category=WidgetBlueprintFactory, meta=(AllowAbstract = ""))
+	TSubclassOf<class AUserWidget> ParentClass;
+
+	// Begin UFactory Interface
+	virtual bool ConfigureProperties() OVERRIDE;
+	virtual UObject* FactoryCreateNew(UClass* Class, UObject* InParent, FName Name, EObjectFlags Flags, UObject* Context, FFeedbackContext* Warn, FName CallingContext) OVERRIDE;
+	virtual UObject* FactoryCreateNew(UClass* Class, UObject* InParent, FName Name, EObjectFlags Flags, UObject* Context, FFeedbackContext* Warn) OVERRIDE;
+	// Begin UFactory Interface	
+};
diff --git a/Engine/Source/Editor/UMGEditor/Private/AssetTypeActions_WidgetBlueprint.cpp b/Engine/Source/Editor/UMGEditor/Private/AssetTypeActions_WidgetBlueprint.cpp
new file mode 100644
--- /dev/null
+++ b/Engine/Source/Editor/UMGEditor/Private/AssetTypeActions_WidgetBlueprint.cpp
@@ -0,0 +1,46 @@
+// Copyright 1998-2014 Epic Games, Inc. All Rights Reserved.
+
+#include "UMGEditorPrivatePCH.h"
+
+#include "PackageTools.h"
+#include "Kismet2/BlueprintEditorUtils.h"
+#include "Kismet2/KismetEditorUtilities.h"
+#include "BlueprintEditorModule.h"
+#include "AssetRegistryModule.h"
+#include "SBlueprintDiff.h"
+#include "ISourceControlModule.h"
+#include "MessageLog.h"
+
+#include "AssetTypeActions_WidgetBlueprint.h"
+
+#define LOCTEXT_NAMESPACE "AssetTypeActions"
+
+void FAssetTypeActions_WidgetBlueprint::OpenAssetEditor( const TArray<UObject*>& InObjects, TSharedPtr<IToolkitHost> EditWithinLevelEditor )
+{
+	EToolkitMode::Type Mode = EditWithinLevelEditor.IsValid() ? EToolkitMode::WorldCentric : EToolkitMode::Standalone;
+
+	for (auto ObjIt = InObjects.CreateConstIterator(); ObjIt; ++ObjIt)
+	{
+		auto Blueprint = Cast<UBlueprint>(*ObjIt);
+		if (Blueprint && Blueprint->SkeletonGeneratedClass && Blueprint->GeneratedClass )
+		{
+			FBlueprintEditorModule& BlueprintEditorModule = FModuleManager::LoadModuleChecked<FBlueprintEditorModule>( "Kismet" );
+			TSharedRef< IBlueprintEditor > NewKismetEditor = BlueprintEditorModule.CreateBlueprintEditor( Mode, EditWithinLevelEditor, Blueprint, ShouldUseDataOnlyEditor(Blueprint) );
+		}
+		else
+		{
+			FMessageDialog::Open( EAppMsgType::Ok, LOCTEXT("FailedToLoadBlueprint", "Blueprint could not be loaded because it derives from an invalid class.  Check to make sure the parent class for this blueprint hasn't been removed!"));
+		}
+	}
+}
+
+bool FAssetTypeActions_WidgetBlueprint::ShouldUseDataOnlyEditor(const UBlueprint* Blueprint) const
+{
+	return FBlueprintEditorUtils::IsDataOnlyBlueprint(Blueprint)
+		&& !FBlueprintEditorUtils::IsLevelScriptBlueprint(Blueprint)
+		&& !FBlueprintEditorUtils::IsInterfaceBlueprint(Blueprint)
+		&& !Blueprint->bForceFullEditor
+		&& !Blueprint->bIsNewlyCreated;
+}
+
+#undef LOCTEXT_NAMESPACE
\ No newline at end of file
diff --git a/Engine/Source/Editor/UMGEditor/Private/AssetTypeActions_WidgetBlueprint.h b/Engine/Source/Editor/UMGEditor/Private/AssetTypeActions_WidgetBlueprint.h
new file mode 100644
--- /dev/null
+++ b/Engine/Source/Editor/UMGEditor/Private/AssetTypeActions_WidgetBlueprint.h
@@ -0,0 +1,20 @@
+// Copyright 1998-2014 Epic Games, Inc. All Rights Reserved.
+
+#pragma once
+
+//TODO Use base class FAssetTypeActions_Blueprint later
+class FAssetTypeActions_WidgetBlueprint : public FAssetTypeActions_Base
+{
+public:
+	// IAssetTypeActions Implementation
+	virtual FText GetName() const OVERRIDE { return NSLOCTEXT("AssetTypeActions", "AssetTypeActions_WidgetBlueprint", "Widget Blueprint"); }
+	virtual FColor GetTypeColor() const OVERRIDE { return FColor(255,255,255); }
+	virtual UClass* GetSupportedClass() const OVERRIDE { return UWidgetBlueprint::StaticClass(); }
+	virtual void OpenAssetEditor(const TArray<UObject*>& InObjects, TSharedPtr<class IToolkitHost> EditWithinLevelEditor = TSharedPtr<IToolkitHost>()) OVERRIDE;
+	virtual uint32 GetCategories() OVERRIDE { return EAssetTypeCategories::UI; }
+	// End IAssetTypeActions Implementation
+
+private:
+	/** Returns true if the blueprint is data only */
+	bool ShouldUseDataOnlyEditor(const UBlueprint* Blueprint) const;
+};
\ No newline at end of file
diff --git a/Engine/Source/Editor/UMGEditor/Private/GUIPageAssetTypeActions.cpp b/Engine/Source/Editor/UMGEditor/Private/GUIPageAssetTypeActions.cpp
deleted file mode 100644
diff --git a/Engine/Source/Editor/UMGEditor/Private/GUIPageAssetTypeActions.h b/Engine/Source/Editor/UMGEditor/Private/GUIPageAssetTypeActions.h
deleted file mode 100644
diff --git a/Engine/Source/Editor/UMGEditor/Private/GUIPageFactory.cpp b/Engine/Source/Editor/UMGEditor/Private/GUIPageFactory.cpp
deleted file mode 100644
diff --git a/Engine/Source/Editor/UMGEditor/Private/SUMGEditorTree.cpp b/Engine/Source/Editor/UMGEditor/Private/SUMGEditorTree.cpp
--- a/Engine/Source/Editor/UMGEditor/Private/SUMGEditorTree.cpp
+++ b/Engine/Source/Editor/UMGEditor/Private/SUMGEditorTree.cpp
@@ -10,12 +10,16 @@
 #include "PreviewScene.h"
 #include "SceneViewport.h"
 
-void SUMGEditorTree::Construct(const FArguments& InArgs, TWeakPtr<IUMGEditor> UMGEditor, UGUIPage* ObjectToEdit)
+#include "BlueprintEditor.h"
+
+void SUMGEditorTree::Construct(const FArguments& InArgs, TSharedPtr<FBlueprintEditor> InBlueprintEditor, USimpleConstructionScript* InSCS)
 {
-	UMGEditorPtr = UMGEditor;
-	Page = ObjectToEdit;
+	BlueprintEditor = InBlueprintEditor;
+
+	UBlueprint* BP = InBlueprintEditor->GetBlueprintObj();
+	//BP->ComponentTemplates
 
-	FCoreDelegates::OnObjectPropertyChanged.Add( FCoreDelegates::FOnObjectPropertyChanged::FDelegate::CreateRaw(this, &SUMGEditorTree::OnObjectPropertyChanged) );
+	//FCoreDelegates::OnObjectPropertyChanged.Add( FCoreDelegates::FOnObjectPropertyChanged::FDelegate::CreateRaw(this, &SUMGEditorTree::OnObjectPropertyChanged) );
 
 	ChildSlot
 	[
@@ -34,18 +38,18 @@ void SUMGEditorTree::Construct(const FArguments& InArgs, TWeakPtr<IUMGEditor> UM
 
 SUMGEditorTree::~SUMGEditorTree()
 {
-	FCoreDelegates::OnObjectPropertyChanged.Remove( FCoreDelegates::FOnObjectPropertyChanged::FDelegate::CreateRaw(this, &SUMGEditorTree::OnObjectPropertyChanged) );
+	//FCoreDelegates::OnObjectPropertyChanged.Remove( FCoreDelegates::FOnObjectPropertyChanged::FDelegate::CreateRaw(this, &SUMGEditorTree::OnObjectPropertyChanged) );
 }
 
-void SUMGEditorTree::AddReferencedObjects( FReferenceCollector& Collector )
-{
-	Collector.AddReferencedObject(Page);
-}
-
-void SUMGEditorTree::OnObjectPropertyChanged(UObject* ObjectBeingModified)
-{
-	if ( !ensure(ObjectBeingModified) )
-	{
-		return;
-	}
-}
+//void SUMGEditorTree::AddReferencedObjects( FReferenceCollector& Collector )
+//{
+//	Collector.AddReferencedObject(Page);
+//}
+//
+//void SUMGEditorTree::OnObjectPropertyChanged(UObject* ObjectBeingModified)
+//{
+//	if ( !ensure(ObjectBeingModified) )
+//	{
+//		return;
+//	}
+//}
diff --git a/Engine/Source/Editor/UMGEditor/Private/SUMGEditorTree.h b/Engine/Source/Editor/UMGEditor/Private/SUMGEditorTree.h
--- a/Engine/Source/Editor/UMGEditor/Private/SUMGEditorTree.h
+++ b/Engine/Source/Editor/UMGEditor/Private/SUMGEditorTree.h
@@ -4,37 +4,21 @@
 
 #include "SCompoundWidget.h"
 #include "IUMGEditor.h"
+#include "BlueprintEditor.h"
 
 //TODO rename SUMGEditorHierarchy
-class SUMGEditorTree : public SCompoundWidget, public FGCObject
+class SUMGEditorTree : public SCompoundWidget
 {
 public:
 	SLATE_BEGIN_ARGS( SUMGEditorTree ){}
 	SLATE_END_ARGS()
 
-	void Construct(const FArguments& InArgs, TWeakPtr<IUMGEditor> UMGEditor, UGUIPage* ObjectToEdit);
+	void Construct(const FArguments& InArgs, TSharedPtr<FBlueprintEditor> InBlueprintEditor, USimpleConstructionScript* InSCS);
 	~SUMGEditorTree();
 	
-	// FGCObject interface
-	virtual void AddReferencedObjects( FReferenceCollector& Collector ) OVERRIDE;
-	// End of FGCObject interface
-
-	/** Set the parent tab of the viewport for determining visibility */
-	void SetParentTab( TSharedRef<SDockTab> InParentTab ) { ParentTab = InParentTab; }
-
-private:
-	/** Callback for updating preview socket meshes if the gui page has been modified. */
-	void OnObjectPropertyChanged(UObject* ObjectBeingModified);
 private:
-	
-	/** The parent tab where this viewport resides */
-	TWeakPtr<SDockTab> ParentTab;
-
-	/** Pointer back to the UMG editor tool that owns us */
-	TWeakPtr<IUMGEditor> UMGEditorPtr;
 
-	/** gui page being edited */
-	UGUIPage* Page;
+	TWeakPtr<class FBlueprintEditor> BlueprintEditor;
 
 	TSharedPtr< STreeView< TSharedPtr< FString > > > PageTreeView;
 };
diff --git a/Engine/Source/Editor/UMGEditor/Private/SUMGEditorViewport.cpp b/Engine/Source/Editor/UMGEditor/Private/SUMGEditorViewport.cpp
--- a/Engine/Source/Editor/UMGEditor/Private/SUMGEditorViewport.cpp
+++ b/Engine/Source/Editor/UMGEditor/Private/SUMGEditorViewport.cpp
@@ -10,7 +10,7 @@
 #include "PreviewScene.h"
 #include "SceneViewport.h"
 
-void SUMGEditorViewport::Construct(const FArguments& InArgs, TWeakPtr<IUMGEditor> UMGEditor, UGUIPage* ObjectToEdit)
+void SUMGEditorViewport::Construct(const FArguments& InArgs, TWeakPtr<IUMGEditor> UMGEditor, AUserWidget* ObjectToEdit)
 {
 	UMGEditorPtr = UMGEditor;
 	Page = ObjectToEdit;
@@ -52,7 +52,7 @@ void SUMGEditorViewport::OnObjectPropertyChanged(UObject* ObjectBeingModified)
 	}
 }
 
-void SUMGEditorViewport::SetPreviewPage(UGUIPage* InPage)
+void SUMGEditorViewport::SetPreviewPage(AUserWidget* InPage)
 {
 	EditorViewportClient->SetPreviewPage(InPage);
 }
diff --git a/Engine/Source/Editor/UMGEditor/Private/SUMGEditorViewport.h b/Engine/Source/Editor/UMGEditor/Private/SUMGEditorViewport.h
--- a/Engine/Source/Editor/UMGEditor/Private/SUMGEditorViewport.h
+++ b/Engine/Source/Editor/UMGEditor/Private/SUMGEditorViewport.h
@@ -15,7 +15,7 @@ public:
 	SLATE_BEGIN_ARGS( SUMGEditorViewport ){}
 	SLATE_END_ARGS()
 
-	void Construct(const FArguments& InArgs, TWeakPtr<IUMGEditor> UMGEditor, UGUIPage* ObjectToEdit);
+	void Construct(const FArguments& InArgs, TWeakPtr<IUMGEditor> UMGEditor, AUserWidget* ObjectToEdit);
 	~SUMGEditorViewport();
 	
 	// FGCObject interface
@@ -24,7 +24,7 @@ public:
 
 	void RefreshViewport();
 
-	void SetPreviewPage(UGUIPage* InPage);
+	void SetPreviewPage(AUserWidget* InPage);
 
 	/** @return The editor viewport client */
 	class FUMGEditorViewportClient& GetViewportClient();
@@ -63,7 +63,7 @@ private:
 	TSharedPtr<class FUMGEditorViewportClient> EditorViewportClient;
 
 	/** gui page being edited */
-	UGUIPage* Page;
+	AUserWidget* Page;
 
 	/** The currently selected view mode. */
 	EViewModeIndex CurrentViewMode;
diff --git a/Engine/Source/Editor/UMGEditor/Private/UMGBlueprintEditorExtensionHook.cpp b/Engine/Source/Editor/UMGEditor/Private/UMGBlueprintEditorExtensionHook.cpp
new file mode 100644
--- /dev/null
+++ b/Engine/Source/Editor/UMGEditor/Private/UMGBlueprintEditorExtensionHook.cpp
@@ -0,0 +1,432 @@
+// Copyright 1998-2014 Epic Games, Inc. All Rights Reserved.
+
+#include "UMGEditorPrivatePCH.h"
+
+#include "UMGBlueprintEditorExtensionHook.h"
+#include "WorkflowOrientedApp/WorkflowCentricApplication.h"
+
+#include "BlueprintEditor.h"
+#include "BlueprintEditorModes.h"
+#include "BlueprintEditorTabs.h"
+#include "SSCSEditorViewport.h"
+
+#include "SUMGEditorTree.h"
+
+#define LOCTEXT_NAMESPACE "UMG_EXTENSION"
+
+//class SDesignerWidget : public SBorder
+//{
+//	SLATE_BEGIN_ARGS(SDesignerWidget) {}
+//		/** Slot for the wrapped content (optional) */
+//		SLATE_DEFAULT_SLOT(FArguments, Content)
+//	SLATE_END_ARGS()
+//
+//	void Construct(const FArguments& InArgs, TSharedPtr<FBlueprintEditor> InBlueprintEditor)
+//	{
+//		LastPreviewActor = NULL;
+//		BlueprintEditor = InBlueprintEditor;
+//
+//		SBorder::Construct(SBorder::FArguments()
+//			.BorderImage(FStyleDefaults::GetNoBrush())
+//			.Padding(0.0f)
+//			[
+//
+//				InArgs._Content.Widget
+//			]);
+//		}
+//	}
+//
+//	//virtual FReply OnMouseButtonDown(const FGeometry& MyGeometry, const FPointerEvent& MouseEvent) OVERRIDE;
+//	//virtual FReply OnMouseButtonUp(const FGeometry& MyGeometry, const FPointerEvent& MouseEvent) OVERRIDE;
+///*virtual bool OnHitTest(const FGeometry& MyGeometry, FVector2D InAbsoluteCursorPosition) OVERRIDE;
+//{
+//return true;
+//}*/
+//
+//protected:
+//	TWeakPtr<FBlueprintEditor> BlueprintEditor;
+//	TWeakObjectPtr<AActor> LastPreviewActor;
+//};
+
+
+//////////////////////////////////////////////////////////////////////////
+
+FWorkflowApplicationModeExtender BlueprintEditorExtenderDelegate;
+
+static bool LocateWidgetsUnderCursor_Helper(FArrangedWidget& Candidate, FVector2D InAbsoluteCursorLocation, FArrangedChildren& OutWidgetsUnderCursor, bool bIgnoreEnabledStatus)
+{
+	const bool bCandidateUnderCursor =
+		// Candidate is physically under the cursor
+		Candidate.Geometry.IsUnderLocation(InAbsoluteCursorLocation) &&
+		// Candidate actually considers itself hit by this test
+		( Candidate.Widget->OnHitTest(Candidate.Geometry, InAbsoluteCursorLocation) );
+
+	bool bHitAnyWidget = false;
+	if ( bCandidateUnderCursor )
+	{
+		// The candidate widget is under the mouse
+		OutWidgetsUnderCursor.AddWidget(Candidate);
+
+		// Check to see if we were asked to still allow children to be hit test visible
+		bool bHitChildWidget = false;
+		if ( ( Candidate.Widget->GetVisibility().AreChildrenHitTestVisible() ) != 0 )
+		{
+			FArrangedChildren ArrangedChildren(OutWidgetsUnderCursor.GetFilter());
+			Candidate.Widget->ArrangeChildren(Candidate.Geometry, ArrangedChildren);
+
+			// A widget's children are implicitly Z-ordered from first to last
+			for ( int32 ChildIndex = ArrangedChildren.Num() - 1; !bHitChildWidget && ChildIndex >= 0; --ChildIndex )
+			{
+				FArrangedWidget& SomeChild = ArrangedChildren(ChildIndex);
+				bHitChildWidget = ( SomeChild.Widget->IsEnabled() || bIgnoreEnabledStatus ) && LocateWidgetsUnderCursor_Helper(SomeChild, InAbsoluteCursorLocation, OutWidgetsUnderCursor, bIgnoreEnabledStatus);
+			}
+		}
+
+		// If we hit a child widget or we hit our candidate widget then we'll append our widgets
+		const bool bHitCandidateWidget = Candidate.Widget->GetVisibility().IsHitTestVisible();
+		bHitAnyWidget = bHitChildWidget || bHitCandidateWidget;
+		if ( !bHitAnyWidget )
+		{
+			// No child widgets were hit, and even though the cursor was over our candidate widget, the candidate
+			// widget was not hit-testable, so we won't report it
+			check(OutWidgetsUnderCursor.Last() == Candidate);
+			OutWidgetsUnderCursor.Remove(OutWidgetsUnderCursor.Num() - 1);
+		}
+	}
+
+	return bHitAnyWidget;
+}
+
+/////////////////////////////////////////////////////
+// SSlateBlueprintPreview
+
+class SSlateBlueprintPreview : public SCompoundWidget
+{
+public:
+
+	SLATE_BEGIN_ARGS( SSlateBlueprintPreview ) {}
+	SLATE_END_ARGS()
+
+	void Construct(const FArguments& InArgs, TSharedPtr<FBlueprintEditor> InBlueprintEditor)
+	{
+		LastPreviewActor = NULL;
+		BlueprintEditor = InBlueprintEditor;
+
+		ChildSlot
+		[
+			SNew(SOverlay)
+			+ SOverlay::Slot()
+			.HAlign(HAlign_Fill)
+			.VAlign(VAlign_Fill)
+			[
+				SAssignNew(PreviewSurface, SBorder)
+				.Visibility(EVisibility::HitTestInvisible)
+			]
+
+			+ SOverlay::Slot()
+			.HAlign(HAlign_Fill)
+			.VAlign(VAlign_Fill)
+			[
+				SNew(STextBlock)
+			]
+		];
+	}
+
+	virtual FReply OnMouseButtonDown(const FGeometry& MyGeometry, const FPointerEvent& MouseEvent) OVERRIDE
+	{
+		FVector2D LocalMouseCoordinates = MyGeometry.AbsoluteToLocal(MouseEvent.GetScreenSpacePosition());
+
+		FArrangedChildren Children(EVisibility::All);
+
+		PreviewSurface->SetVisibility(EVisibility::Visible);
+		FArrangedWidget WindowWidgetGeometry(PreviewSurface.ToSharedRef(), MyGeometry);
+		LocateWidgetsUnderCursor_Helper(WindowWidgetGeometry, MouseEvent.GetScreenSpacePosition(), Children, true);
+
+		PreviewSurface->SetVisibility(EVisibility::HitTestInvisible);
+
+		return FReply::Handled();
+	}
+
+	// SWidget interface
+	virtual void Tick(const FGeometry& AllottedGeometry, const double InCurrentTime, const float InDeltaTime) OVERRIDE;
+	// End of Swidget interface
+protected:
+	TWeakPtr<FBlueprintEditor> BlueprintEditor;
+	TWeakObjectPtr<AActor> LastPreviewActor;
+	TWeakPtr<SWidget> LastPreviewWidget;
+
+	TSharedPtr<SBorder> PreviewSurface;
+};
+
+void SSlateBlueprintPreview::Tick(const FGeometry& AllottedGeometry, const double InCurrentTime, const float InDeltaTime)
+{
+	AActor* PreviewActor = BlueprintEditor.Pin()->GetPreviewActor();
+	if (PreviewActor != LastPreviewActor.Get())
+	{
+		LastPreviewActor = PreviewActor;
+	}
+
+	if (AUserWidget* WidgetActor = Cast<AUserWidget>(PreviewActor))
+	{
+		TSharedRef<SWidget> CurrentWidget = WidgetActor->GetWidget();
+
+		if (CurrentWidget != LastPreviewWidget.Pin())
+		{
+			LastPreviewWidget = CurrentWidget;
+			PreviewSurface->SetContent(CurrentWidget);
+		}
+	}
+	else
+	{
+		ChildSlot
+		[
+			SNew(SHorizontalBox)
+			+SHorizontalBox::Slot()
+			.HAlign(HAlign_Center)
+			.VAlign(VAlign_Center)
+			[
+				SNew(STextBlock)
+				.Text(LOCTEXT("NoWrappedWidget", "No actor; Open the viewport and tab back"))
+			]
+		];
+	}
+
+	SCompoundWidget::Tick(AllottedGeometry, InCurrentTime, InDeltaTime);
+}
+
+/////////////////////////////////////////////////////
+// FSlatePreviewSummoner
+
+static const FName SlatePreviewTabID(TEXT("SlatePreview"));
+
+struct FSlatePreviewSummoner : public FWorkflowTabFactory
+{
+protected:
+	TWeakPtr<class FBlueprintEditor> BlueprintEditor;
+public:
+	FSlatePreviewSummoner(TSharedPtr<class FBlueprintEditor> InBlueprintEditor)
+		: FWorkflowTabFactory(SlatePreviewTabID, InBlueprintEditor)
+		, BlueprintEditor(InBlueprintEditor)
+	{
+		TabLabel = LOCTEXT("SlatePreviewTabLabel", "UI Preview");
+		//TabIcon = FEditorStyle::GetBrush("Kismet.Tabs.Components");
+
+		bIsSingleton = true;
+
+		ViewMenuDescription = LOCTEXT("SlatePreview_ViewMenu_Desc", "UI Preview");
+		ViewMenuTooltip = LOCTEXT("SlatePreview_ViewMenu_ToolTip", "Show the UI preview");
+	}
+
+	virtual TSharedRef<SWidget> CreateTabBody(const FWorkflowTabSpawnInfo& Info) const OVERRIDE
+	{
+		TSharedPtr<FBlueprintEditor> BlueprintEditorPtr = BlueprintEditor.Pin();
+
+		TSharedPtr<SWidget> Result;
+		if ( BlueprintEditorPtr->CanAccessComponentsMode() )
+		{
+			Result = BlueprintEditorPtr->GetSCSViewport();
+		}
+
+		if ( Result.IsValid() )
+		{
+			return SNew(SHorizontalBox)
+			+ SHorizontalBox::Slot()
+			.FillWidth(0)
+			[
+				Result.ToSharedRef()
+			]
+
+			+ SHorizontalBox::Slot()
+			.FillWidth(1)
+			[
+				SNew(SSlateBlueprintPreview, BlueprintEditorPtr)
+			];
+		}
+		else
+		{
+			return SNew(SErrorText)
+				.BackgroundColor(FLinearColor::Transparent)
+				.ErrorText(LOCTEXT("SCSViewportView_Unavailable", "Viewport is not available for this Blueprint.").ToString());
+		}
+	}
+};
+
+/////////////////////////////////////////////////////
+// FSlateTreeSummoner
+
+static const FName SlateHierarchyTabID(TEXT("SlateHierarchy"));
+
+struct FSlateTreeSummoner : public FWorkflowTabFactory
+{
+protected:
+	TWeakPtr<class FBlueprintEditor> BlueprintEditor;
+public:
+	FSlateTreeSummoner(TSharedPtr<class FBlueprintEditor> InBlueprintEditor)
+		: FWorkflowTabFactory(SlateHierarchyTabID, InBlueprintEditor)
+		, BlueprintEditor(InBlueprintEditor)
+	{
+		TabLabel = LOCTEXT("SlateHierarchyTabLabel", "Hierarchy");
+		//TabIcon = FEditorStyle::GetBrush("Kismet.Tabs.Components");
+
+		bIsSingleton = true;
+
+		ViewMenuDescription = LOCTEXT("SlateHierarchy_ViewMenu_Desc", "Hierarchy");
+		ViewMenuTooltip = LOCTEXT("SlateHierarchy_ViewMenu_ToolTip", "Show the Hierarchy");
+	}
+
+	virtual TSharedRef<SWidget> CreateTabBody(const FWorkflowTabSpawnInfo& Info) const OVERRIDE
+	{
+		return SNew(SUMGEditorTree, BlueprintEditor.Pin(), BlueprintEditor.Pin()->GetBlueprintObj()->SimpleConstructionScript);
+	}
+};
+
+//////////////////////////////////////////////////////////////////////////
+// FComponentsEditorModeOverride
+
+class FComponentsEditorModeOverride : public FBlueprintComponentsApplicationMode
+{
+public:
+	FComponentsEditorModeOverride(TSharedPtr<class FBlueprintEditor> InBlueprintEditor)
+		: FBlueprintComponentsApplicationMode(InBlueprintEditor)
+	{
+		//BlueprintComponentsTabFactories.RegisterFactory(MakeShareable(new FConstructionScriptEditorSummoner(InBlueprintEditor)));
+		//BlueprintComponentsTabFactories.RegisterFactory(MakeShareable(new FSCSViewportSummoner(InBlueprintEditor)));
+		//BlueprintComponentsTabFactories.RegisterFactory(MakeShareable(new FSelectionDetailsSummoner(InBlueprintEditor)));
+		//BlueprintComponentsTabFactories.RegisterFactory(MakeShareable(new FDefaultsEditorSummoner(InBlueprintEditor)));
+		//BlueprintComponentsTabFactories.RegisterFactory(MakeShareable(new FFindResultsSummoner(InBlueprintEditor)));
+
+		BlueprintComponentsTabFactories.UnregisterFactory(FBlueprintEditorTabs::SCSViewportID);
+		BlueprintComponentsTabFactories.RegisterFactory(MakeShareable(new FSlatePreviewSummoner(InBlueprintEditor)));
+		BlueprintComponentsTabFactories.RegisterFactory(MakeShareable(new FSlateTreeSummoner(InBlueprintEditor)));
+
+		TabLayout = FTabManager::NewLayout( "Standalone_UMGEditor_Layout_v2" )
+		->AddArea
+		(
+			FTabManager::NewPrimaryArea()
+			->SetOrientation(Orient_Vertical)
+			->Split
+			(
+				FTabManager::NewStack()
+				->SetSizeCoefficient( 0.2f )
+				->SetHideTabWell(true)
+				->AddTab( InBlueprintEditor->GetToolbarTabId(), ETabState::OpenedTab )
+			)
+			->Split
+			(
+				FTabManager::NewSplitter()
+				->SetOrientation(Orient_Horizontal)
+				->Split
+				(
+					FTabManager::NewSplitter()
+					->SetSizeCoefficient( 0.15f )
+					->SetOrientation(Orient_Vertical)
+					->Split
+					(
+						FTabManager::NewStack()
+						->SetSizeCoefficient( 0.5f )
+						->AddTab( FBlueprintEditorTabs::ConstructionScriptEditorID, ETabState::OpenedTab )
+					)
+				)
+				->Split
+				(
+					FTabManager::NewSplitter()
+					->SetSizeCoefficient( 0.15f )
+					->SetOrientation(Orient_Vertical)
+					->Split
+					(
+						FTabManager::NewStack()
+						->SetSizeCoefficient( 0.5f )
+						->AddTab( SlateHierarchyTabID, ETabState::OpenedTab )
+					)
+				)
+				->Split
+				(
+					FTabManager::NewSplitter()
+					->SetSizeCoefficient( 0.45f )
+					->SetOrientation(Orient_Horizontal)
+					->Split
+					(
+						FTabManager::NewStack()
+						->SetHideTabWell(true)
+						->AddTab( SlatePreviewTabID, ETabState::OpenedTab )
+					)
+					->Split
+					(
+						FTabManager::NewStack()
+						->SetSizeCoefficient( 0.5f )
+						->AddTab( FBlueprintEditorTabs::DetailsID, ETabState::OpenedTab )
+					)
+				)
+			)
+		);
+
+	}
+
+	UBlueprint* GetBlueprint() const
+	{
+		if (FBlueprintEditor* Editor = MyBlueprintEditor.Pin().Get())
+		{
+			return Editor->GetBlueprintObj();
+		}
+		else
+		{
+			return NULL;
+		}
+	}	
+	
+	TSharedPtr<FBlueprintEditor> GetBlueprintEditor() const
+	{
+		return MyBlueprintEditor.Pin();
+	}
+
+	void Foo()
+	{
+		if (AActor* PreviewActor = GetPreviewActor())
+		{
+		}
+	}
+};
+
+//////////////////////////////////////////////////////////////////////////
+// FBlueprintEditorExtenderForSlatePreview
+
+class FBlueprintEditorExtenderForSlatePreview
+{
+public:
+	static TSharedRef<FApplicationMode> OnModeCreated(const FName ModeName, TSharedRef<FApplicationMode> InMode)
+	{
+		if (ModeName == FBlueprintEditorApplicationModes::BlueprintComponentsMode)
+		{
+			//@TODO: Bit of a lie - push GetBlueprint up, or pass in editor!
+			auto LieMode = StaticCastSharedRef<FComponentsEditorModeOverride>(InMode);
+
+			if (UBlueprint* BP = LieMode->GetBlueprint())
+			{
+				if (BP->ParentClass->IsChildOf(AUserWidget::StaticClass()))
+				{
+					return MakeShareable(new FComponentsEditorModeOverride(LieMode->GetBlueprintEditor()));
+				}
+			}
+		}
+
+		return InMode;
+	}
+};
+
+//////////////////////////////////////////////////////////////////////////
+// FUMGBlueprintEditorExtensionHook
+
+void FUMGBlueprintEditorExtensionHook::InstallHooks()
+{
+	BlueprintEditorExtenderDelegate = FWorkflowApplicationModeExtender::CreateStatic(&FBlueprintEditorExtenderForSlatePreview::OnModeCreated);
+	FWorkflowCentricApplication::GetModeExtenderList().Add(BlueprintEditorExtenderDelegate);
+}
+
+void FUMGBlueprintEditorExtensionHook::RemoveHooks()
+{
+	FWorkflowCentricApplication::GetModeExtenderList().Remove(BlueprintEditorExtenderDelegate);
+}
+
+//////////////////////////////////////////////////////////////////////////
+
+#undef LOCTEXT_NAMESPACE
\ No newline at end of file
diff --git a/Engine/Source/Editor/UMGEditor/Private/UMGBlueprintEditorExtensionHook.h b/Engine/Source/Editor/UMGEditor/Private/UMGBlueprintEditorExtensionHook.h
new file mode 100644
--- /dev/null
+++ b/Engine/Source/Editor/UMGEditor/Private/UMGBlueprintEditorExtensionHook.h
@@ -0,0 +1,11 @@
+// Copyright 1998-2014 Epic Games, Inc. All Rights Reserved.
+
+#pragma once
+
+// Integrate SlateComponentWrapper into the blueprint editor (custom preview panel)
+class FUMGBlueprintEditorExtensionHook
+{
+public:
+	static void InstallHooks();
+	static void RemoveHooks();
+};
\ No newline at end of file
diff --git a/Engine/Source/Editor/UMGEditor/Private/UMGEditor.cpp b/Engine/Source/Editor/UMGEditor/Private/UMGEditor.cpp
--- a/Engine/Source/Editor/UMGEditor/Private/UMGEditor.cpp
+++ b/Engine/Source/Editor/UMGEditor/Private/UMGEditor.cpp
@@ -24,9 +24,9 @@ void FUMGEditor::RegisterTabSpawners(const TSharedRef<class FTabManager>& TabMan
 
 	const IWorkspaceMenuStructure& MenuStructure = WorkspaceMenu::GetMenuStructure();
 
-	TabManager->RegisterTabSpawner(UMGEditor_HierarchyTab, FOnSpawnTab::CreateSP(this, &FUMGEditor::SpawnTab_Hierarchy))
-		.SetDisplayName(LOCTEXT("Hierarchy", "Hierarchy"))
-		.SetGroup(MenuStructure.GetAssetEditorCategory());
+	//TabManager->RegisterTabSpawner(UMGEditor_HierarchyTab, FOnSpawnTab::CreateSP(this, &FUMGEditor::SpawnTab_Hierarchy))
+	//	.SetDisplayName(LOCTEXT("Hierarchy", "Hierarchy"))
+	//	.SetGroup(MenuStructure.GetAssetEditorCategory());
 
 	TabManager->RegisterTabSpawner(UMGEditor_PreviewViewportTab, FOnSpawnTab::CreateSP(this, &FUMGEditor::SpawnTab_Viewport))
 		.SetDisplayName(LOCTEXT("Viewport", "Viewport"))
@@ -45,27 +45,27 @@ void FUMGEditor::UnregisterTabSpawners(const TSharedRef<class FTabManager>& TabM
 {
 	FAssetEditorToolkit::UnregisterTabSpawners(TabManager);
 
-	TabManager->UnregisterTabSpawner(UMGEditor_HierarchyTab);
+	//TabManager->UnregisterTabSpawner(UMGEditor_HierarchyTab);
 	TabManager->UnregisterTabSpawner(UMGEditor_PreviewViewportTab);
 	TabManager->UnregisterTabSpawner(UMGEditor_DetailsTab);
 	//TabManager->UnregisterTabSpawner(UMGEditor_CurveEditorTab);
 }
 
-TSharedRef<SDockTab> FUMGEditor::SpawnTab_Hierarchy(const FSpawnTabArgs& Args)
-{
-	check(Args.GetTabId() == UMGEditor_HierarchyTab);
-
-	TSharedRef<SDockTab> SpawnedTab =
-		SNew(SDockTab)
-		.Label(LOCTEXT("Hierarchy", "Hierarchy"))
-		[
-			Hierarchy.ToSharedRef()
-		];
-
-	Hierarchy->SetParentTab(SpawnedTab);
-
-	return SpawnedTab;
-}
+//TSharedRef<SDockTab> FUMGEditor::SpawnTab_Hierarchy(const FSpawnTabArgs& Args)
+//{
+//	check(Args.GetTabId() == UMGEditor_HierarchyTab);
+//
+//	TSharedRef<SDockTab> SpawnedTab =
+//		SNew(SDockTab)
+//		.Label(LOCTEXT("Hierarchy", "Hierarchy"))
+//		[
+//			Hierarchy.ToSharedRef()
+//		];
+//
+//	//Hierarchy->SetParentTab(SpawnedTab);
+//
+//	return SpawnedTab;
+//}
 
 TSharedRef<SDockTab> FUMGEditor::SpawnTab_Viewport(const FSpawnTabArgs& Args)
 {
@@ -104,7 +104,7 @@ FUMGEditor::~FUMGEditor()
 
 void FUMGEditor::Initialize(const EToolkitMode::Type Mode, const TSharedPtr< class IToolkitHost >& InitToolkitHost, UObject* ObjectToEdit)
 {
-	Page = CastChecked<UGUIPage>(ObjectToEdit);
+	Page = CastChecked<AUserWidget>(ObjectToEdit);
 
 	// Support undo/redo
 	Page->SetFlags(RF_Transactional);
@@ -114,7 +114,7 @@ void FUMGEditor::Initialize(const EToolkitMode::Type Mode, const TSharedPtr< cla
 	BuildDetailsWidget();
 
 	Viewport = SNew(SUMGEditorViewport, SharedThis(this), Page);
-	Hierarchy = SNew(SUMGEditorTree, SharedThis(this), Page);
+	//Hierarchy = SNew(SUMGEditorTree, SharedThis(this), Page);
 
 	const TSharedRef<FTabManager::FLayout> StandaloneDefaultLayout = FTabManager::NewLayout( "UMGEditor_Layout_v1" )
 	->AddArea(
diff --git a/Engine/Source/Editor/UMGEditor/Private/UMGEditor.h b/Engine/Source/Editor/UMGEditor/Private/UMGEditor.h
--- a/Engine/Source/Editor/UMGEditor/Private/UMGEditor.h
+++ b/Engine/Source/Editor/UMGEditor/Private/UMGEditor.h
@@ -14,7 +14,7 @@ public:
 	virtual void RegisterTabSpawners(const TSharedRef<class FTabManager>& TabManager) OVERRIDE;
 	virtual void UnregisterTabSpawners(const TSharedRef<class FTabManager>& TabManager) OVERRIDE;
 	
-	TSharedRef<SDockTab> SpawnTab_Hierarchy(const FSpawnTabArgs& Args);
+	//TSharedRef<SDockTab> SpawnTab_Hierarchy(const FSpawnTabArgs& Args);
 	TSharedRef<SDockTab> SpawnTab_Viewport(const FSpawnTabArgs& Args);
 	TSharedRef<SDockTab> SpawnTab_Details(const FSpawnTabArgs& Args);
 
@@ -44,7 +44,7 @@ public:
 private:
 	void BuildDetailsWidget();
 
-	UGUIPage* Page;
+	AUserWidget* Page;
 
 	TSharedPtr<class IDetailsView> Details;
 	TSharedPtr<class SUMGEditorViewport> Viewport;
diff --git a/Engine/Source/Editor/UMGEditor/Private/UMGEditorModule.cpp b/Engine/Source/Editor/UMGEditor/Private/UMGEditorModule.cpp
--- a/Engine/Source/Editor/UMGEditor/Private/UMGEditorModule.cpp
+++ b/Engine/Source/Editor/UMGEditor/Private/UMGEditorModule.cpp
@@ -6,7 +6,8 @@
 #include "ModuleManager.h"
 #include "UMGEditor.h"
 #include "AssetToolsModule.h"
-#include "GUIPageAssetTypeActions.h"
+#include "UMGBlueprintEditorExtensionHook.h"
+#include "AssetTypeActions_WidgetBlueprint.h"
 
 #include "UMGEditor.generated.inl"
 
@@ -23,12 +24,16 @@ public:
 	/** Called right after the module DLL has been loaded and the module object has been created */
 	virtual void StartupModule() OVERRIDE
 	{
+		FModuleManager::LoadModuleChecked<IUMGModule>("UMG");
+
 		MenuExtensibilityManager = MakeShareable(new FExtensibilityManager());
 		ToolBarExtensibilityManager = MakeShareable(new FExtensibilityManager());
 
 		// Register asset types
 		IAssetTools& AssetTools = FModuleManager::LoadModuleChecked<FAssetToolsModule>("AssetTools").Get();
-		RegisterAssetTypeAction(AssetTools, MakeShareable(new FGUIPageAssetTypeActions()));
+		RegisterAssetTypeAction(AssetTools, MakeShareable(new FAssetTypeActions_WidgetBlueprint()));
+		
+		FUMGBlueprintEditorExtensionHook::InstallHooks();
 	}
 
 	/** Called before the module is unloaded, right before the module object is destroyed. */
@@ -39,6 +44,7 @@ public:
 
 		if ( UObjectInitialized() )
 		{
+			FUMGBlueprintEditorExtensionHook::InstallHooks();
 		}
 
 		// Unregister all the asset types that we registered
diff --git a/Engine/Source/Editor/UMGEditor/Private/UMGEditorPrivatePCH.h b/Engine/Source/Editor/UMGEditor/Private/UMGEditorPrivatePCH.h
--- a/Engine/Source/Editor/UMGEditor/Private/UMGEditorPrivatePCH.h
+++ b/Engine/Source/Editor/UMGEditor/Private/UMGEditorPrivatePCH.h
@@ -6,3 +6,5 @@
 #include "UMGClasses.h"
 #include "IUMGModule.h"
 #include "UMGEditorClasses.h"
+#include "AssetTypeActions_Base.h"
+#include "AssetTypeActions_Blueprint.h"
\ No newline at end of file
diff --git a/Engine/Source/Editor/UMGEditor/Private/UMGEditorViewportClient.cpp b/Engine/Source/Editor/UMGEditor/Private/UMGEditorViewportClient.cpp
--- a/Engine/Source/Editor/UMGEditor/Private/UMGEditorViewportClient.cpp
+++ b/Engine/Source/Editor/UMGEditor/Private/UMGEditorViewportClient.cpp
@@ -16,7 +16,7 @@ namespace {
 	static const int32 CellSize = 16;
 }
 
-FUMGEditorViewportClient::FUMGEditorViewportClient(TWeakPtr<IUMGEditor> InUMGEditor, FPreviewScene& InPreviewScene, UGUIPage* InPreviewPage)
+FUMGEditorViewportClient::FUMGEditorViewportClient(TWeakPtr<IUMGEditor> InUMGEditor, FPreviewScene& InPreviewScene, AUserWidget* InPreviewPage)
 	: FEditorViewportClient(&InPreviewScene)
 	, UMGEditorPtr(InUMGEditor)
 {
@@ -42,7 +42,7 @@ FUMGEditorViewportClient::FUMGEditorViewportClient(TWeakPtr<IUMGEditor> InUMGEdi
 	SetPreviewPage(InPreviewPage);
 }
 
-void FUMGEditorViewportClient::SetPreviewPage(UGUIPage* InPreviewPage)
+void FUMGEditorViewportClient::SetPreviewPage(AUserWidget* InPreviewPage)
 {
 	PreviewPage = InPreviewPage;
 }
diff --git a/Engine/Source/Editor/UMGEditor/Private/UMGEditorViewportClient.h b/Engine/Source/Editor/UMGEditor/Private/UMGEditorViewportClient.h
--- a/Engine/Source/Editor/UMGEditor/Private/UMGEditorViewportClient.h
+++ b/Engine/Source/Editor/UMGEditor/Private/UMGEditorViewportClient.h
@@ -9,7 +9,7 @@
 class FUMGEditorViewportClient : public FEditorViewportClient, public TSharedFromThis<FUMGEditorViewportClient>
 {
 public:
-	FUMGEditorViewportClient(TWeakPtr<IUMGEditor> InUMGEditor, FPreviewScene& InPreviewScene, UGUIPage* InPreviewGUIPage);
+	FUMGEditorViewportClient(TWeakPtr<IUMGEditor> InUMGEditor, FPreviewScene& InPreviewScene, AUserWidget* InPreviewGUIPage);
 
 	// FEditorViewportClient interface
 	virtual bool InputKey(FViewport* Viewport, int32 ControllerId, FKey Key, EInputEvent Event, float AmountDepressed = 1.f, bool bGamepad=false) OVERRIDE;
@@ -22,7 +22,7 @@ public:
 
 	void ResetCamera();
 
-	void SetPreviewPage(UGUIPage* InPreviewPage);
+	void SetPreviewPage(AUserWidget* InPreviewPage);
 
 protected:
 	// FEditorViewportClient interface
@@ -34,5 +34,8 @@ private:
 	/** Pointer back to the UMG editor tool that owns us */
 	TWeakPtr<IUMGEditor> UMGEditorPtr;
 
-	UGUIPage* PreviewPage;
+	AUserWidget* PreviewPage;
+
+	/** Overlay widget that contains widgets to draw on top of the game viewport */
+	TWeakPtr<class SOverlay> ViewportOverlayWidget;
 };
\ No newline at end of file
diff --git a/Engine/Source/Editor/UMGEditor/Private/WidgetBlueprintFactory.cpp b/Engine/Source/Editor/UMGEditor/Private/WidgetBlueprintFactory.cpp
new file mode 100644
--- /dev/null
+++ b/Engine/Source/Editor/UMGEditor/Private/WidgetBlueprintFactory.cpp
@@ -0,0 +1,57 @@
+// Copyright 1998-2014 Epic Games, Inc. All Rights Reserved.
+
+#include "UMGEditorPrivatePCH.h"
+
+#define LOCTEXT_NAMESPACE "UWidgetBlueprintFactory"
+
+/*------------------------------------------------------------------------------
+	UWidgetBlueprintFactory implementation.
+------------------------------------------------------------------------------*/
+
+UWidgetBlueprintFactory::UWidgetBlueprintFactory(const class FPostConstructInitializeProperties& PCIP)
+	: Super(PCIP)
+{
+	bCreateNew = true;
+	bEditAfterNew = true;
+	SupportedClass = UWidgetBlueprint::StaticClass();
+	ParentClass = AUserWidget::StaticClass();
+}
+
+bool UWidgetBlueprintFactory::ConfigureProperties()
+{
+	// TODO Make config dialog like the anim blueprints.
+	return true;
+};
+
+UObject* UWidgetBlueprintFactory::FactoryCreateNew(UClass* Class, UObject* InParent, FName Name, EObjectFlags Flags, UObject* Context, FFeedbackContext* Warn, FName CallingContext)
+{
+	// Make sure we are trying to factory a Anim Blueprint, then create and init one
+	check(Class->IsChildOf(UWidgetBlueprint::StaticClass()));
+
+	// If they selected an interface, force the parent class to be UInterface
+	if (BlueprintType == BPTYPE_Interface)
+	{
+		ParentClass = UInterface::StaticClass();
+	}
+
+	if ( ( ParentClass == NULL ) || !FKismetEditorUtilities::CanCreateBlueprintOfClass(ParentClass) || !ParentClass->IsChildOf(AUserWidget::StaticClass()) )
+	{
+		FFormatNamedArguments Args;
+		Args.Add( TEXT("ClassName"), (ParentClass != NULL) ? FText::FromString( ParentClass->GetName() ) : LOCTEXT("Null", "(null)") );
+		FMessageDialog::Open( EAppMsgType::Ok, FText::Format( LOCTEXT("CannotCreateAnimBlueprint", "Cannot create a Anim Blueprint based on the class '{0}'."), Args ) );
+		return NULL;
+	}
+	else
+	{
+		UWidgetBlueprint* NewBP = CastChecked<UWidgetBlueprint>(FKismetEditorUtilities::CreateBlueprint(ParentClass, InParent, Name, BlueprintType, UWidgetBlueprint::StaticClass(), CallingContext));
+
+		return NewBP;
+	}
+}
+
+UObject* UWidgetBlueprintFactory::FactoryCreateNew(UClass* Class, UObject* InParent, FName Name, EObjectFlags Flags, UObject* Context, FFeedbackContext* Warn)
+{
+	return FactoryCreateNew(Class, InParent, Name, Flags, Context, Warn, NAME_None);
+}
+
+#undef LOCTEXT_NAMESPACE
diff --git a/Engine/Source/Editor/UMGEditor/UMGEditor.Build.cs b/Engine/Source/Editor/UMGEditor/UMGEditor.Build.cs
--- a/Engine/Source/Editor/UMGEditor/UMGEditor.Build.cs
+++ b/Engine/Source/Editor/UMGEditor/UMGEditor.Build.cs
@@ -6,7 +6,13 @@ public class UMGEditor : ModuleRules
 {
 	public UMGEditor(TargetInfo Target)
 	{
-        PrivateIncludePaths.Add("Editor/UMGEditor/Private");	// For PCH includes (because they don't work with relative paths, yet)
+        PrivateIncludePaths.AddRange(
+            new string[] {
+                "Editor/GraphEditor/Private",
+				"Editor/Kismet/Private",
+                "Developer/AssetTools/Private/AssetTypeActions",
+                "Editor/UMGEditor/Private" // For PCH includes (because they don't work with relative paths, yet)
+            });
 
         PrivateIncludePathModuleNames.AddRange(
             new string[] {
diff --git a/Engine/Source/Runtime/UMG/Classes/BorderComponent.h b/Engine/Source/Runtime/UMG/Classes/BorderComponent.h
new file mode 100644
--- /dev/null
+++ b/Engine/Source/Runtime/UMG/Classes/BorderComponent.h
@@ -0,0 +1,76 @@
+// Copyright 1998-2014 Epic Games, Inc. All Rights Reserved.
+
+#pragma once
+
+#include "BorderComponent.generated.h"
+
+/** A border is a container widget that can contain one child widget, providing an opportunity to surround it with a border image and adjustable padding. */
+
+UCLASS(meta=(BlueprintSpawnableComponent), ClassGroup=UserInterface)
+class UMG_API UBorderComponent : public USlateNonLeafWidgetComponent
+{
+	GENERATED_UCLASS_BODY()
+
+	/** Horizontal positioning of the content within the border */
+	UPROPERTY(EditAnywhere, BlueprintReadOnly, Category=Appearance)
+	TEnumAsByte<EHorizontalAlignment> HorizontalAlignment;
+
+	/** Vertical positioning of the content within the border */
+	UPROPERTY(EditAnywhere, BlueprintReadOnly, Category=Appearance)
+	TEnumAsByte<EVerticalAlignment> VerticalAlignment;
+
+	/** The padding to add around the inner content. */
+	UPROPERTY(EditAnywhere, BlueprintReadWrite, Category=Appearance)
+	FMargin ContentPadding;
+
+	//@TODO: Should these be exposed?
+	// SLATE_EVENT( FPointerEventHandler, OnMouseButtonDown )
+	// SLATE_EVENT( FPointerEventHandler, OnMouseButtonUp )
+	// SLATE_EVENT( FPointerEventHandler, OnMouseMove )
+	// SLATE_EVENT( FPointerEventHandler, OnMouseDoubleClick )
+
+	/** Image to use for the border */
+	UPROPERTY(EditAnywhere, BlueprintReadWrite, Category=Appearance)
+	FSlateBrush BorderImage;
+
+	/** The scaling factor for the border */
+	UPROPERTY(EditAnywhere, BlueprintReadOnly, Category=Appearance)
+	FVector2D DesiredSizeScale;
+
+	/** The scaling factor for the border content */
+	UPROPERTY(EditAnywhere, BlueprintReadOnly, Category=Appearance)
+	FVector2D ContentScale;
+
+	/** Color and opacity multiplier of content in the border */
+	UPROPERTY(EditAnywhere, BlueprintReadWrite, Category=Appearance)
+	FLinearColor ContentColorAndOpacity;
+
+	/** Color and opacity of the actual border image */
+	UPROPERTY(EditAnywhere, BlueprintReadWrite, Category=Appearance)
+	FSlateColor BorderColorAndOpacity;
+
+	/** The foreground color of text and some glyphs that appear as the border's content. */
+	UPROPERTY(EditAnywhere, BlueprintReadWrite, Category=Appearance)
+	FSlateColor ForegroundColor;
+
+	/** Whether or not to show the disabled effect when this border is disabled */
+	UPROPERTY(EditAnywhere, BlueprintReadOnly, Category=Appearance)
+	bool bShowEffectWhenDisabled;
+
+protected:
+	TWeakPtr<class SBorder> MyBorder;
+
+protected:
+	// USlateWrapperComponent interface
+	virtual TSharedRef<SWidget> RebuildWidget() OVERRIDE;
+	// End of USlateWrapperComponent interface
+
+	// USlateNonLeafWidgetComponent interface
+	virtual void OnKnownChildrenChanged() OVERRIDE;
+	// End of USlateNonLeafWidgetComponent
+
+	FMargin GetContentPadding() const;
+	FLinearColor GetContentColor() const;
+	FSlateColor GetBorderColor() const;
+	FSlateColor GetForegroundColor() const;
+};
diff --git a/Engine/Source/Runtime/UMG/Classes/BoxPanelComponent.h b/Engine/Source/Runtime/UMG/Classes/BoxPanelComponent.h
new file mode 100644
--- /dev/null
+++ b/Engine/Source/Runtime/UMG/Classes/BoxPanelComponent.h
@@ -0,0 +1,88 @@
+// Copyright 1998-2014 Epic Games, Inc. All Rights Reserved.
+
+#pragma once
+
+#include "BoxPanelComponent.generated.h"
+
+
+USTRUCT()
+struct UMG_API FBoxPanelSlot
+{
+	GENERATED_USTRUCT_BODY()
+
+	/** Name of the slot (use this to attach children to the slot) */
+	UPROPERTY(EditAnywhere, Category=Appearance)
+	FName SlotName;
+
+	/** How much space this slot should occupy in the direction of the panel. */
+	UPROPERTY(EditAnywhere, Category=Appearance)
+	FSlateChildSize SlotSize;
+
+	/** Horizontal positioning of child within the allocated slot */
+	UPROPERTY(EditAnywhere, Category=Appearance)
+	TEnumAsByte<EHorizontalAlignment> HorizontalAlignment;
+
+	/** Vertical positioning of child within the allocated slot */
+	UPROPERTY(EditAnywhere, Category=Appearance)
+	TEnumAsByte<EVerticalAlignment> VerticalAlignment;
+
+	/** The padding to add around the child. */
+	UPROPERTY(EditAnywhere, Category=Appearance)
+	FMargin Padding;
+
+	/** Does this slot have a maximum size? */
+	UPROPERTY(EditAnywhere, Category=Appearance)
+	bool bHasMaximumSize;
+
+	/** The maximum size that this slot can be */
+	UPROPERTY(EditAnywhere, Category=Appearance, meta=(EditCondition="bHasMaximumSize"))
+	float MaxSize;
+
+	FBoxPanelSlot()
+		: SlotName(NAME_None) 
+		, SlotSize( FSlateChildSize() )
+		, HorizontalAlignment( HAlign_Fill )
+		, VerticalAlignment( VAlign_Fill )
+		, Padding( FMargin(0.0f) )
+		, bHasMaximumSize(false)
+		, MaxSize( 50.0f )
+	{
+	}
+};
+
+
+UCLASS(meta=(BlueprintSpawnableComponent), ClassGroup=UserInterface)
+class UMG_API UBoxPanelComponent : public USlateNonLeafWidgetComponent
+{
+	GENERATED_UCLASS_BODY()
+
+#if WITH_EDITOR
+	// UObject interface
+	virtual void PostEditChangeProperty(struct FPropertyChangedEvent& PropertyChangedEvent) OVERRIDE;
+	// End of UObject interface
+#endif
+
+	// USceneComponent interface
+	virtual bool HasAnySockets() const OVERRIDE;
+	virtual FTransform GetSocketTransform(FName InSocketName, ERelativeTransformSpace TransformSpace = RTS_World) const OVERRIDE;
+	virtual void QuerySupportedSockets(TArray<FComponentSocketDescription>& OutSockets) const OVERRIDE;
+	// End of USceneComponent interface
+
+protected:
+	UPROPERTY(EditAnywhere, Category=Slots)
+	TArray<FBoxPanelSlot> Slots;
+
+	UPROPERTY(EditAnywhere, Category=Appearance)
+	TEnumAsByte<EOrientation> Orientation;
+
+	TWeakPtr<class SConcreteBoxPanel> MyBox;
+
+protected:
+	// USlateWrapperComponent interface
+	virtual TSharedRef<SWidget> RebuildWidget() OVERRIDE;
+	// End of USlateWrapperComponent interface
+
+	// USlateNonLeafWidgetComponent interface
+	virtual void OnKnownChildrenChanged() OVERRIDE;
+	// End of USlateNonLeafWidgetComponent
+};
diff --git a/Engine/Source/Runtime/UMG/Classes/ButtonComponent.h b/Engine/Source/Runtime/UMG/Classes/ButtonComponent.h
new file mode 100644
--- /dev/null
+++ b/Engine/Source/Runtime/UMG/Classes/ButtonComponent.h
@@ -0,0 +1,93 @@
+// Copyright 1998-2014 Epic Games, Inc. All Rights Reserved.
+
+#pragma once
+
+#include "ButtonComponent.generated.h"
+
+DECLARE_DYNAMIC_MULTICAST_DELEGATE( FOnButtonComponentClicked );
+
+DECLARE_DYNAMIC_DELEGATE_RetVal( FText, FOnGetText );
+
+/** Buttons are clickable widgets */
+UCLASS(meta=(BlueprintSpawnableComponent), ClassGroup=UserInterface)
+class UMG_API UButtonComponent : public USlateNonLeafWidgetComponent
+{
+	GENERATED_UCLASS_BODY()
+
+protected:
+	/** Style of the button */
+	UPROPERTY(EditAnywhere, BlueprintReadOnly, Category=Appearance)
+	USlateWidgetStyleAsset* ButtonStyle;
+
+	/** Horizontal positioning of the content within the button */
+	UPROPERTY(EditAnywhere, BlueprintReadOnly, Category=Appearance)
+	TEnumAsByte<EHorizontalAlignment> HorizontalAlignment;
+
+	/** Vertical positioning of the content within the button */
+	UPROPERTY(EditAnywhere, BlueprintReadOnly, Category=Appearance)
+	TEnumAsByte<EVerticalAlignment> VerticalAlignment;
+
+	/** The padding to add around the button content. */
+	UPROPERTY(EditAnywhere, BlueprintReadWrite, Category=Appearance)
+	FMargin ContentPadding;
+
+	/** Should the button have text content inside? */
+	UPROPERTY(EditAnywhere, BlueprintReadOnly, Category=Appearance)
+	bool bDisplayText;
+
+	/** Text to display as the content of the button */
+	UPROPERTY(EditAnywhere, BlueprintReadWrite, Category=Appearance, meta=(EditCondition="bDisplayText"))
+	FText ButtonText;
+
+	/** Text to display as the content of the button */
+	UPROPERTY(EditAnywhere, BlueprintReadOnly, Category=Appearance, meta=(EditCondition="bDisplayText"))
+	USlateWidgetStyleAsset* ButtonTextStyle;
+
+	/** The scaling factor for the button border */
+	UPROPERTY(EditAnywhere, BlueprintReadOnly, Category=Appearance)
+	FVector2D DesiredSizeScale;
+
+	/** The scaling factor for the button content */
+	UPROPERTY(EditAnywhere, BlueprintReadOnly, Category=Appearance)
+	FVector2D ContentScale;
+
+	/** The color multiplier for the button images */
+	UPROPERTY(EditAnywhere, BlueprintReadWrite, Category=Appearance)
+	FSlateColor ButtonColorAndOpacity;
+
+	/** The foreground color of the button */
+	UPROPERTY(EditAnywhere, BlueprintReadWrite, Category=Appearance)
+	FSlateColor ForegroundColor;
+
+	/** Called when the button is clicked */
+	UPROPERTY(BlueprintAssignable)
+	FOnButtonComponentClicked OnClicked;
+
+	/** Called when the controller has instigated damage in any way */
+	UPROPERTY(BlueprintReadWrite, Category=Appearance)
+	FOnGetText OnGetText;
+
+	/** Called when the controller has instigated damage in any way */
+	UFUNCTION(BlueprintNativeEvent, Category=Appearance)
+	FText GetText() const;
+
+protected:
+	TAttribute<FText> TextAttribute;
+	TWeakPtr<class SButton> MyButton;
+
+protected:
+	// USlateWrapperComponent interface
+	virtual TSharedRef<SWidget> RebuildWidget() OVERRIDE;
+	// End of USlateWrapperComponent interface
+
+	// USlateNonLeafWidgetComponent interface
+	virtual void OnKnownChildrenChanged() OVERRIDE;
+	// End of USlateNonLeafWidgetComponent
+
+	FMargin GetContentPadding() const;
+	FText GetButtonText() const;
+	FSlateColor GetButtonColor() const;
+	FSlateColor GetForegroundColor() const;
+
+	virtual FReply SlateOnClickedCallback();
+};
diff --git a/Engine/Source/Runtime/UMG/Classes/CanvasPanelComponent.h b/Engine/Source/Runtime/UMG/Classes/CanvasPanelComponent.h
new file mode 100644
--- /dev/null
+++ b/Engine/Source/Runtime/UMG/Classes/CanvasPanelComponent.h
@@ -0,0 +1,43 @@
+// Copyright 1998-2014 Epic Games, Inc. All Rights Reserved.
+
+#pragma once
+
+#include "CanvasPanelComponent.generated.h"
+
+UCLASS(meta=(BlueprintSpawnableComponent), ClassGroup=UserInterface)
+class UMG_API UCanvasPanelComponent : public USlateNonLeafWidgetComponent
+{
+	GENERATED_UCLASS_BODY()
+
+#if WITH_EDITOR
+	// UObject interface
+	virtual void PostEditChangeProperty(struct FPropertyChangedEvent& PropertyChangedEvent) OVERRIDE;
+	// End of UObject interface
+#endif
+
+	// USceneComponent interface
+	virtual bool HasAnySockets() const OVERRIDE;
+	virtual FTransform GetSocketTransform(FName InSocketName, ERelativeTransformSpace TransformSpace = RTS_World) const OVERRIDE;
+	virtual void QuerySupportedSockets(TArray<FComponentSocketDescription>& OutSockets) const OVERRIDE;
+	// End of USceneComponent interface
+
+protected:
+	/** The items placed on the canvas */
+	UPROPERTY(EditAnywhere, Category=Slots)
+	TArray<class UCanvasPanelSlot*> Slots;
+
+	/** The desired size of the canvas */
+	UPROPERTY(EditAnywhere, Category=Appearance)
+	FVector2D DesiredCanvasSize;
+
+	TWeakPtr<class SFixedSizeCanvas> MyCanvas;
+
+protected:
+	// USlateWrapperComponent interface
+	virtual TSharedRef<SWidget> RebuildWidget() OVERRIDE;
+	// End of USlateWrapperComponent interface
+
+	// USlateNonLeafWidgetComponent interface
+	virtual void OnKnownChildrenChanged() OVERRIDE;
+	// End of USlateNonLeafWidgetComponent
+};
diff --git a/Engine/Source/Runtime/UMG/Classes/CanvasPanelSlot.h b/Engine/Source/Runtime/UMG/Classes/CanvasPanelSlot.h
new file mode 100644
--- /dev/null
+++ b/Engine/Source/Runtime/UMG/Classes/CanvasPanelSlot.h
@@ -0,0 +1,72 @@
+// Copyright 1998-2014 Epic Games, Inc. All Rights Reserved.
+
+#pragma once
+
+#include "CanvasPanelSlot.generated.h"
+
+USTRUCT()
+struct UMG_API FCanvasSlotData
+{
+	GENERATED_USTRUCT_BODY()
+
+	/** Name of the slot (use this to attach children to the slot) */
+	UPROPERTY(EditAnywhere, Category=Appearance)
+	FName SlotName;
+
+	/** Position. */
+	UPROPERTY(EditAnywhere, Category=Appearance)
+	FVector2D Position;
+
+	/** Size. */
+	UPROPERTY(EditAnywhere, Category=Appearance)
+	FVector2D Size;
+
+	/**
+	* Horizontal pivot position
+	*  Given a top aligned slot, where '+' represents the 
+	*  anchor point defined by PositionAttr.
+	*  
+	*   Left				Center				Right
+		+ _ _ _ _            _ _ + _ _          _ _ _ _ +
+		|		  |		   | 		   |	  |		    |
+		| _ _ _ _ |        | _ _ _ _ _ |	  | _ _ _ _ |
+	* 
+	*  Note: FILL is NOT supported.
+	*/
+	UPROPERTY(EditAnywhere, Category=Appearance)
+	TEnumAsByte<EHorizontalAlignment> HorizontalAlignment;
+
+	/**
+	* Vertical pivot position
+	*   Given a left aligned slot, where '+' represents the 
+	*   anchor point defined by PositionAttr.
+	*  
+	*   Top					Center			  Bottom
+	*	+_ _ _ _ _		 _ _ _ _ _		 _ _ _ _ _ 
+	*	|         |		| 		  |		|		  |
+	*	|         |     +		  |		|		  |
+	*	| _ _ _ _ |		| _ _ _ _ |		+ _ _ _ _ |
+	* 
+	*  Note: FILL is NOT supported.
+	*/
+	UPROPERTY(EditAnywhere, Category=Appearance)
+	TEnumAsByte<EVerticalAlignment> VerticalAlignment;
+
+	FCanvasSlotData()
+		: SlotName(NAME_None) 
+		, Position(FVector2D::ZeroVector)
+		, Size(FVector2D(1.0f, 1.0f))
+		, HorizontalAlignment(HAlign_Left)
+		, VerticalAlignment(VAlign_Top)
+	{
+	}
+};
+
+UCLASS()
+class UMG_API UCanvasPanelSlot : public UPanelSlot
+{
+	GENERATED_UCLASS_BODY()
+
+	UPROPERTY(EditAnywhere, Category=Slot)
+	FCanvasSlotData Data;
+};
\ No newline at end of file
diff --git a/Engine/Source/Runtime/UMG/Classes/CheckBoxComponent.h b/Engine/Source/Runtime/UMG/Classes/CheckBoxComponent.h
new file mode 100644
--- /dev/null
+++ b/Engine/Source/Runtime/UMG/Classes/CheckBoxComponent.h
@@ -0,0 +1,70 @@
+// Copyright 1998-2014 Epic Games, Inc. All Rights Reserved.
+
+#pragma once
+
+#include "CheckBoxComponent.generated.h"
+
+DECLARE_DYNAMIC_MULTICAST_DELEGATE_OneParam( FOnCheckBoxComponentStateChanged, bool, bIsChecked );
+
+/** Check box widget */
+UCLASS(meta=(BlueprintSpawnableComponent), ClassGroup=UserInterface)
+class UMG_API UCheckBoxComponent : public USlateWrapperComponent
+{
+	GENERATED_UCLASS_BODY()
+
+protected:
+	/** Style of the check box */
+	UPROPERTY(EditAnywhere, BlueprintReadOnly, Category=Appearance)
+	USlateWidgetStyleAsset* CheckBoxStyle;
+
+
+	/** Type of check box */
+	//SLATE_ATTRIBUTE( ESlateCheckBoxType::Type, Type )
+
+	/** Content to be placed next to the check box, or for a toggle button, the content to be placed inside the button */
+	//SLATE_DEFAULT_SLOT( FArguments, Content )
+
+	/** The name of the visual style for this checkbox */
+	UPROPERTY(EditAnywhere, BlueprintReadOnly, Category=Appearance)
+	FName CheckBoxStyleName;
+
+	/** The style structure (image references) for this checkbox' various states */
+	//SLATE_ARGUMENT( const FCheckBoxStyle *, CheckBoxStyle )
+
+	/** Called when the checked state has changed */
+	UPROPERTY(BlueprintAssignable)
+	FOnCheckBoxComponentStateChanged OnCheckStateChanged;
+
+	/** Whether the check box is currently in a checked state */
+	UPROPERTY(EditAnywhere, BlueprintReadWrite, Category=Appearance)
+	bool bIsChecked;
+
+	/** How the content of the toggle button should align within the given space */
+	UPROPERTY(EditAnywhere, BlueprintReadOnly, Category=Appearance)
+	TEnumAsByte<EHorizontalAlignment> HorizontalAlignment;
+
+	/** Spacing between the check box image and its content */
+	UPROPERTY(EditAnywhere, BlueprintReadOnly, Category=Appearance)
+	FMargin Padding;
+
+	/** Foreground color for the checkbox's content and parts */
+	UPROPERTY(EditAnywhere, BlueprintReadOnly, Category=Appearance)
+	FSlateColor ForegroundColor;
+
+	/** The color of the background border */
+	UPROPERTY(EditAnywhere, BlueprintReadOnly, Category=Appearance)
+	FSlateColor BorderBackgroundColor;
+
+	//SLATE_ATTRIBUTE( bool, ReadOnly )
+	//SLATE_ARGUMENT( bool, IsFocusable )
+
+	//SLATE_EVENT( FOnGetContent, OnGetMenuContent )
+
+protected:
+	// USlateWrapperComponent interface
+	virtual TSharedRef<SWidget> RebuildWidget() OVERRIDE;
+	// End of USlateWrapperComponent interface
+
+	ESlateCheckBoxState::Type GetCheckState() const;
+	void SlateOnCheckStateChangedCallback(ESlateCheckBoxState::Type NewState);
+};
diff --git a/Engine/Source/Runtime/UMG/Classes/CircularThrobberComponent.h b/Engine/Source/Runtime/UMG/Classes/CircularThrobberComponent.h
new file mode 100644
--- /dev/null
+++ b/Engine/Source/Runtime/UMG/Classes/CircularThrobberComponent.h
@@ -0,0 +1,35 @@
+// Copyright 1998-2014 Epic Games, Inc. All Rights Reserved.
+
+#pragma once
+
+#include "CircularThrobberComponent.generated.h"
+
+/**  A throbber widget that orients images in a spinning circle. */
+UCLASS(meta=(BlueprintSpawnableComponent), ClassGroup=UserInterface)
+class UMG_API UCircularThrobberComponent : public USlateLeafWidgetComponent
+{
+	GENERATED_UCLASS_BODY()
+
+protected:
+
+	/** How many pieces there are */
+	UPROPERTY(EditAnywhere, BlueprintReadOnly, Category=Appearance, meta=(ClampMin = "1", UIMin = "1", UIMax = "8"))
+	int32 NumberOfPieces;
+
+	/** The amount of time for a full circle (in seconds) */
+	UPROPERTY(EditAnywhere, BlueprintReadOnly, Category=Appearance)
+	float Period;
+
+	/** The radius of the circle */
+	UPROPERTY(EditAnywhere, BlueprintReadOnly, Category=Appearance)
+	float Radius;
+
+	/** Image to use for each segment of the throbber */
+	UPROPERTY(EditAnywhere, BlueprintReadOnly, Category=Appearance)
+	FSlateBrush PieceImage;
+
+protected:
+	// USlateWrapperComponent interface
+	virtual TSharedRef<SWidget> RebuildWidget() OVERRIDE;
+	// End of USlateWrapperComponent interface
+};
diff --git a/Engine/Source/Runtime/UMG/Classes/EditableTextBlockComponent.h b/Engine/Source/Runtime/UMG/Classes/EditableTextBlockComponent.h
new file mode 100644
--- /dev/null
+++ b/Engine/Source/Runtime/UMG/Classes/EditableTextBlockComponent.h
@@ -0,0 +1,20 @@
+// Copyright 1998-2014 Epic Games, Inc. All Rights Reserved.
+
+#pragma once
+
+#include "EditableTextBlockComponent.generated.h"
+
+/** Editable text box widget */
+UCLASS(meta=(BlueprintSpawnableComponent), ClassGroup=UserInterface)
+class UMG_API UEditableTextBlockComponent : public USlateWrapperComponent
+{
+	GENERATED_UCLASS_BODY()
+
+protected:
+
+
+protected:
+	// USlateWrapperComponent interface
+	virtual TSharedRef<SWidget> RebuildWidget() OVERRIDE;
+	// End of USlateWrapperComponent
+};
diff --git a/Engine/Source/Runtime/UMG/Classes/GUIPage.h b/Engine/Source/Runtime/UMG/Classes/GUIPage.h
deleted file mode 100644
diff --git a/Engine/Source/Runtime/UMG/Classes/GridPanelComponent.h b/Engine/Source/Runtime/UMG/Classes/GridPanelComponent.h
new file mode 100644
--- /dev/null
+++ b/Engine/Source/Runtime/UMG/Classes/GridPanelComponent.h
@@ -0,0 +1,100 @@
+// Copyright 1998-2014 Epic Games, Inc. All Rights Reserved.
+
+#pragma once
+
+#include "GridPanelComponent.generated.h"
+
+
+USTRUCT()
+struct UMG_API FGridPanelSlot
+{
+	GENERATED_USTRUCT_BODY()
+
+	/** Name of the slot (use this to attach children to the slot) */
+	UPROPERTY(VisibleAnywhere, Category=Appearance)
+	FName SlotName;
+
+	/** The row position within the grid. */
+	UPROPERTY(EditAnywhere, Category=Appearance, meta=(ClampMin = "0"))
+	int32 Row;
+
+	/** The column position within the grid. */
+	UPROPERTY(EditAnywhere, Category=Appearance, meta=(ClampMin = "0"))
+	int32 Column;
+
+	/** How many rows this this slot spans over */
+	UPROPERTY(EditAnywhere, Category=Appearance, meta=(ClampMin = "1"))
+	int32 RowSpan;
+
+	/** How many columns this slot spans over */
+	UPROPERTY(EditAnywhere, Category=Appearance, meta=(ClampMin = "1"))
+	int32 ColumnSpan;
+
+	/** Horizontal positioning of child within the allocated slot */
+	UPROPERTY(EditAnywhere, Category=Appearance)
+	TEnumAsByte<EHorizontalAlignment> HorizontalAlignment;
+
+	/** Vertical positioning of child within the allocated slot */
+	UPROPERTY(EditAnywhere, Category=Appearance)
+	TEnumAsByte<EVerticalAlignment> VerticalAlignment;
+
+	/** The padding to add around the child. */
+	UPROPERTY(EditAnywhere, Category=Appearance)
+	FMargin Padding;
+
+	FIntPoint AsPoint() const
+	{
+		return FIntPoint(Column, Row);
+	}
+
+	FGridPanelSlot()
+		: SlotName(NAME_None) 
+		, Row(0)
+		, Column(0)
+		, RowSpan(1)
+		, ColumnSpan(1)
+		, HorizontalAlignment(HAlign_Fill)
+		, VerticalAlignment(VAlign_Fill)
+		, Padding(FMargin(0.0f))
+	{
+	}
+};
+
+
+UCLASS(meta=(BlueprintSpawnableComponent), ClassGroup=UserInterface)
+class UMG_API UGridPanelComponent : public USlateNonLeafWidgetComponent
+{
+	GENERATED_UCLASS_BODY()
+
+#if WITH_EDITOR
+	// UObject interface
+	virtual void PostEditChangeProperty(struct FPropertyChangedEvent& PropertyChangedEvent) OVERRIDE;
+	// End of UObject interface
+#endif
+
+	// USceneComponent interface
+	virtual bool HasAnySockets() const OVERRIDE;
+	virtual FTransform GetSocketTransform(FName InSocketName, ERelativeTransformSpace TransformSpace = RTS_World) const OVERRIDE;
+	virtual void QuerySupportedSockets(TArray<FComponentSocketDescription>& OutSockets) const OVERRIDE;
+	// End of USceneComponent interface
+
+protected:
+	UPROPERTY(EditAnywhere, Category=Slots)
+	TArray<FGridPanelSlot> Slots;
+
+	TWeakPtr<class SGridPanel> MyGrid;
+
+	TSet<FIntPoint> UniqueSlotCoordinates;
+
+protected:
+	// USlateWrapperComponent interface
+	virtual TSharedRef<SWidget> RebuildWidget() OVERRIDE;
+	// End of USlateWrapperComponent interface
+
+	// USlateNonLeafWidgetComponent interface
+	virtual void OnKnownChildrenChanged() OVERRIDE;
+	// End of USlateNonLeafWidgetComponent
+
+	void RebuildUniqueSlotTable();
+	void EnsureSlotIsUnique(FGridPanelSlot& SlotConfig);
+};
diff --git a/Engine/Source/Runtime/UMG/Classes/ImageComponent.h b/Engine/Source/Runtime/UMG/Classes/ImageComponent.h
new file mode 100644
--- /dev/null
+++ b/Engine/Source/Runtime/UMG/Classes/ImageComponent.h
@@ -0,0 +1,27 @@
+// Copyright 1998-2014 Epic Games, Inc. All Rights Reserved.
+
+#pragma once
+
+#include "ImageComponent.generated.h"
+
+/** Image widget */
+UCLASS(meta=(BlueprintSpawnableComponent), ClassGroup=UserInterface)
+class UMG_API UImageComponent : public USlateLeafWidgetComponent
+{
+	GENERATED_UCLASS_BODY()
+
+	/** Image to draw */
+	UPROPERTY(EditAnywhere, BlueprintReadWrite, Category=Appearance)
+	FSlateBrush Image;
+
+	/** Color and opacity */
+	UPROPERTY(EditAnywhere, BlueprintReadWrite, Category=Appearance)
+ 	FSlateColor ColorAndOpacity;
+
+protected:
+	// USlateWrapperComponent interface
+	virtual TSharedRef<SWidget> RebuildWidget() OVERRIDE;
+	// End of USlateWrapperComponent interface
+
+ 	FSlateColor GetColorAndOpacity() const;
+};
diff --git a/Engine/Source/Runtime/UMG/Classes/PanelSlot.h b/Engine/Source/Runtime/UMG/Classes/PanelSlot.h
new file mode 100644
--- /dev/null
+++ b/Engine/Source/Runtime/UMG/Classes/PanelSlot.h
@@ -0,0 +1,11 @@
+// Copyright 1998-2014 Epic Games, Inc. All Rights Reserved.
+
+#pragma once
+
+#include "PanelSlot.generated.h"
+
+UCLASS()
+class UMG_API UPanelSlot : public UObject
+{
+	GENERATED_UCLASS_BODY()
+};
diff --git a/Engine/Source/Runtime/UMG/Classes/ScrollBoxComponent.h b/Engine/Source/Runtime/UMG/Classes/ScrollBoxComponent.h
new file mode 100644
--- /dev/null
+++ b/Engine/Source/Runtime/UMG/Classes/ScrollBoxComponent.h
@@ -0,0 +1,18 @@
+// Copyright 1998-2014 Epic Games, Inc. All Rights Reserved.
+
+#pragma once
+
+#include "ScrollBoxComponent.generated.h"
+
+UCLASS(meta=(BlueprintSpawnableComponent), ClassGroup=UserInterface)
+class UMG_API UScrollBoxComponent : public USlateWrapperComponent
+{
+	GENERATED_UCLASS_BODY()
+
+protected:
+
+protected:
+	// USlateWrapperComponent interface
+	virtual TSharedRef<SWidget> RebuildWidget() OVERRIDE;
+	// End of USlateWrapperComponent interface
+};
diff --git a/Engine/Source/Runtime/UMG/Classes/SlateLeafWidgetComponent.h b/Engine/Source/Runtime/UMG/Classes/SlateLeafWidgetComponent.h
new file mode 100644
--- /dev/null
+++ b/Engine/Source/Runtime/UMG/Classes/SlateLeafWidgetComponent.h
@@ -0,0 +1,16 @@
+// Copyright 1998-2014 Epic Games, Inc. All Rights Reserved.
+
+#pragma once
+
+#include "SlateLeafWidgetComponent.generated.h"
+
+/** Leaf widget that cannot have any children */
+UCLASS(Abstract)
+class UMG_API USlateLeafWidgetComponent: public USlateWrapperComponent
+{
+	GENERATED_UCLASS_BODY()
+
+	// USceneComponent interface
+	virtual bool CanAttachAsChild(USceneComponent* ChildComponent, FName SocketName) const OVERRIDE;
+	// End of USceneComponent interface
+};
diff --git a/Engine/Source/Runtime/UMG/Classes/SlateNonLeafWidgetComponent.h b/Engine/Source/Runtime/UMG/Classes/SlateNonLeafWidgetComponent.h
new file mode 100644
--- /dev/null
+++ b/Engine/Source/Runtime/UMG/Classes/SlateNonLeafWidgetComponent.h
@@ -0,0 +1,42 @@
+// Copyright 1998-2014 Epic Games, Inc. All Rights Reserved.
+
+#pragma once
+
+#include "SlateNonLeafWidgetComponent.generated.h"
+
+UCLASS(Abstract)
+class UMG_API USlateNonLeafWidgetComponent : public USlateWrapperComponent
+{
+	GENERATED_UCLASS_BODY()
+
+	// UActorComponent interface
+	virtual void TickComponent(float DeltaTime, enum ELevelTick TickType, FActorComponentTickFunction* ThisTickFunction) OVERRIDE;
+	// End of UActorComponent interface
+
+#if WITH_EDITOR
+	// UActorComponent interface
+	virtual void CheckForErrors() OVERRIDE;
+	// End of UActorComponent interface
+#endif
+
+	// USceneComponent interface
+	virtual bool CanAttachAsChild(USceneComponent* ChildComponent, FName SocketName) const OVERRIDE;
+	// End of USceneComponent interface
+
+protected:
+	// This function is called when the known children list has changed; use this to repopulate child slots in derived classes
+	virtual void OnKnownChildrenChanged() {}
+
+	// Helper to get the first child (for bCanHaveMultipleChildren = false subclasses)
+	USlateWrapperComponent* GetFirstWrappedChild() const;
+
+	// Always returns something valid, even if there are no children
+	TSharedRef<SWidget> GetFirstChildWidget() const;
+
+protected:
+	// The list of known children
+	TSet<TWeakObjectPtr<USlateWrapperComponent>> KnownChildren;
+
+	// Can we have more than one child?
+	bool bCanHaveMultipleChildren;
+};
diff --git a/Engine/Source/Runtime/UMG/Classes/SlateWrapperComponent.h b/Engine/Source/Runtime/UMG/Classes/SlateWrapperComponent.h
new file mode 100644
--- /dev/null
+++ b/Engine/Source/Runtime/UMG/Classes/SlateWrapperComponent.h
@@ -0,0 +1,36 @@
+// Copyright 1998-2014 Epic Games, Inc. All Rights Reserved.
+
+#pragma once
+
+#include "SlateWrapperTypes.h"
+#include "SlateWrapperComponent.generated.h"
+
+/** This is the base class for all Slate widget wrapper components */
+UCLASS(Abstract, hideCategories=(Activation, "Components|Activation", ComponentReplication, LOD, Rendering))
+class UMG_API USlateWrapperComponent : public USceneComponent
+{
+	GENERATED_UCLASS_BODY()
+
+#if WITH_EDITOR
+	UPROPERTY(Transient, EditAnywhere, BlueprintReadWrite, Category=Layout)
+	UPanelSlot* Slot;
+#endif
+
+	// UActorComponent interface
+	virtual void OnRegister() OVERRIDE;
+	virtual void OnUnregister() OVERRIDE;
+	// End of UActorComponent interface
+
+	TSharedRef<SWidget> GetWidget();
+
+	// Utility methods
+	//@TODO: Should move elsewhere
+	static EVisibility ConvertSerializedVisibilityToRuntime(TEnumAsByte<ESlateVisibility::Type> Input);
+	static FSizeParam ConvertSerializedSizeParamToRuntime(const FSlateChildSize& Input);
+
+protected:
+	virtual TSharedRef<SWidget> RebuildWidget();
+
+protected:
+	TSharedPtr<SWidget> MyWidget;
+};
diff --git a/Engine/Source/Runtime/UMG/Classes/SlateWrapperTypes.h b/Engine/Source/Runtime/UMG/Classes/SlateWrapperTypes.h
new file mode 100644
--- /dev/null
+++ b/Engine/Source/Runtime/UMG/Classes/SlateWrapperTypes.h
@@ -0,0 +1,60 @@
+// Copyright 1998-2014 Epic Games, Inc. All Rights Reserved.
+
+#pragma once
+
+#include "SlateWrapperTypes.generated.h"
+
+
+#define BIND_UOBJECT_ATTRIBUTE(Type, Function) \
+	TAttribute<Type>::Create( TAttribute<Type>::FGetter::CreateUObject( this, &ThisClass::Function ) )
+
+#define BIND_UOBJECT_DELEGATE(Type, Function) \
+	Type::CreateUObject( this, &ThisClass::Function )
+
+/** Is an entity visible? */
+UENUM()
+namespace ESlateVisibility
+{
+	enum Type
+	{
+		/** Default widget visibility - visible and can interactive with the cursor */
+		Visible,
+		/** Not visible and takes up no space in the layout; can never be clicked on because it takes up no space. */
+		Collapsed,
+		/** Not visible, but occupies layout space. Not interactive for obvious reasons. */
+		Hidden,
+		/** Visible to the user, but only as art. The cursors hit tests will never see this widget. */
+		HitTestInvisible,
+		/** Same as HitTestInvisible, but doesn't apply to child widgets. */
+		SelfHitTestInvisible
+	};
+}
+
+UENUM()
+namespace ESlateSizeRule
+{
+	enum Type
+	{
+		Automatic,
+		AspectRatio,
+		Stretch
+	};
+}
+
+USTRUCT()
+struct FSlateChildSize
+{
+	GENERATED_USTRUCT_BODY()
+
+	UPROPERTY(EditAnywhere, Category=Appearance)
+	float StretchValue;
+
+	UPROPERTY(EditAnywhere, Category=Appearance)
+	TEnumAsByte<ESlateSizeRule::Type> SizeRule;
+
+	FSlateChildSize()
+		: StretchValue(1.0f)
+		, SizeRule(ESlateSizeRule::Stretch)
+	{
+	}
+};
diff --git a/Engine/Source/Runtime/UMG/Classes/SpacerComponent.h b/Engine/Source/Runtime/UMG/Classes/SpacerComponent.h
new file mode 100644
--- /dev/null
+++ b/Engine/Source/Runtime/UMG/Classes/SpacerComponent.h
@@ -0,0 +1,23 @@
+// Copyright 1998-2014 Epic Games, Inc. All Rights Reserved.
+
+#pragma once
+
+#include "SpacerComponent.generated.h"
+
+/** A spacer widget; it does not have a visual representation, and just provides padding between other widgets. */
+UCLASS(meta=(BlueprintSpawnableComponent), ClassGroup=UserInterface)
+class UMG_API USpacerComponent : public USlateWrapperComponent
+{
+	GENERATED_UCLASS_BODY()
+
+	/** The size of the spacer */
+	UPROPERTY(EditAnywhere, BlueprintReadWrite, Category=Appearance)
+	FVector2D Size;
+
+protected:
+	// USlateWrapperComponent interface
+	virtual TSharedRef<SWidget> RebuildWidget() OVERRIDE;
+	// End of USlateWrapperComponent interface
+
+	FVector2D GetSpacerSize() const;
+};
diff --git a/Engine/Source/Runtime/UMG/Classes/SpinningImageComponent.h b/Engine/Source/Runtime/UMG/Classes/SpinningImageComponent.h
new file mode 100644
--- /dev/null
+++ b/Engine/Source/Runtime/UMG/Classes/SpinningImageComponent.h
@@ -0,0 +1,21 @@
+// Copyright 1998-2014 Epic Games, Inc. All Rights Reserved.
+
+#pragma once
+
+#include "SpinningImageComponent.generated.h"
+
+/** A widget that displays a spinning image. */
+UCLASS(meta=(BlueprintSpawnableComponent), ClassGroup=UserInterface)
+class UMG_API USpinningImageComponent : public UImageComponent
+{
+	GENERATED_UCLASS_BODY()
+
+	/** The amount of time in seconds for a full rotation */
+	UPROPERTY(EditAnywhere, BlueprintReadOnly, Category=Appearance)
+	float Period;
+
+protected:
+	// USlateWrapperComponent interface
+	virtual TSharedRef<SWidget> RebuildWidget() OVERRIDE;
+	// End of USlateWrapperComponent interface
+};
diff --git a/Engine/Source/Runtime/UMG/Classes/TextBlockComponent.h b/Engine/Source/Runtime/UMG/Classes/TextBlockComponent.h
new file mode 100644
--- /dev/null
+++ b/Engine/Source/Runtime/UMG/Classes/TextBlockComponent.h
@@ -0,0 +1,31 @@
+// Copyright 1998-2014 Epic Games, Inc. All Rights Reserved.
+
+#pragma once
+
+#include "TextBlockComponent.generated.h"
+
+UCLASS(meta=(BlueprintSpawnableComponent), ClassGroup=UserInterface)
+class UMG_API UTextBlockComponent : public USlateLeafWidgetComponent
+{
+	GENERATED_UCLASS_BODY()
+
+protected:
+	UPROPERTY(EditAnywhere, Category=TextBlock)
+	FText DisplayText;
+
+	//@TODO: Expose FSlateColor
+	UPROPERTY(EditAnywhere, Category=TextBlock)
+	FLinearColor ColorAndOpacity;
+
+	UPROPERTY(EditAnywhere, Category=TextBlock)
+	FLinearColor ShadowColorAndOpacity;
+
+protected:
+	// USlateWrapperComponent interface
+	virtual TSharedRef<SWidget> RebuildWidget() OVERRIDE;
+	// End of USlateWrapperComponent interface
+
+	FString GetDisplayText() const;
+	FSlateColor GetColorAndOpacity() const;
+	FLinearColor GetShadowColorAndOpacity() const;
+};
diff --git a/Engine/Source/Runtime/UMG/Classes/ThrobberComponent.h b/Engine/Source/Runtime/UMG/Classes/ThrobberComponent.h
new file mode 100644
--- /dev/null
+++ b/Engine/Source/Runtime/UMG/Classes/ThrobberComponent.h
@@ -0,0 +1,37 @@
+// Copyright 1998-2014 Epic Games, Inc. All Rights Reserved.
+
+#pragma once
+
+#include "ThrobberComponent.generated.h"
+
+/** A Throbber widget that shows several zooming circles in a row. */
+UCLASS(meta=(BlueprintSpawnableComponent), ClassGroup=UserInterface)
+class UMG_API UThrobberComponent : public USlateLeafWidgetComponent
+{
+	GENERATED_UCLASS_BODY()
+
+	/** How many pieces there are */
+	UPROPERTY(EditAnywhere, BlueprintReadOnly, Category=Appearance, meta=(ClampMin = "1", UIMin = "1", UIMax = "8"))
+	int32 NumberOfPieces;
+
+	/** Should the pieces animate horizontally? */
+	UPROPERTY(EditAnywhere, BlueprintReadWrite, Category=Appearance)
+	bool bAnimateHorizontally;
+
+	/** Should the pieces animate vertically? */
+	UPROPERTY(EditAnywhere, BlueprintReadWrite, Category=Appearance)
+	bool bAnimateVertically;
+
+	/** Should the pieces animate their opacity? */
+	UPROPERTY(EditAnywhere, BlueprintReadWrite, Category=Appearance)
+	bool bAnimateOpacity;
+
+	/** Image to use for each segment of the throbber */
+	UPROPERTY(EditAnywhere, BlueprintReadOnly, Category=Appearance)
+	FSlateBrush PieceImage;
+
+protected:
+	// USlateWrapperComponent interface
+	virtual TSharedRef<SWidget> RebuildWidget() OVERRIDE;
+	// End of USlateWrapperComponent interface
+};
diff --git a/Engine/Source/Runtime/UMG/Classes/UserWidget.h b/Engine/Source/Runtime/UMG/Classes/UserWidget.h
new file mode 100644
--- /dev/null
+++ b/Engine/Source/Runtime/UMG/Classes/UserWidget.h
@@ -0,0 +1,26 @@
+// Copyright 1998-2014 Epic Games, Inc. All Rights Reserved.
+
+#pragma once
+
+#include "UserWidget.generated.h"
+
+UCLASS(Abstract, hideCategories=(Object, Actor))
+class UMG_API AUserWidget : public AActor
+{
+	GENERATED_UCLASS_BODY()
+
+	// AActor interface
+	virtual void PostActorCreated() OVERRIDE;
+	virtual void Destroyed() OVERRIDE;
+	virtual void RerunConstructionScripts() OVERRIDE;
+	// End of AActor interface
+
+	TSharedRef<SWidget> GetWidget();
+
+private:
+	TSharedPtr<SOverlay> MyWrapperWidget;
+
+private:
+	void EnsureWidgetExists();
+	void RebuildWrapperWidget();
+};
diff --git a/Engine/Source/Runtime/UMG/Classes/WidgetBlueprint.h b/Engine/Source/Runtime/UMG/Classes/WidgetBlueprint.h
new file mode 100644
--- /dev/null
+++ b/Engine/Source/Runtime/UMG/Classes/WidgetBlueprint.h
@@ -0,0 +1,22 @@
+// Copyright 1998-2014 Epic Games, Inc. All Rights Reserved.
+
+#pragma once
+
+#include "WidgetBlueprint.generated.h"
+
+UCLASS(dependson=(UBlueprint, AUserWidget), BlueprintType)
+class UMG_API UWidgetBlueprint : public UBlueprint
+{
+	GENERATED_UCLASS_BODY()
+
+	//TODO Add Hierarchy.	
+	
+#if WITH_EDITOR
+	// UBlueprint interface
+	virtual bool SupportedByDefaultBlueprintFactory() const OVERRIDE
+	{
+		return false;
+	}
+	// End of UBlueprint interface
+#endif
+};
\ No newline at end of file
diff --git a/Engine/Source/Runtime/UMG/Classes/WidgetBlueprintGeneratedClass.h b/Engine/Source/Runtime/UMG/Classes/WidgetBlueprintGeneratedClass.h
new file mode 100644
--- /dev/null
+++ b/Engine/Source/Runtime/UMG/Classes/WidgetBlueprintGeneratedClass.h
@@ -0,0 +1,12 @@
+// Copyright 1998-2014 Epic Games, Inc. All Rights Reserved.
+
+#pragma once
+
+#include "WidgetBlueprintGeneratedClass.generated.h"
+
+UCLASS(dependson=(UWidgetBlueprint))
+class UMG_API UWidgetBlueprintGeneratedClass : public UBlueprintGeneratedClass
+{
+	GENERATED_UCLASS_BODY()
+
+};
\ No newline at end of file
diff --git a/Engine/Source/Runtime/UMG/Private/BorderComponent.cpp b/Engine/Source/Runtime/UMG/Private/BorderComponent.cpp
new file mode 100644
--- /dev/null
+++ b/Engine/Source/Runtime/UMG/Private/BorderComponent.cpp
@@ -0,0 +1,82 @@
+// Copyright 1998-2014 Epic Games, Inc. All Rights Reserved.
+
+#include "UMGPrivatePCH.h"
+
+#define LOCTEXT_NAMESPACE "SlateComponentWrappers"
+
+/////////////////////////////////////////////////////
+// UBorderComponent
+
+UBorderComponent::UBorderComponent(const FPostConstructInitializeProperties& PCIP)
+	: Super(PCIP)
+{
+	bCanHaveMultipleChildren = false;
+
+	SBorder::FArguments BorderDefaults;
+
+	HorizontalAlignment = BorderDefaults._HAlign;
+	VerticalAlignment = BorderDefaults._VAlign;
+	ContentPadding = BorderDefaults._Padding.Get();
+	//BorderImage = *BorderDefaults._BorderImage.Get();
+	DesiredSizeScale = BorderDefaults._DesiredSizeScale.Get();
+	ContentScale = BorderDefaults._ContentScale.Get();
+	ContentColorAndOpacity = BorderDefaults._ColorAndOpacity.Get();
+	BorderColorAndOpacity = BorderDefaults._BorderBackgroundColor.Get();
+	ForegroundColor = BorderDefaults._ForegroundColor.Get();
+	bShowEffectWhenDisabled = BorderDefaults._ShowEffectWhenDisabled.Get();
+}
+
+TSharedRef<SWidget> UBorderComponent::RebuildWidget()
+{
+	TSharedRef<SBorder> NewBorder = SNew(SBorder)
+		.HAlign(HorizontalAlignment)
+		.VAlign(VerticalAlignment)
+		.Padding(BIND_UOBJECT_ATTRIBUTE(FMargin, GetContentPadding))
+		// SLATE_EVENT( FPointerEventHandler, OnMouseButtonDown )
+		// SLATE_EVENT( FPointerEventHandler, OnMouseButtonUp )
+		// SLATE_EVENT( FPointerEventHandler, OnMouseMove )
+		// SLATE_EVENT( FPointerEventHandler, OnMouseDoubleClick )
+		.BorderImage(&BorderImage)
+		.ContentScale(ContentScale)
+		.DesiredSizeScale(DesiredSizeScale)
+		.ColorAndOpacity(BIND_UOBJECT_ATTRIBUTE(FLinearColor, GetContentColor))
+		.BorderBackgroundColor(BIND_UOBJECT_ATTRIBUTE(FSlateColor, GetBorderColor))
+		.ForegroundColor(BIND_UOBJECT_ATTRIBUTE(FSlateColor, GetForegroundColor))
+		.ShowEffectWhenDisabled(bShowEffectWhenDisabled);
+
+	MyBorder = NewBorder;
+	return NewBorder;
+}
+
+void UBorderComponent::OnKnownChildrenChanged()
+{
+	TSharedPtr<SBorder> Border = MyBorder.Pin();
+	if (Border.IsValid())
+	{
+		Border->SetContent(GetFirstChildWidget());
+	}
+}
+
+FMargin UBorderComponent::GetContentPadding() const
+{
+	return ContentPadding;
+}
+
+FLinearColor UBorderComponent::GetContentColor() const
+{
+	return ContentColorAndOpacity;
+}
+
+FSlateColor UBorderComponent::GetBorderColor() const
+{
+	return BorderColorAndOpacity;
+}
+
+FSlateColor UBorderComponent::GetForegroundColor() const
+{
+	return ForegroundColor;
+}
+
+/////////////////////////////////////////////////////
+
+#undef LOCTEXT_NAMESPACE
diff --git a/Engine/Source/Runtime/UMG/Private/BoxPanelComponent.cpp b/Engine/Source/Runtime/UMG/Private/BoxPanelComponent.cpp
new file mode 100644
--- /dev/null
+++ b/Engine/Source/Runtime/UMG/Private/BoxPanelComponent.cpp
@@ -0,0 +1,215 @@
+// Copyright 1998-2014 Epic Games, Inc. All Rights Reserved.
+
+#include "UMGPrivatePCH.h"
+
+/////////////////////////////////////////////////////
+// SConcreteBoxPanel
+
+class SConcreteBoxPanel : public SBoxPanel
+{
+public:
+	class FSlot : public SBoxPanel::FSlot
+	{
+	public:
+		FSlot& SetSizeParam(const FSizeParam& InParam)
+		{
+			SizeParam = InParam;
+			return *this;
+		}
+
+		FSlot& MaxWidth( const TAttribute< float >& InMaxWidth )
+		{
+			MaxSize = InMaxWidth;
+			return *this;
+		}
+
+		FSlot& Padding( TAttribute<FMargin> InPadding )
+		{
+			SlotPadding = InPadding;
+			return *this;
+		}
+
+		FSlot& HAlign( EHorizontalAlignment InHAlignment )
+		{
+			HAlignment = InHAlignment;
+			return *this;
+		}
+
+		FSlot& VAlign( EVerticalAlignment InVAlignment )
+		{
+			VAlignment = InVAlignment;
+			return *this;
+		}
+
+		FSlot& operator[]( TSharedRef<SWidget> InWidget )
+		{
+			Widget = InWidget;
+			return *this;
+		}
+
+		FSlot& Expose( FSlot*& OutVarToInit )
+		{
+			OutVarToInit = this;
+			return *this;
+		}
+	};
+
+	static FSlot& Slot()
+	{
+		return *(new FSlot());
+	}
+
+	SLATE_BEGIN_ARGS( SConcreteBoxPanel )
+		{
+			_Visibility = EVisibility::SelfHitTestInvisible;
+		}
+		SLATE_SUPPORTS_SLOT(SConcreteBoxPanel::FSlot)
+	SLATE_END_ARGS()
+
+	FSlot& AddSlot()
+	{
+		SConcreteBoxPanel::FSlot& NewSlot = SConcreteBoxPanel::Slot();
+		this->Children.Add( &NewSlot );
+		return NewSlot;
+	}
+
+	void Construct(const FArguments& InArgs, EOrientation InOrientation)
+	{
+		DirtyHackSetOrientation(InOrientation);
+
+		const int32 NumSlots = InArgs.Slots.Num();
+		for (int32 SlotIndex = 0; SlotIndex < NumSlots; ++SlotIndex)
+		{
+			Children.Add(InArgs.Slots[SlotIndex]);
+		}
+	}
+
+	SConcreteBoxPanel()
+		: SBoxPanel(Orient_Horizontal)
+	{
+	}
+
+protected:
+	//@TODO: Make Orientation not const in the base class; it's already protected and AFAIK there are no
+	// negative consequences from altering it before the first tick/computedesiredsize call
+	void DirtyHackSetOrientation(EOrientation InOrientation)
+	{
+		EOrientation& NonConstOrientation = *(const_cast<EOrientation*>(&Orientation));
+		NonConstOrientation = InOrientation;
+	}
+};
+
+/////////////////////////////////////////////////////
+// UBoxPanelComponent
+
+UBoxPanelComponent::UBoxPanelComponent(const FPostConstructInitializeProperties& PCIP)
+	: Super(PCIP)
+{
+}
+
+TSharedRef<SWidget> UBoxPanelComponent::RebuildWidget()
+{
+	SConcreteBoxPanel::FArguments Args;
+
+	TSharedRef<SConcreteBoxPanel> Box = SNew(SConcreteBoxPanel, Orientation);
+	MyBox = Box;
+
+	return Box;
+}
+
+void UBoxPanelComponent::OnKnownChildrenChanged()
+{
+	TSharedPtr<SConcreteBoxPanel> Box = MyBox.Pin();
+	if (Box.IsValid())
+	{
+		Box->ClearChildren();
+
+		// Add slots
+		TMap<FName, int32> SlotNameToIndex;
+
+		for (int32 SlotIndex = 0; SlotIndex < Slots.Num(); ++SlotIndex)
+		{
+			const FBoxPanelSlot& SlotConfig = Slots[SlotIndex];
+
+			SlotNameToIndex.Add(SlotConfig.SlotName, SlotIndex);
+
+			FSizeParam SlotSize = USlateWrapperComponent::ConvertSerializedSizeParamToRuntime(SlotConfig.SlotSize);
+
+			auto NewSlot = Box->AddSlot()
+				.Padding(SlotConfig.Padding)
+				.HAlign(SlotConfig.HorizontalAlignment)
+				.VAlign(SlotConfig.VerticalAlignment)
+				.SetSizeParam(SlotSize)
+				.MaxWidth(SlotConfig.bHasMaximumSize ? SlotConfig.MaxSize : 0.0f)
+			[
+				SNullWidget::NullWidget
+			];
+		}
+
+		// Place widgets in their slots, and add anything left over to the end
+		for (int32 ComponentIndex = 0; ComponentIndex < AttachChildren.Num(); ++ComponentIndex)
+		{
+			if (USlateWrapperComponent* Wrapper = Cast<USlateWrapperComponent>(AttachChildren[ComponentIndex]))
+			{
+				if (Wrapper->IsRegistered())
+				{
+					if (int32* pSlotIndex = SlotNameToIndex.Find(Wrapper->AttachSocketName))
+					{
+						auto Foo = (TPanelChildren<SHorizontalBox::FSlot>*)(Box->GetChildren());
+						(*Foo)[*pSlotIndex].Widget = Wrapper->GetWidget();
+					}
+					else
+					{
+						Box->AddSlot()
+						[
+							Wrapper->GetWidget()
+						];
+					}
+				}
+			}
+		}
+	}
+}
+
+bool UBoxPanelComponent::HasAnySockets() const
+{
+	return true;
+}
+
+FTransform UBoxPanelComponent::GetSocketTransform(FName InSocketName, ERelativeTransformSpace TransformSpace) const
+{
+	return FTransform::Identity;
+}
+
+void UBoxPanelComponent::QuerySupportedSockets(TArray<FComponentSocketDescription>& OutSockets) const
+{
+	for (int32 SlotIndex = 0; SlotIndex < Slots.Num(); ++SlotIndex)
+	{
+		FName SocketName = Slots[SlotIndex].SlotName;
+		new (OutSockets) FComponentSocketDescription(SocketName, EComponentSocketType::Socket);
+	}
+}
+
+#if WITH_EDITOR
+void UBoxPanelComponent::PostEditChangeProperty(struct FPropertyChangedEvent& PropertyChangedEvent)
+{
+	// Ensure the slots have unique names
+	int32 SlotNumbering = 1;
+
+	TSet<FName> UniqueSlotNames;
+	for (int32 SlotIndex = 0; SlotIndex < Slots.Num(); ++SlotIndex)
+	{
+		FName& SlotName = Slots[SlotIndex].SlotName;
+
+		if ((SlotName == NAME_None) || UniqueSlotNames.Contains(SlotName))
+		{
+			do 
+			{
+				SlotName = FName(TEXT("Slot"), SlotNumbering++);
+			} while (UniqueSlotNames.Contains(SlotName));
+		}
+
+		UniqueSlotNames.Add(SlotName);
+	}
+}
+#endif
\ No newline at end of file
diff --git a/Engine/Source/Runtime/UMG/Private/ButtonComponent.cpp b/Engine/Source/Runtime/UMG/Private/ButtonComponent.cpp
new file mode 100644
--- /dev/null
+++ b/Engine/Source/Runtime/UMG/Private/ButtonComponent.cpp
@@ -0,0 +1,134 @@
+// Copyright 1998-2014 Epic Games, Inc. All Rights Reserved.
+
+#include "UMGPrivatePCH.h"
+
+#define LOCTEXT_NAMESPACE "SlateComponentWrappers"
+
+/////////////////////////////////////////////////////
+// UButtonComponent
+
+UButtonComponent::UButtonComponent(const FPostConstructInitializeProperties& PCIP)
+	: Super(PCIP)
+{
+	SButton::FArguments ButtonDefaults;
+
+	ButtonStyle = NULL;
+	HorizontalAlignment = ButtonDefaults._HAlign;
+	VerticalAlignment = ButtonDefaults._VAlign;
+	ContentPadding = ButtonDefaults._ContentPadding.Get();
+
+	bDisplayText = true;
+	ButtonText = LOCTEXT("ButtonDefaultValue", "This is a button");
+	ButtonTextStyle = NULL;
+
+	DesiredSizeScale = ButtonDefaults._DesiredSizeScale.Get();
+	ContentScale = ButtonDefaults._ContentScale.Get();
+
+	ButtonColorAndOpacity = ButtonDefaults._ButtonColorAndOpacity.Get();
+	ForegroundColor = ButtonDefaults._ForegroundColor.Get();
+}
+
+TSharedRef<SWidget> UButtonComponent::RebuildWidget()
+{
+	SButton::FArguments ButtonDefaults;
+
+	const FButtonStyle* StylePtr = (ButtonStyle != NULL) ? ButtonStyle->GetStyle<FButtonStyle>() : NULL;
+	if (StylePtr == NULL)
+	{
+		StylePtr = ButtonDefaults._ButtonStyle;
+	}
+
+	const FTextBlockStyle* TextStylePtr = ( ButtonTextStyle != NULL ) ? ButtonTextStyle->GetStyle<FTextBlockStyle>() : NULL;
+	if ( TextStylePtr == NULL )
+	{
+		TextStylePtr = ButtonDefaults._TextStyle;
+	}
+
+	if (bDisplayText)
+	{
+		TextAttribute = BIND_UOBJECT_ATTRIBUTE(FText, GetButtonText);
+	}
+	else
+	{
+		TAttribute<FText> NotSet;
+		TextAttribute = NotSet;
+	}
+
+	TSharedRef<SButton> NewButton = SNew(SButton)
+		.ButtonStyle(StylePtr)
+		.TextStyle(TextStylePtr)
+		.HAlign(HorizontalAlignment)
+		.VAlign(VerticalAlignment)
+		.ContentPadding(BIND_UOBJECT_ATTRIBUTE(FMargin, GetContentPadding))
+		.Text(TextAttribute)
+		.OnClicked(BIND_UOBJECT_DELEGATE(FOnClicked, SlateOnClickedCallback))
+		.DesiredSizeScale(DesiredSizeScale)
+		.ContentScale(ContentScale)
+		.ButtonColorAndOpacity(BIND_UOBJECT_ATTRIBUTE(FSlateColor, GetButtonColor))
+		.ForegroundColor(BIND_UOBJECT_ATTRIBUTE(FSlateColor, GetForegroundColor));
+	
+	MyButton = NewButton;
+	return NewButton;
+}
+
+void UButtonComponent::OnKnownChildrenChanged()
+{
+	TSharedPtr<SButton> Button = MyButton.Pin();
+	if (Button.IsValid())
+	{
+		USlateWrapperComponent* ChildComponent = GetFirstWrappedChild();
+		if (ChildComponent)
+		{
+			Button->SetContent(ChildComponent->GetWidget());
+		}
+		else
+		{
+			// Removed a child; replace the text if desired
+			if (bDisplayText)
+			{
+				Button->SetContent(SNew(STextBlock)
+					.Text(TextAttribute)
+					.TextStyle( ButtonTextStyle ) );
+			}
+			else
+			{
+				Button->SetContent(SNullWidget::NullWidget);
+			}
+		}
+	}
+}
+
+FMargin UButtonComponent::GetContentPadding() const
+{
+	return ContentPadding;
+}
+
+FText UButtonComponent::GetButtonText() const
+{
+	return ButtonText;
+}
+
+FText UButtonComponent::GetText_Implementation() const
+{
+	return ButtonText;
+}
+
+FReply UButtonComponent::SlateOnClickedCallback()
+{
+	OnClicked.Broadcast();
+	return FReply::Handled();
+}
+
+FSlateColor UButtonComponent::GetButtonColor() const
+{
+	return ButtonColorAndOpacity;
+}
+
+FSlateColor UButtonComponent::GetForegroundColor() const
+{
+	return ForegroundColor;
+}
+
+/////////////////////////////////////////////////////
+
+#undef LOCTEXT_NAMESPACE
diff --git a/Engine/Source/Runtime/UMG/Private/CanvasPanelComponent.cpp b/Engine/Source/Runtime/UMG/Private/CanvasPanelComponent.cpp
new file mode 100644
--- /dev/null
+++ b/Engine/Source/Runtime/UMG/Private/CanvasPanelComponent.cpp
@@ -0,0 +1,149 @@
+// Copyright 1998-2014 Epic Games, Inc. All Rights Reserved.
+
+#include "UMGPrivatePCH.h"
+
+/////////////////////////////////////////////////////
+// SFixedSizeCanvas
+
+class SFixedSizeCanvas : public SCanvas
+{
+	SLATE_BEGIN_ARGS( SFixedSizeCanvas )
+		{
+			_Visibility = EVisibility::SelfHitTestInvisible;
+		}
+	SLATE_END_ARGS()
+
+	FVector2D CanvasSize;
+
+	void Construct(const FArguments& InArgs, FVector2D InSize)
+	{
+		SCanvas::FArguments ParentArgs;
+		SCanvas::Construct(ParentArgs);
+
+		CanvasSize = InSize;
+	}
+
+	// SWidget interface
+	virtual FVector2D ComputeDesiredSize() const OVERRIDE
+	{
+		return CanvasSize;
+	}
+	// End of SWidget interface
+};
+
+/////////////////////////////////////////////////////
+// UCanvasPanelComponent
+
+UCanvasPanelComponent::UCanvasPanelComponent(const FPostConstructInitializeProperties& PCIP)
+	: Super(PCIP)
+{
+	DesiredCanvasSize = FVector2D(128.0f, 128.0f);
+}
+
+TSharedRef<SWidget> UCanvasPanelComponent::RebuildWidget()
+{
+	TSharedRef<SFixedSizeCanvas> NewCanvas = SNew(SFixedSizeCanvas, DesiredCanvasSize);
+	MyCanvas = NewCanvas;
+
+	return NewCanvas;
+}
+
+void UCanvasPanelComponent::OnKnownChildrenChanged()
+{
+	TSharedPtr<SFixedSizeCanvas> Canvas = MyCanvas.Pin();
+	if (Canvas.IsValid())
+	{
+		Canvas->ClearChildren();
+
+		// Add slots
+		TMap<FName, int32> SlotNameToIndex;
+
+		for (int32 SlotIndex = 0; SlotIndex < Slots.Num(); ++SlotIndex)
+		{
+			const UCanvasPanelSlot* Slot = Slots[SlotIndex];
+			if ( Slot == NULL )
+			{
+				continue;
+			}
+
+			const FCanvasSlotData& SlotConfig = Slot->Data;
+
+			SlotNameToIndex.Add(SlotConfig.SlotName, SlotIndex);
+
+			auto NewSlot = Canvas->AddSlot()
+				.Position(SlotConfig.Position)
+				.Size(SlotConfig.Size)
+				.HAlign(SlotConfig.HorizontalAlignment)
+				.VAlign(SlotConfig.VerticalAlignment)
+			[
+				SNullWidget::NullWidget
+			];
+		}
+
+		// Place widgets in their slots, and add anything left over to the end
+		for (int32 ComponentIndex = 0; ComponentIndex < AttachChildren.Num(); ++ComponentIndex)
+		{
+			if (USlateWrapperComponent* Wrapper = Cast<USlateWrapperComponent>(AttachChildren[ComponentIndex]))
+			{
+				if (Wrapper->IsRegistered())
+				{
+					if (int32* pSlotIndex = SlotNameToIndex.Find(Wrapper->AttachSocketName))
+					{
+						auto ExistingSlot = (TPanelChildren<SCanvas::FSlot>*)(Canvas->GetChildren());
+						(*ExistingSlot)[*pSlotIndex].Widget = Wrapper->GetWidget();
+					}
+					else
+					{
+						Canvas->AddSlot()
+						[
+							Wrapper->GetWidget()
+						];
+					}
+				}
+			}
+		}
+	}
+}
+
+bool UCanvasPanelComponent::HasAnySockets() const
+{
+	return true;
+}
+
+FTransform UCanvasPanelComponent::GetSocketTransform(FName InSocketName, ERelativeTransformSpace TransformSpace) const
+{
+	return FTransform::Identity;
+}
+
+void UCanvasPanelComponent::QuerySupportedSockets(TArray<FComponentSocketDescription>& OutSockets) const
+{
+	for (int32 SlotIndex = 0; SlotIndex < Slots.Num(); ++SlotIndex)
+	{
+		FName SocketName = Slots[SlotIndex]->Data.SlotName;
+		new (OutSockets) FComponentSocketDescription(SocketName, EComponentSocketType::Socket);
+	}
+}
+
+#if WITH_EDITOR
+void UCanvasPanelComponent::PostEditChangeProperty(struct FPropertyChangedEvent& PropertyChangedEvent)
+{
+	// Ensure the slots have unique names
+	int32 SlotNumbering = 1;
+
+	TSet<FName> UniqueSlotNames;
+	for (int32 SlotIndex = 0; SlotIndex < Slots.Num(); ++SlotIndex)
+	{
+		FName& SlotName = Slots[SlotIndex]->Data.SlotName;
+
+		if ((SlotName == NAME_None) || UniqueSlotNames.Contains(SlotName))
+		{
+			do 
+			{
+				SlotName = FName(TEXT("Slot"), SlotNumbering++);
+			} while (UniqueSlotNames.Contains(SlotName));
+		}
+
+		UniqueSlotNames.Add(SlotName);
+	}
+}
+#endif
\ No newline at end of file
diff --git a/Engine/Source/Runtime/UMG/Private/CanvasPanelSlot.cpp b/Engine/Source/Runtime/UMG/Private/CanvasPanelSlot.cpp
new file mode 100644
--- /dev/null
+++ b/Engine/Source/Runtime/UMG/Private/CanvasPanelSlot.cpp
@@ -0,0 +1,11 @@
+// Copyright 1998-2014 Epic Games, Inc. All Rights Reserved.
+
+#include "UMGPrivatePCH.h"
+
+/////////////////////////////////////////////////////
+// UCanvasPanelSlot
+
+UCanvasPanelSlot::UCanvasPanelSlot(const FPostConstructInitializeProperties& PCIP)
+	: Super(PCIP)
+{
+}
\ No newline at end of file
diff --git a/Engine/Source/Runtime/UMG/Private/CheckBoxComponent.cpp b/Engine/Source/Runtime/UMG/Private/CheckBoxComponent.cpp
new file mode 100644
--- /dev/null
+++ b/Engine/Source/Runtime/UMG/Private/CheckBoxComponent.cpp
@@ -0,0 +1,65 @@
+// Copyright 1998-2014 Epic Games, Inc. All Rights Reserved.
+
+#include "UMGPrivatePCH.h"
+
+#define LOCTEXT_NAMESPACE "SlateComponentWrappers"
+
+/////////////////////////////////////////////////////
+// UCheckBoxComponent
+
+UCheckBoxComponent::UCheckBoxComponent(const FPostConstructInitializeProperties& PCIP)
+	: Super(PCIP)
+{
+	SCheckBox::FArguments CheckBoxDefaults;
+
+	CheckBoxStyleName = NAME_None;
+	bIsChecked = false;
+
+	HorizontalAlignment = CheckBoxDefaults._HAlign;
+	Padding = CheckBoxDefaults._Padding.Get();
+
+	ForegroundColor = CheckBoxDefaults._ForegroundColor.Get();
+	BorderBackgroundColor = CheckBoxDefaults._BorderBackgroundColor.Get();
+}
+
+TSharedRef<SWidget> UCheckBoxComponent::RebuildWidget()
+{
+	return SNew(SCheckBox)
+// 		SLATE_ATTRIBUTE( ESlateCheckBoxType::Type, Type )
+// 		SLATE_DEFAULT_SLOT( FArguments, Content )
+		//.Style( FEditorStyle::Get(), CheckBoxStyleName )
+// 		SLATE_ARGUMENT( const FCheckBoxStyle *, CheckBoxStyle )
+		.OnCheckStateChanged( BIND_UOBJECT_DELEGATE(FOnCheckStateChanged, SlateOnCheckStateChangedCallback) )
+		.IsChecked( BIND_UOBJECT_ATTRIBUTE(ESlateCheckBoxState::Type, GetCheckState) )
+		.HAlign( HorizontalAlignment )
+		.Padding( Padding )
+		.ForegroundColor( ForegroundColor )
+		.BorderBackgroundColor( BorderBackgroundColor );
+// 		SLATE_ATTRIBUTE( bool, ReadOnly )
+// 		SLATE_ARGUMENT( bool, IsFocusable )
+// 		SLATE_EVENT( FOnGetContent, OnGetMenuContent )
+}
+
+ESlateCheckBoxState::Type UCheckBoxComponent::GetCheckState() const
+{
+	return bIsChecked ? ESlateCheckBoxState::Checked : ESlateCheckBoxState::Unchecked;
+}
+
+void UCheckBoxComponent::SlateOnCheckStateChangedCallback(ESlateCheckBoxState::Type NewState)
+{
+	//@TODO: Choosing to treat Undetermined as Checked
+	const bool bWantsToBeChecked = NewState != ESlateCheckBoxState::Unchecked;
+
+	if (OnCheckStateChanged.IsBound())
+	{
+		OnCheckStateChanged.Broadcast(bWantsToBeChecked);
+	}
+	else
+	{
+		bIsChecked = bWantsToBeChecked;
+	}
+}
+
+/////////////////////////////////////////////////////
+
+#undef LOCTEXT_NAMESPACE
diff --git a/Engine/Source/Runtime/UMG/Private/CircularThrobberComponent.cpp b/Engine/Source/Runtime/UMG/Private/CircularThrobberComponent.cpp
new file mode 100644
--- /dev/null
+++ b/Engine/Source/Runtime/UMG/Private/CircularThrobberComponent.cpp
@@ -0,0 +1,32 @@
+// Copyright 1998-2014 Epic Games, Inc. All Rights Reserved.
+
+#include "UMGPrivatePCH.h"
+
+#define LOCTEXT_NAMESPACE "SlateComponentWrappers"
+
+/////////////////////////////////////////////////////
+// UCircularThrobberComponent
+
+UCircularThrobberComponent::UCircularThrobberComponent(const FPostConstructInitializeProperties& PCIP)
+	: Super(PCIP)
+{
+	SCircularThrobber::FArguments DefaultArgs;
+
+	NumberOfPieces = DefaultArgs._NumPieces;
+	Period = DefaultArgs._Period;
+	Radius = DefaultArgs._Radius;
+	PieceImage = *DefaultArgs._PieceImage;
+}
+
+TSharedRef<SWidget> UCircularThrobberComponent::RebuildWidget()
+{
+	return SNew(SCircularThrobber)
+		.PieceImage(&PieceImage)
+		.NumPieces(NumberOfPieces)
+		.Period(Period)
+		.Radius(Radius);
+}
+
+/////////////////////////////////////////////////////
+
+#undef LOCTEXT_NAMESPACE
diff --git a/Engine/Source/Runtime/UMG/Private/EditableTextBlockComponent.cpp b/Engine/Source/Runtime/UMG/Private/EditableTextBlockComponent.cpp
new file mode 100644
--- /dev/null
+++ b/Engine/Source/Runtime/UMG/Private/EditableTextBlockComponent.cpp
@@ -0,0 +1,22 @@
+// Copyright 1998-2014 Epic Games, Inc. All Rights Reserved.
+
+#include "UMGPrivatePCH.h"
+
+#define LOCTEXT_NAMESPACE "SlateComponentWrappers"
+
+/////////////////////////////////////////////////////
+// UEditableTextBlockComponent
+
+UEditableTextBlockComponent::UEditableTextBlockComponent(const FPostConstructInitializeProperties& PCIP)
+	: Super(PCIP)
+{
+}
+
+TSharedRef<SWidget> UEditableTextBlockComponent::RebuildWidget()
+{
+	return SNew(SEditableTextBox);
+}
+
+/////////////////////////////////////////////////////
+
+#undef LOCTEXT_NAMESPACE
diff --git a/Engine/Source/Runtime/UMG/Private/GUIPage.cpp b/Engine/Source/Runtime/UMG/Private/GUIPage.cpp
deleted file mode 100644
diff --git a/Engine/Source/Runtime/UMG/Private/GridPanelComponent.cpp b/Engine/Source/Runtime/UMG/Private/GridPanelComponent.cpp
new file mode 100644
--- /dev/null
+++ b/Engine/Source/Runtime/UMG/Private/GridPanelComponent.cpp
@@ -0,0 +1,139 @@
+// Copyright 1998-2014 Epic Games, Inc. All Rights Reserved.
+
+#include "UMGPrivatePCH.h"
+
+/////////////////////////////////////////////////////
+// UGridPanelComponent
+
+UGridPanelComponent::UGridPanelComponent(const FPostConstructInitializeProperties& PCIP)
+	: Super(PCIP)
+{
+}
+
+TSharedRef<SWidget> UGridPanelComponent::RebuildWidget()
+{
+	RebuildUniqueSlotTable();
+
+	TSharedRef<SGridPanel> NewGrid = SNew(SGridPanel);
+	MyGrid = NewGrid;
+
+	//@TODO: Expose FillColumn and FillRow
+
+	return NewGrid;
+}
+
+void UGridPanelComponent::OnKnownChildrenChanged()
+{
+	TSharedPtr<SGridPanel> Grid = MyGrid.Pin();
+	if (Grid.IsValid())
+	{
+		// Create a mapping of slot names to slot indices
+		TMap<FName, int32> SlotNameToIndex;
+		for (int32 SlotIndex = 0; SlotIndex < Slots.Num(); ++SlotIndex)
+		{
+			const FGridPanelSlot& SlotConfig = Slots[SlotIndex];
+			SlotNameToIndex.Add(SlotConfig.SlotName, SlotIndex);
+		}
+		
+		// Make sure all children are assigned to a slot
+		TArray<USlateWrapperComponent*> WrappedComponentsInSlotOrder;
+		WrappedComponentsInSlotOrder.AddZeroed(Slots.Num());
+
+		for (int32 ComponentIndex = 0; ComponentIndex < AttachChildren.Num(); ++ComponentIndex)
+		{
+			if (USlateWrapperComponent* Wrapper = Cast<USlateWrapperComponent>(AttachChildren[ComponentIndex]))
+			{
+				if (Wrapper->IsRegistered())
+				{
+					WrappedComponentsInSlotOrder.Add(Wrapper);
+
+					if (int32* pSlotIndex = SlotNameToIndex.Find(Wrapper->AttachSocketName))
+					{
+						WrappedComponentsInSlotOrder[*pSlotIndex] = Wrapper;
+					}
+					else
+					{
+						FGridPanelSlot& NewSlot = *new (Slots) FGridPanelSlot();
+						EnsureSlotIsUnique(NewSlot);
+						WrappedComponentsInSlotOrder.Add(Wrapper);
+
+						Wrapper->AttachSocketName = NewSlot.SlotName;
+					}
+				}						
+			}
+		}
+
+		// Add slots to the grid panel
+		Grid->ClearChildren();
+		for (int32 SlotIndex = 0; SlotIndex < Slots.Num(); ++SlotIndex)
+		{
+			const FGridPanelSlot& SlotConfig = Slots[SlotIndex];
+
+			TSharedRef<SWidget> ContentWidget = SNullWidget::NullWidget;
+			if (WrappedComponentsInSlotOrder[SlotIndex] != NULL)
+			{
+				ContentWidget = WrappedComponentsInSlotOrder[SlotIndex]->GetWidget();
+			}
+
+			Grid->AddSlot(SlotConfig.Column, SlotConfig.Row)
+				.RowSpan(SlotConfig.RowSpan)
+				.ColumnSpan(SlotConfig.ColumnSpan)
+				.Padding(SlotConfig.Padding)
+				.HAlign(SlotConfig.HorizontalAlignment)
+				.VAlign(SlotConfig.VerticalAlignment)
+			[
+				ContentWidget
+			];
+		}
+	}
+}
+
+bool UGridPanelComponent::HasAnySockets() const
+{
+	return true;
+}
+
+FTransform UGridPanelComponent::GetSocketTransform(FName InSocketName, ERelativeTransformSpace TransformSpace) const
+{
+	return FTransform::Identity;
+}
+
+void UGridPanelComponent::QuerySupportedSockets(TArray<FComponentSocketDescription>& OutSockets) const
+{
+	for (int32 SlotIndex = 0; SlotIndex < Slots.Num(); ++SlotIndex)
+	{
+		FName SocketName = Slots[SlotIndex].SlotName;
+		new (OutSockets) FComponentSocketDescription(SocketName, EComponentSocketType::Socket);
+	}
+}
+
+#if WITH_EDITOR
+void UGridPanelComponent::PostEditChangeProperty(struct FPropertyChangedEvent& PropertyChangedEvent)
+{
+	RebuildUniqueSlotTable();
+}
+#endif
+
+void UGridPanelComponent::RebuildUniqueSlotTable()
+{
+	// Ensure the slots have unique names and coordinates
+	UniqueSlotCoordinates.Empty();
+
+	for (int32 SlotIndex = 0; SlotIndex < Slots.Num(); ++SlotIndex)
+	{
+		EnsureSlotIsUnique(Slots[SlotIndex]);
+	}
+}
+
+void UGridPanelComponent::EnsureSlotIsUnique(FGridPanelSlot& SlotConfig)
+{
+	FName& SlotName = SlotConfig.SlotName;
+
+	while (UniqueSlotCoordinates.Contains(SlotConfig.AsPoint()))
+	{
+		SlotConfig.Row++;
+	}
+
+	UniqueSlotCoordinates.Add(SlotConfig.AsPoint());
+	SlotName = FName(*FString::Printf(TEXT("X%d_Y%d"), SlotConfig.Column, SlotConfig.Row));
+}
diff --git a/Engine/Source/Runtime/UMG/Private/ImageComponent.cpp b/Engine/Source/Runtime/UMG/Private/ImageComponent.cpp
new file mode 100644
--- /dev/null
+++ b/Engine/Source/Runtime/UMG/Private/ImageComponent.cpp
@@ -0,0 +1,33 @@
+// Copyright 1998-2014 Epic Games, Inc. All Rights Reserved.
+
+#include "UMGPrivatePCH.h"
+
+#define LOCTEXT_NAMESPACE "SlateComponentWrappers"
+
+/////////////////////////////////////////////////////
+// UImageComponent
+
+UImageComponent::UImageComponent(const FPostConstructInitializeProperties& PCIP)
+	: Super(PCIP)
+	, Image()
+	, ColorAndOpacity(FLinearColor::White)
+{
+	//@TODO: Would like to default to the checkerboard, but TextureName always wins over TextureObject :(
+//	Image = *FEditorStyle::GetDefaultBrush();
+}
+
+TSharedRef<SWidget> UImageComponent::RebuildWidget()
+{
+	return SNew(SImage)
+		.Image( &Image )
+		.ColorAndOpacity( BIND_UOBJECT_ATTRIBUTE(FSlateColor, GetColorAndOpacity) );
+}
+
+FSlateColor UImageComponent::GetColorAndOpacity() const
+{
+	return ColorAndOpacity;
+}
+
+/////////////////////////////////////////////////////
+
+#undef LOCTEXT_NAMESPACE
diff --git a/Engine/Source/Runtime/UMG/Private/PanelSlot.cpp b/Engine/Source/Runtime/UMG/Private/PanelSlot.cpp
new file mode 100644
--- /dev/null
+++ b/Engine/Source/Runtime/UMG/Private/PanelSlot.cpp
@@ -0,0 +1,11 @@
+// Copyright 1998-2014 Epic Games, Inc. All Rights Reserved.
+
+#include "UMGPrivatePCH.h"
+
+/////////////////////////////////////////////////////
+// UPanelSlot
+
+UPanelSlot::UPanelSlot(const FPostConstructInitializeProperties& PCIP)
+	: Super(PCIP)
+{
+}
\ No newline at end of file
diff --git a/Engine/Source/Runtime/UMG/Private/ScrollBoxComponent.cpp b/Engine/Source/Runtime/UMG/Private/ScrollBoxComponent.cpp
new file mode 100644
--- /dev/null
+++ b/Engine/Source/Runtime/UMG/Private/ScrollBoxComponent.cpp
@@ -0,0 +1,22 @@
+// Copyright 1998-2014 Epic Games, Inc. All Rights Reserved.
+
+#include "UMGPrivatePCH.h"
+
+#define LOCTEXT_NAMESPACE "SlateComponentWrappers"
+
+/////////////////////////////////////////////////////
+// UScrollBoxComponent
+
+UScrollBoxComponent::UScrollBoxComponent(const FPostConstructInitializeProperties& PCIP)
+	: Super(PCIP)
+{
+}
+
+TSharedRef<SWidget> UScrollBoxComponent::RebuildWidget()
+{
+	return SNew(SScrollBox);
+}
+
+/////////////////////////////////////////////////////
+
+#undef LOCTEXT_NAMESPACE
diff --git a/Engine/Source/Runtime/UMG/Private/SlateLeafWidgetComponent.cpp b/Engine/Source/Runtime/UMG/Private/SlateLeafWidgetComponent.cpp
new file mode 100644
--- /dev/null
+++ b/Engine/Source/Runtime/UMG/Private/SlateLeafWidgetComponent.cpp
@@ -0,0 +1,16 @@
+// Copyright 1998-2014 Epic Games, Inc. All Rights Reserved.
+
+#include "UMGPrivatePCH.h"
+
+/////////////////////////////////////////////////////
+// USlateLeafWidgetComponent
+
+USlateLeafWidgetComponent::USlateLeafWidgetComponent(const FPostConstructInitializeProperties& PCIP)
+	: Super(PCIP)
+{
+}
+
+bool USlateLeafWidgetComponent::CanAttachAsChild(USceneComponent* ChildComponent, FName SocketName) const
+{
+	return false;
+}
diff --git a/Engine/Source/Runtime/UMG/Private/SlateNonLeafWidgetComponent.cpp b/Engine/Source/Runtime/UMG/Private/SlateNonLeafWidgetComponent.cpp
new file mode 100644
--- /dev/null
+++ b/Engine/Source/Runtime/UMG/Private/SlateNonLeafWidgetComponent.cpp
@@ -0,0 +1,134 @@
+// Copyright 1998-2014 Epic Games, Inc. All Rights Reserved.
+
+#include "UMGPrivatePCH.h"
+
+#if WITH_EDITOR
+#include "MessageLog.h"
+#include "UObjectToken.h"
+#endif
+
+#define LOCTEXT_NAMESPACE "SlateComponentWrappers"
+
+/////////////////////////////////////////////////////
+// USlateNonLeafWidgetComponent
+
+USlateNonLeafWidgetComponent::USlateNonLeafWidgetComponent(const FPostConstructInitializeProperties& PCIP)
+	: Super(PCIP)
+{
+	PrimaryComponentTick.bCanEverTick = true;
+
+	bAutoActivate = true;
+	bTickInEditor = true;
+
+	bCanHaveMultipleChildren = true;
+}
+
+void USlateNonLeafWidgetComponent::TickComponent(float DeltaTime, enum ELevelTick TickType, FActorComponentTickFunction* ThisTickFunction)
+{
+	bool bNeedRebuild = false;
+	TSet<USlateWrapperComponent*> RemainingChildren;
+	for (auto SetIt = KnownChildren.CreateIterator(); SetIt; ++SetIt)
+	{
+		if (USlateWrapperComponent* Wrapper = (*SetIt).Get())
+		{
+			if (Wrapper->IsRegistered())
+			{
+				RemainingChildren.Add(Wrapper);
+			}
+		}
+	}
+
+	for (int32 ComponentIndex = 0; (ComponentIndex < AttachChildren.Num()) && !bNeedRebuild; ++ComponentIndex)
+	{
+		if (USlateWrapperComponent* Wrapper = Cast<USlateWrapperComponent>(AttachChildren[ComponentIndex]))
+		{
+			if (Wrapper->IsRegistered())
+			{
+				if (KnownChildren.Contains(Wrapper))
+				{
+					RemainingChildren.Remove(Wrapper);
+				}
+				else
+				{
+					bNeedRebuild = true;
+				}
+			}
+		}
+	}
+
+	bNeedRebuild = bNeedRebuild || (RemainingChildren.Num() > 0);
+
+	if (bNeedRebuild)
+	{
+		// Recreate the list of known children
+		KnownChildren.Empty();
+		for (int32 ComponentIndex = 0; ComponentIndex < AttachChildren.Num(); ++ComponentIndex)
+		{
+			if (USlateWrapperComponent* Wrapper = Cast<USlateWrapperComponent>(AttachChildren[ComponentIndex]))
+			{
+				if (Wrapper->IsRegistered())
+				{
+					KnownChildren.Add(Wrapper);
+				}
+			}
+		}
+
+		// Let derived classes do something with the list
+		OnKnownChildrenChanged();
+	}
+}
+
+USlateWrapperComponent* USlateNonLeafWidgetComponent::GetFirstWrappedChild() const
+{
+	if (auto It = KnownChildren.CreateConstIterator())
+	{
+		return It->Get();
+	}
+	else
+	{
+		return NULL;
+	}
+}
+
+TSharedRef<SWidget> USlateNonLeafWidgetComponent::GetFirstChildWidget() const
+{
+	if (USlateWrapperComponent* WrappedChild = GetFirstWrappedChild())
+	{
+		return WrappedChild->GetWidget();
+	}
+	else
+	{
+		return SNullWidget::NullWidget;
+	}
+}
+
+#if WITH_EDITOR
+
+void USlateNonLeafWidgetComponent::CheckForErrors()
+{
+	Super::CheckForErrors();
+
+	if ((KnownChildren.Num() > 1) && !bCanHaveMultipleChildren)
+	{
+		AActor* Owner = GetOwner();
+		FFormatNamedArguments Arguments;
+		Arguments.Add(TEXT("ComponentName"), FText::FromString(GetName()));
+		Arguments.Add(TEXT("OwnerName"), FText::FromString(Owner->GetName()));
+		FText Message = FText::Format( LOCTEXT( "MapCheck_Message_TooManyChildren", "{ComponentName}::{OwnerName} cannot have more than one child attached" ), Arguments );
+		FMessageLog("MapCheck").Warning()
+			->AddToken(FUObjectToken::Create(Owner))
+			->AddToken(FTextToken::Create(Message) );
+	}
+}
+
+#endif
+
+bool USlateNonLeafWidgetComponent::CanAttachAsChild(USceneComponent* ChildComponent, FName SocketName) const
+{
+	const bool bIsSlateWidgetWrapper = ChildComponent->IsA(USlateWrapperComponent::StaticClass());
+	return bIsSlateWidgetWrapper && (bCanHaveMultipleChildren ? true : (AttachChildren.Num() == 0));
+}
+
+/////////////////////////////////////////////////////
+
+#undef LOCTEXT_NAMESPACE
diff --git a/Engine/Source/Runtime/UMG/Private/SlateWrapperComponent.cpp b/Engine/Source/Runtime/UMG/Private/SlateWrapperComponent.cpp
new file mode 100644
--- /dev/null
+++ b/Engine/Source/Runtime/UMG/Private/SlateWrapperComponent.cpp
@@ -0,0 +1,75 @@
+// Copyright 1998-2014 Epic Games, Inc. All Rights Reserved.
+
+#include "UMGPrivatePCH.h"
+
+/////////////////////////////////////////////////////
+// USlateWrapperComponent
+
+USlateWrapperComponent::USlateWrapperComponent(const FPostConstructInitializeProperties& PCIP)
+	: Super(PCIP)
+{
+}
+
+void USlateWrapperComponent::OnRegister()
+{
+	MyWidget = RebuildWidget();
+	Super::OnRegister();
+}
+
+void USlateWrapperComponent::OnUnregister()
+{
+	MyWidget = NULL;
+	Super::OnUnregister();
+}
+
+TSharedRef<SWidget> USlateWrapperComponent::GetWidget()
+{
+	if (MyWidget.IsValid())
+	{
+		return MyWidget.ToSharedRef();
+	}
+	else
+	{
+		return SNullWidget::NullWidget;
+	}
+}
+
+TSharedRef<SWidget> USlateWrapperComponent::RebuildWidget()
+{
+	ensureMsg(false, TEXT("You must implement RebuildWidget() in your child class"));
+	return SNullWidget::NullWidget;
+}
+
+EVisibility USlateWrapperComponent::ConvertSerializedVisibilityToRuntime(TEnumAsByte<ESlateVisibility::Type> Input)
+{
+	switch ( Input )
+	{
+	case ESlateVisibility::Visible:
+		return EVisibility::Visible;
+	case ESlateVisibility::Collapsed:
+		return EVisibility::Collapsed;
+	case ESlateVisibility::Hidden:
+		return EVisibility::Hidden;
+	case ESlateVisibility::HitTestInvisible:
+		return EVisibility::HitTestInvisible;
+	case ESlateVisibility::SelfHitTestInvisible:
+		return EVisibility::SelfHitTestInvisible;
+	default:
+		check(false);
+		return EVisibility::Visible;
+	}
+}
+
+FSizeParam USlateWrapperComponent::ConvertSerializedSizeParamToRuntime(const FSlateChildSize& Input)
+{
+	switch ( Input.SizeRule )
+	{
+	default:
+	case ESlateSizeRule::Automatic:
+		return FAuto();
+	case ESlateSizeRule::AspectRatio:
+		return FAspectRatio();
+	case ESlateSizeRule::Stretch:
+		return FStretch(Input.StretchValue);
+	}
+}
\ No newline at end of file
diff --git a/Engine/Source/Runtime/UMG/Private/SpacerComponent.cpp b/Engine/Source/Runtime/UMG/Private/SpacerComponent.cpp
new file mode 100644
--- /dev/null
+++ b/Engine/Source/Runtime/UMG/Private/SpacerComponent.cpp
@@ -0,0 +1,29 @@
+// Copyright 1998-2014 Epic Games, Inc. All Rights Reserved.
+
+#include "UMGPrivatePCH.h"
+
+#define LOCTEXT_NAMESPACE "SlateComponentWrappers"
+
+/////////////////////////////////////////////////////
+// USpacerComponent
+
+USpacerComponent::USpacerComponent(const FPostConstructInitializeProperties& PCIP)
+	: Super(PCIP)
+	, Size(1.0f, 1.0f)
+{
+}
+
+TSharedRef<SWidget> USpacerComponent::RebuildWidget()
+{
+	return SNew(SSpacer)
+		.Size( BIND_UOBJECT_ATTRIBUTE(FVector2D, GetSpacerSize) );
+}
+
+FVector2D USpacerComponent::GetSpacerSize() const
+{
+	return Size;
+}
+
+/////////////////////////////////////////////////////
+
+#undef LOCTEXT_NAMESPACE
diff --git a/Engine/Source/Runtime/UMG/Private/SpinningImageComponent.cpp b/Engine/Source/Runtime/UMG/Private/SpinningImageComponent.cpp
new file mode 100644
--- /dev/null
+++ b/Engine/Source/Runtime/UMG/Private/SpinningImageComponent.cpp
@@ -0,0 +1,27 @@
+// Copyright 1998-2014 Epic Games, Inc. All Rights Reserved.
+
+#include "UMGPrivatePCH.h"
+
+#define LOCTEXT_NAMESPACE "SlateComponentWrappers"
+
+/////////////////////////////////////////////////////
+// UImageComponent
+
+USpinningImageComponent::USpinningImageComponent(const FPostConstructInitializeProperties& PCIP)
+	: Super(PCIP)
+{
+	SSpinningImage::FArguments DefaultArgs;
+	Period = DefaultArgs._Period;
+}
+
+TSharedRef<SWidget> USpinningImageComponent::RebuildWidget()
+{
+	return SNew(SSpinningImage)
+		.Image( &Image )
+		.ColorAndOpacity( BIND_UOBJECT_ATTRIBUTE(FSlateColor, GetColorAndOpacity) )
+		.Period( Period );
+}
+
+/////////////////////////////////////////////////////
+
+#undef LOCTEXT_NAMESPACE
diff --git a/Engine/Source/Runtime/UMG/Private/TextBlockComponent.cpp b/Engine/Source/Runtime/UMG/Private/TextBlockComponent.cpp
new file mode 100644
--- /dev/null
+++ b/Engine/Source/Runtime/UMG/Private/TextBlockComponent.cpp
@@ -0,0 +1,45 @@
+// Copyright 1998-2014 Epic Games, Inc. All Rights Reserved.
+
+#include "UMGPrivatePCH.h"
+
+#define LOCTEXT_NAMESPACE "SlateComponentWrappers"
+
+/////////////////////////////////////////////////////
+// UTextBlockComponent
+
+UTextBlockComponent::UTextBlockComponent(const FPostConstructInitializeProperties& PCIP)
+	: Super(PCIP)
+{
+	DisplayText = LOCTEXT("TextBlockDefaultValue", "Hello World");
+	ColorAndOpacity = FLinearColor::White;
+	ShadowColorAndOpacity = FLinearColor::Transparent;
+}
+
+TSharedRef<SWidget> UTextBlockComponent::RebuildWidget()
+{
+	return SNew(STextBlock)
+		.Text( BIND_UOBJECT_ATTRIBUTE(FString, GetDisplayText) )
+		.Font( FSlateFontInfo(FPaths::EngineContentDir() / TEXT("Slate/Fonts/Roboto-Bold.ttf"), 24) )
+		.ColorAndOpacity( BIND_UOBJECT_ATTRIBUTE(FSlateColor, GetColorAndOpacity) )
+		.ShadowOffset(FVector2D(1.0f, 1.0f))
+		.ShadowColorAndOpacity( BIND_UOBJECT_ATTRIBUTE(FLinearColor, GetShadowColorAndOpacity) );
+}
+
+FString UTextBlockComponent::GetDisplayText() const
+{
+	return DisplayText.ToString();
+}
+
+FSlateColor UTextBlockComponent::GetColorAndOpacity() const
+{
+	return ColorAndOpacity;
+}
+
+FLinearColor UTextBlockComponent::GetShadowColorAndOpacity() const
+{
+	return ShadowColorAndOpacity;
+}
+
+/////////////////////////////////////////////////////
+
+#undef LOCTEXT_NAMESPACE
diff --git a/Engine/Source/Runtime/UMG/Private/ThrobberComponent.cpp b/Engine/Source/Runtime/UMG/Private/ThrobberComponent.cpp
new file mode 100644
--- /dev/null
+++ b/Engine/Source/Runtime/UMG/Private/ThrobberComponent.cpp
@@ -0,0 +1,37 @@
+// Copyright 1998-2014 Epic Games, Inc. All Rights Reserved.
+
+#include "UMGPrivatePCH.h"
+
+#define LOCTEXT_NAMESPACE "SlateComponentWrappers"
+
+/////////////////////////////////////////////////////
+// UThrobberComponent
+
+UThrobberComponent::UThrobberComponent(const FPostConstructInitializeProperties& PCIP)
+	: Super(PCIP)
+{
+	SThrobber::FArguments DefaultArgs;
+
+	PieceImage = *DefaultArgs._PieceImage;
+	NumberOfPieces = DefaultArgs._NumPieces;
+
+	bAnimateVertically = (DefaultArgs._Animate & SThrobber::Vertical) != 0;
+	bAnimateHorizontally = (DefaultArgs._Animate & SThrobber::Horizontal) != 0;
+	bAnimateOpacity = (DefaultArgs._Animate & SThrobber::Opacity) != 0;
+}
+
+TSharedRef<SWidget> UThrobberComponent::RebuildWidget()
+{
+	const int32 AnimationParams = (bAnimateVertically ? SThrobber::Vertical : 0) |
+		(bAnimateHorizontally ? SThrobber::Horizontal : 0) |
+		(bAnimateOpacity ? SThrobber::Opacity : 0);
+
+	return SNew(SThrobber)
+		.PieceImage(&PieceImage)
+		.NumPieces(NumberOfPieces)
+		.Animate(static_cast<SThrobber::EAnimation>(AnimationParams));
+}
+
+/////////////////////////////////////////////////////
+
+#undef LOCTEXT_NAMESPACE
diff --git a/Engine/Source/Runtime/UMG/Private/UMGModule.cpp b/Engine/Source/Runtime/UMG/Private/UMGModule.cpp
--- a/Engine/Source/Runtime/UMG/Private/UMGModule.cpp
+++ b/Engine/Source/Runtime/UMG/Private/UMGModule.cpp
@@ -18,6 +18,7 @@ public:
 	/** Called right after the module DLL has been loaded and the module object has been created */
 	virtual void StartupModule() OVERRIDE
 	{
+		printf("ALIVE");
 	}
 
 	/** Called before the module is unloaded, right before the module object is destroyed. */
diff --git a/Engine/Source/Runtime/UMG/Private/UserWidget.cpp b/Engine/Source/Runtime/UMG/Private/UserWidget.cpp
new file mode 100644
--- /dev/null
+++ b/Engine/Source/Runtime/UMG/Private/UserWidget.cpp
@@ -0,0 +1,72 @@
+// Copyright 1998-2014 Epic Games, Inc. All Rights Reserved.
+
+#include "UMGPrivatePCH.h"
+
+/////////////////////////////////////////////////////
+// AUserWidget
+
+AUserWidget::AUserWidget(const FPostConstructInitializeProperties& PCIP)
+	: Super(PCIP)
+{
+}
+
+void AUserWidget::PostActorCreated()
+{
+	EnsureWidgetExists();
+	Super::PostActorCreated();
+}
+
+void AUserWidget::Destroyed()
+{
+	MyWrapperWidget = NULL;
+	Super::Destroyed();
+}
+
+void AUserWidget::RerunConstructionScripts()
+{
+	Super::RerunConstructionScripts();
+
+	RebuildWrapperWidget();
+}
+
+void AUserWidget::EnsureWidgetExists()
+{
+	if (!MyWrapperWidget.IsValid())
+	{
+		MyWrapperWidget = SNew(SOverlay);
+	}
+}
+
+void AUserWidget::RebuildWrapperWidget()
+{
+	EnsureWidgetExists();
+	MyWrapperWidget->ClearChildren();
+	
+	TArray<USlateWrapperComponent*> SlateWrapperComponents;
+	GetComponents(SlateWrapperComponents);
+	// Place all of our top-level children Slate wrapped components into the overlay
+	for (int32 ComponentIndex = 0; ComponentIndex < SlateWrapperComponents.Num(); ++ComponentIndex)
+	{
+		if (SlateWrapperComponents[ComponentIndex]->IsRegistered() && (SlateWrapperComponents[ComponentIndex]->AttachParent == NULL))
+		{
+			MyWrapperWidget->AddSlot()
+			[
+				SlateWrapperComponents[ComponentIndex]->GetWidget()
+			];
+		}
+	}
+	
+	if (GEngine->GameViewport != NULL)
+	{
+		GEngine->GameViewport->AddViewportWidgetContent(MyWrapperWidget.ToSharedRef());
+	}
+//@TODO: Figure out how hot-reloading/etc... should work with slate wrapped components
+	//SAssignNew(ScoreboardWidgetContainer,SWeakWidget)
+	//.PossiblyNullContent(ScoreboardWidgetOverlay));
+}
+
+TSharedRef<SWidget> AUserWidget::GetWidget()
+{
+	EnsureWidgetExists();
+	return MyWrapperWidget.ToSharedRef();
+}
diff --git a/Engine/Source/Runtime/UMG/Private/WidgetBlueprint.cpp b/Engine/Source/Runtime/UMG/Private/WidgetBlueprint.cpp
new file mode 100644
--- /dev/null
+++ b/Engine/Source/Runtime/UMG/Private/WidgetBlueprint.cpp
@@ -0,0 +1,14 @@
+// Copyright 1998-2014 Epic Games, Inc. All Rights Reserved.
+
+#include "UMGPrivatePCH.h"
+
+#define LOCTEXT_NAMESPACE "SlateComponentWrappers"
+
+/////////////////////////////////////////////////////
+// UWidgetBlueprint
+
+UWidgetBlueprint::UWidgetBlueprint(const FPostConstructInitializeProperties& PCIP)
+	: Super(PCIP)
+{
+	printf("TEST");
+}
\ No newline at end of file
diff --git a/Engine/Source/Runtime/UMG/Private/WidgetBlueprintGeneratedClass.cpp b/Engine/Source/Runtime/UMG/Private/WidgetBlueprintGeneratedClass.cpp
new file mode 100644
--- /dev/null
+++ b/Engine/Source/Runtime/UMG/Private/WidgetBlueprintGeneratedClass.cpp
@@ -0,0 +1,13 @@
+// Copyright 1998-2014 Epic Games, Inc. All Rights Reserved.
+
+#include "UMGPrivatePCH.h"
+
+#define LOCTEXT_NAMESPACE "SlateComponentWrappers"
+
+/////////////////////////////////////////////////////
+// UWidgetBlueprintGeneratedClass
+
+UWidgetBlueprintGeneratedClass::UWidgetBlueprintGeneratedClass(const FPostConstructInitializeProperties& PCIP)
+	: Super(PCIP)
+{
+}
\ No newline at end of file
diff --git a/Engine/Source/Runtime/UMG/UMG.Build.cs b/Engine/Source/Runtime/UMG/UMG.Build.cs
--- a/Engine/Source/Runtime/UMG/UMG.Build.cs
+++ b/Engine/Source/Runtime/UMG/UMG.Build.cs
@@ -13,6 +13,7 @@ public class UMG : ModuleRules
 				"Engine",
 				"ShaderCore",
 				"RenderCore",
+                "InputCore",
 				"RHI",
 			}
 			);
