Nonunity built fix.

diff --git a/Engine/Source/Editor/UMGEditor/Private/WidgetBlueprintFactory.cpp b/Engine/Source/Editor/UMGEditor/Private/WidgetBlueprintFactory.cpp
--- a/Engine/Source/Editor/UMGEditor/Private/WidgetBlueprintFactory.cpp
+++ b/Engine/Source/Editor/UMGEditor/Private/WidgetBlueprintFactory.cpp
@@ -2,6 +2,9 @@
 
 #include "UMGEditorPrivatePCH.h"
 
+#include "Kismet2/BlueprintEditorUtils.h"
+#include "Kismet2/KismetEditorUtilities.h"
+
 #define LOCTEXT_NAMESPACE "UWidgetBlueprintFactory"
 
 /*------------------------------------------------------------------------------
