* Introducing a new kismet compiler for widget blueprints (for UMG).
* Extending the signature of a few kismet utility methods to make it more extensible.
* Now generating a list of template widgets (currrently still components).  Able to safely build a hierarchy and then clone that into new objects without GLEO problems.
* Adding vertical and horrizontal boxes to UMG.

diff --git a/Engine/Source/Developer/AssetTools/Private/AssetTypeActions/AssetTypeActions_Blueprint.cpp b/Engine/Source/Developer/AssetTools/Private/AssetTypeActions/AssetTypeActions_Blueprint.cpp
--- a/Engine/Source/Developer/AssetTools/Private/AssetTypeActions/AssetTypeActions_Blueprint.cpp
+++ b/Engine/Source/Developer/AssetTools/Private/AssetTypeActions/AssetTypeActions_Blueprint.cpp
@@ -151,7 +151,7 @@ void FAssetTypeActions_Blueprint::ExecuteNewDerivedBlueprint(TWeakObjectPtr<UBlu
 		if (ensure(Package))
 		{
 			// Create and init a new Blueprint
-			UBlueprint* NewBP = FKismetEditorUtilities::CreateBlueprint(TargetParentClass, Package, FName(*Name), BPTYPE_Normal, TargetParentBP->GetClass());
+			UBlueprint* NewBP = FKismetEditorUtilities::CreateBlueprint(TargetParentClass, Package, FName(*Name), BPTYPE_Normal, TargetParentBP->GetClass(), UBlueprintGeneratedClass::StaticClass());
 			if (NewBP)
 			{
 				FAssetEditorManager::Get().OpenEditorForAsset(NewBP);
diff --git a/Engine/Source/Developer/AssetTools/Private/AssetTypeActions/AssetTypeActions_BlueprintGeneratedClass.cpp b/Engine/Source/Developer/AssetTools/Private/AssetTypeActions/AssetTypeActions_BlueprintGeneratedClass.cpp
--- a/Engine/Source/Developer/AssetTools/Private/AssetTypeActions/AssetTypeActions_BlueprintGeneratedClass.cpp
+++ b/Engine/Source/Developer/AssetTools/Private/AssetTypeActions/AssetTypeActions_BlueprintGeneratedClass.cpp
@@ -148,7 +148,7 @@ void FAssetTypeActions_BlueprintGeneratedClass::ExecuteNewDerivedBlueprint(TWeak
 		if ( ensure(Package) )
 		{
 			// Create and init a new Blueprint
-			UBlueprint* NewBP = FKismetEditorUtilities::CreateBlueprint(TargetClass, Package, FName(*Name), BPTYPE_Normal, UBlueprint::StaticClass());
+			UBlueprint* NewBP = FKismetEditorUtilities::CreateBlueprint(TargetClass, Package, FName(*Name), BPTYPE_Normal, UBlueprint::StaticClass(), UBlueprintGeneratedClass::StaticClass());
 			if(NewBP)
 			{
 				FAssetEditorManager::Get().OpenEditorForAsset(NewBP);
diff --git a/Engine/Source/Editor/Blutility/Private/AssetTypeActions_EditorUtilityBlueprint.cpp b/Engine/Source/Editor/Blutility/Private/AssetTypeActions_EditorUtilityBlueprint.cpp
--- a/Engine/Source/Editor/Blutility/Private/AssetTypeActions_EditorUtilityBlueprint.cpp
+++ b/Engine/Source/Editor/Blutility/Private/AssetTypeActions_EditorUtilityBlueprint.cpp
@@ -179,7 +179,7 @@ void FAssetTypeActions_EditorUtilityBlueprint::ExecuteNewDerivedBlueprint(TWeakO
 		if (ensure(Package))
 		{
 			// Create and init a new Blueprint
-			if (UBlueprint* NewBP = FKismetEditorUtilities::CreateBlueprint(TargetClass, Package, FName(*Name), BPTYPE_Normal, UEditorUtilityBlueprint::StaticClass()))
+			if (UBlueprint* NewBP = FKismetEditorUtilities::CreateBlueprint(TargetClass, Package, FName(*Name), BPTYPE_Normal, UEditorUtilityBlueprint::StaticClass(), UBlueprintGeneratedClass::StaticClass()))
 			{
 				FAssetEditorManager::Get().OpenEditorForAsset(NewBP);
 
diff --git a/Engine/Source/Editor/Blutility/Private/EditorUtilityBlueprintFactory.cpp b/Engine/Source/Editor/Blutility/Private/EditorUtilityBlueprintFactory.cpp
--- a/Engine/Source/Editor/Blutility/Private/EditorUtilityBlueprintFactory.cpp
+++ b/Engine/Source/Editor/Blutility/Private/EditorUtilityBlueprintFactory.cpp
@@ -96,7 +96,7 @@ UObject* UEditorUtilityBlueprintFactory::FactoryCreateNew(UClass* Class, UObject
 	}
 	else
 	{
-		return FKismetEditorUtilities::CreateBlueprint(ParentClass, InParent, Name, BPTYPE_Normal, UEditorUtilityBlueprint::StaticClass());
+		return FKismetEditorUtilities::CreateBlueprint(ParentClass, InParent, Name, BPTYPE_Normal, UEditorUtilityBlueprint::StaticClass(), UBlueprintGeneratedClass::StaticClass());
 	}
 }
 
diff --git a/Engine/Source/Editor/ClassViewer/Private/SClassViewer.cpp b/Engine/Source/Editor/ClassViewer/Private/SClassViewer.cpp
--- a/Engine/Source/Editor/ClassViewer/Private/SClassViewer.cpp
+++ b/Engine/Source/Editor/ClassViewer/Private/SClassViewer.cpp
@@ -875,7 +875,7 @@ namespace ClassViewer
 			if(PromptUserIfExistingObject(BPName.ToString(), PackageName,FString(), Package))
 			{
 				// Create and init a new Blueprint
-				UBlueprint* NewBP = FKismetEditorUtilities::CreateBlueprint(InCreationClass, Package, BPName, BPTYPE_Normal, UBlueprint::StaticClass(), FName("ClassViewer"));
+				UBlueprint* NewBP = FKismetEditorUtilities::CreateBlueprint(InCreationClass, Package, BPName, BPTYPE_Normal, UBlueprint::StaticClass(), UBlueprintGeneratedClass::StaticClass(), FName("ClassViewer"));
 				if(NewBP)
 				{
 					FAssetEditorManager::Get().OpenEditorForAsset(NewBP);
diff --git a/Engine/Source/Editor/KismetCompiler/Private/KismetCompilerModule.cpp b/Engine/Source/Editor/KismetCompiler/Private/KismetCompilerModule.cpp
--- a/Engine/Source/Editor/KismetCompiler/Private/KismetCompilerModule.cpp
+++ b/Engine/Source/Editor/KismetCompiler/Private/KismetCompilerModule.cpp
@@ -56,17 +56,37 @@ void FKismet2CompilerModule::CompileBlueprintInner(class UBlueprint* Blueprint,
 	{
 		Results.Error(*LOCTEXT("KismetCompileError_MalformedParentClasss", "Blueprint @@ has missing or NULL parent class.").ToString(), Blueprint);
 	}
-	else if (UAnimBlueprint* AnimBlueprint = Cast<UAnimBlueprint>(Blueprint))
-	{
-		FAnimBlueprintCompiler Compiler(AnimBlueprint, Results, CompileOptions, ObjLoaded);
-		Compiler.Compile();
-		check(Compiler.NewClass);
-	}
 	else
 	{
-		FKismetCompilerContext Compiler(Blueprint, Results, CompileOptions, ObjLoaded);
-		Compiler.Compile();
-		check(Compiler.NewClass);
+		// Loop through all external compiler delegates attempting to compile the blueprint.
+		FReply Handled = FReply::Unhandled();
+		for ( FBlueprintCompileDelegate& Compiler : Compilers )
+		{
+			Handled = Compiler.Execute(Blueprint, CompileOptions, Results, ObjLoaded);
+
+			// Don't allow any other compiler to handle it if they reported it was handled.
+			if ( Handled.IsEventHandled() )
+			{
+				break;
+			}
+		}
+
+		// if no one handles it, then use the default blueprint compiler.
+		if ( !Handled.IsEventHandled() )
+		{
+			if ( UAnimBlueprint* AnimBlueprint = Cast<UAnimBlueprint>(Blueprint) )
+			{
+				FAnimBlueprintCompiler Compiler(AnimBlueprint, Results, CompileOptions, ObjLoaded);
+				Compiler.Compile();
+				check(Compiler.NewClass);
+			}
+			else
+			{
+				FKismetCompilerContext Compiler(Blueprint, Results, CompileOptions, ObjLoaded);
+				Compiler.Compile();
+				check(Compiler.NewClass);
+			}
+		}
 	}
 
 	double FinishTime = FPlatformTime::Seconds();
@@ -131,7 +151,7 @@ void FKismet2CompilerModule::CompileBlueprint(class UBlueprint* Blueprint, const
 	// If this was a full compile, take appropriate actions depending on the success of failure of the compile
 	if( CompileOptions.IsGeneratedClassCompileType() )
 	{
- 		// Perform the full compile
+		// Perform the full compile
 		CompileBlueprintInner(Blueprint, /*bPrintResultSuccess=*/ true, CompileOptions, Results, ObjLoaded);
 
 		if (Results.NumErrors == 0)
diff --git a/Engine/Source/Editor/KismetCompiler/Public/EdGraphCompilerUtilities.h b/Engine/Source/Editor/KismetCompiler/Public/EdGraphCompilerUtilities.h
--- a/Engine/Source/Editor/KismetCompiler/Public/EdGraphCompilerUtilities.h
+++ b/Engine/Source/Editor/KismetCompiler/Public/EdGraphCompilerUtilities.h
@@ -29,7 +29,7 @@ DECLARE_CYCLE_STAT_EXTERN(TEXT("Update Instances"), EKismetCompilerStats_UpdateB
 
 //////////////////////////////////////////////////////////////////////////
 
-class FGraphCompilerContext
+class KISMETCOMPILER_API FGraphCompilerContext
 {
 public:
 	// Compiler message log (errors, warnings, notes)
diff --git a/Engine/Source/Editor/KismetCompiler/Public/KismetCompiler.h b/Engine/Source/Editor/KismetCompiler/Public/KismetCompiler.h
--- a/Engine/Source/Editor/KismetCompiler/Public/KismetCompiler.h
+++ b/Engine/Source/Editor/KismetCompiler/Public/KismetCompiler.h
@@ -13,7 +13,7 @@ KISMETCOMPILER_API DECLARE_LOG_CATEGORY_EXTERN(LogK2Compiler, Log, All);
 //////////////////////////////////////////////////////////////////////////
 // FKismetCompilerContext
 
-class FKismetCompilerContext : public FGraphCompilerContext
+class KISMETCOMPILER_API FKismetCompilerContext : public FGraphCompilerContext
 {
 protected:
 	typedef FGraphCompilerContext Super;
@@ -122,7 +122,7 @@ public:
 	 * @param  IntermediatePin	The pin you want the SourcePin's links moved to.
 	 * @return The result from calling the schema's MovePinLinks().
 	 */
-	KISMETCOMPILER_API FPinConnectionResponse MovePinLinksToIntermediate(UEdGraphPin& SourcePin, UEdGraphPin& IntermediatePin);
+	FPinConnectionResponse MovePinLinksToIntermediate(UEdGraphPin& SourcePin, UEdGraphPin& IntermediatePin);
 
 	/**
 	 * Copies pin links over from the source-pin to the specified intermediate, 
@@ -133,9 +133,9 @@ public:
 	 * @param  IntermediatePin	The pin you want the SourcePin's links copied to.
 	 * @return The result from calling the schema's CopyPinLinks().
 	 */
-	KISMETCOMPILER_API FPinConnectionResponse CopyPinLinksToIntermediate(UEdGraphPin& SourcePin, UEdGraphPin& IntermediatePin);
+	FPinConnectionResponse CopyPinLinksToIntermediate(UEdGraphPin& SourcePin, UEdGraphPin& IntermediatePin);
 
-	KISMETCOMPILER_API UK2Node_TemporaryVariable* SpawnInternalVariable(UEdGraphNode* SourceNode, FString Category, FString SubCategory = TEXT(""), UObject* SubcategoryObject = NULL, bool bIsArray = false);
+	UK2Node_TemporaryVariable* SpawnInternalVariable(UEdGraphNode* SourceNode, FString Category, FString SubCategory = TEXT(""), UObject* SubcategoryObject = NULL, bool bIsArray = false);
 
 protected:
 	virtual UEdGraphSchema_K2* CreateSchema();
@@ -160,7 +160,7 @@ protected:
 	 * users should be using MovePinLinksToIntermediate/CopyPinLinksToIntermediate 
 	 * instead of wrapping their own with this).
 	 */
-	KISMETCOMPILER_API void CheckConnectionResponse(const FPinConnectionResponse &Response, const UEdGraphNode *Node);
+	void CheckConnectionResponse(const FPinConnectionResponse &Response, const UEdGraphNode *Node);
 
 protected:
 	// FGraphCompilerContext interface
@@ -182,7 +182,7 @@ protected:
 
 	// Advances the macro position tracking
 	void AdvanceMacroPlacement(int32 Width, int32 Height);
-	KISMETCOMPILER_API void AutoAssignNodePosition(UEdGraphNode* Node);
+	void AutoAssignNodePosition(UEdGraphNode* Node);
 	void CreateCommentBlockAroundNodes(const TArray<UEdGraphNode*>& Nodes, UObject* SourceObject, UEdGraph* TargetGraph, FString CommentText, FLinearColor CommentColor, int32& Out_OffsetX, int32& Out_OffsetY);
 
 	/** Creates a class variable */
diff --git a/Engine/Source/Editor/KismetCompiler/Public/KismetCompilerModule.h b/Engine/Source/Editor/KismetCompiler/Public/KismetCompilerModule.h
--- a/Engine/Source/Editor/KismetCompiler/Public/KismetCompilerModule.h
+++ b/Engine/Source/Editor/KismetCompiler/Public/KismetCompilerModule.h
@@ -5,12 +5,15 @@
 #include "Engine.h"
 
 class FCompilerResultsLog;
+class UBlueprint;
 
 #define KISMET_COMPILER_MODULENAME "KismetCompiler"
 
 //////////////////////////////////////////////////////////////////////////
 // IKismetCompilerInterface
 
+DECLARE_DELEGATE_RetVal_FourParams(FReply, FBlueprintCompileDelegate, UBlueprint* /*Blueprint*/, const FKismetCompilerOptions& /*CompileOptions*/, FCompilerResultsLog& /*Results*/, TArray<UObject*>* /*ObjLoaded*/);
+
 class IKismetCompilerInterface : public IModuleInterface
 {
 public:
@@ -36,6 +39,11 @@ public:
 	 * @param	Blueprint	The blueprint to clear the classes for
 	 */
 	virtual void RemoveBlueprintGeneratedClasses(class UBlueprint* Blueprint)=0;
+
+	/**
+	 * Gets a list of all compilers for blueprints.  You can register new compilers through this list.
+	 */
+	virtual TArray<FBlueprintCompileDelegate>& GetCompilers() = 0;
 };
 
 //////////////////////////////////////////////////////////////////////////
@@ -51,8 +59,11 @@ public:
 	virtual void CompileBlueprint(class UBlueprint* Blueprint, const FKismetCompilerOptions& CompileOptions, FCompilerResultsLog& Results, class FBlueprintCompileReinstancer* ParentReinstancer = NULL, TArray<UObject*>* ObjLoaded = NULL) OVERRIDE;
 	virtual void RecoverCorruptedBlueprint(class UBlueprint* Blueprint) OVERRIDE;
 	virtual void RemoveBlueprintGeneratedClasses(class UBlueprint* Blueprint) OVERRIDE;
+	virtual TArray<FBlueprintCompileDelegate>& GetCompilers() OVERRIDE { return Compilers; }
 	// End implementation
 private:
 	void CompileBlueprintInner(class UBlueprint* Blueprint, bool bPrintResultSuccess, const FKismetCompilerOptions& CompileOptions, FCompilerResultsLog& Results, TArray<UObject*>* ObjLoaded);
+
+	TArray<FBlueprintCompileDelegate> Compilers;
 };
 
diff --git a/Engine/Source/Editor/Sequencer/Private/Sequencer.cpp b/Engine/Source/Editor/Sequencer/Private/Sequencer.cpp
--- a/Engine/Source/Editor/Sequencer/Private/Sequencer.cpp
+++ b/Engine/Source/Editor/Sequencer/Private/Sequencer.cpp
@@ -1350,7 +1350,7 @@ FGuid FSequencer::AddSpawnableForAssetOrClass( UObject* Object, UObject* Counter
 		else if( bIsActorClass )
 		{
 			// We don't have a factory, but we can still try to create a blueprint for this actor class
-			NewBlueprint = FKismetEditorUtilities::CreateBlueprint( Object->GetClass(), OwnerMovieScene, BlueprintName, EBlueprintType::BPTYPE_Normal, UBlueprint::StaticClass() );
+			NewBlueprint = FKismetEditorUtilities::CreateBlueprint( Object->GetClass(), OwnerMovieScene, BlueprintName, EBlueprintType::BPTYPE_Normal, UBlueprint::StaticClass(), UBlueprintGeneratedClass::StaticClass() );
 		}
 	}
 
diff --git a/Engine/Source/Editor/UMGEditor/Private/SUMGEditorTree.cpp b/Engine/Source/Editor/UMGEditor/Private/SUMGEditorTree.cpp
--- a/Engine/Source/Editor/UMGEditor/Private/SUMGEditorTree.cpp
+++ b/Engine/Source/Editor/UMGEditor/Private/SUMGEditorTree.cpp
@@ -24,7 +24,19 @@ void SUMGEditorTree::Construct(const FArguments& InArgs, TSharedPtr<FBlueprintEd
 	ChildSlot
 	[
 		SNew(SVerticalBox)
-		+SVerticalBox::Slot()
+		+ SVerticalBox::Slot()
+		.AutoHeight()
+		[
+			SNew(SButton)
+			.OnClicked(this, &SUMGEditorTree::CreateTestUI)
+			[
+				SNew(STextBlock)
+				.Text(NSLOCTEXT("SUMGEditorTree", "CreateTestUI", "Create Test UI"))
+			]
+		]
+
+		+ SVerticalBox::Slot()
+		.FillHeight(1.0f)
 		[
 			SAssignNew(PageTreeView, STreeView<TSharedPtr<FString>>)
 			.ItemHeight(20.0f)
@@ -41,6 +53,44 @@ SUMGEditorTree::~SUMGEditorTree()
 	//FCoreDelegates::OnObjectPropertyChanged.Remove( FCoreDelegates::FOnObjectPropertyChanged::FDelegate::CreateRaw(this, &SUMGEditorTree::OnObjectPropertyChanged) );
 }
 
+FReply SUMGEditorTree::CreateTestUI()
+{
+	//TSharedRef<ISCSEditor> SCSEditor = BlueprintEditor.Pin()->GetSCSEditorModel();
+	//UCanvasPanelComponent* Canvas = Cast<UCanvasPanelComponent>(SCSEditor->AddNewComponent(UCanvasPanelComponent::StaticClass(), NULL));
+	//UButtonComponent* Button = Cast<UButtonComponent>(SCSEditor->AddNewComponent(UButtonComponent::StaticClass(), NULL));
+
+	UWidgetBlueprint* BP = CastChecked<UWidgetBlueprint>(BlueprintEditor.Pin()->GetBlueprintObj());
+
+	UCanvasPanelComponent* Canvas = ConstructObject<UCanvasPanelComponent>(UCanvasPanelComponent::StaticClass(), BP);
+	UVerticalBoxComponent* Vertical = ConstructObject<UVerticalBoxComponent>(UVerticalBoxComponent::StaticClass(), BP);
+	UButtonComponent* Button1 = ConstructObject<UButtonComponent>(UButtonComponent::StaticClass(), BP);
+	Button1->ButtonText = FText::FromString("Button 1");
+	UButtonComponent* Button2 = ConstructObject<UButtonComponent>(UButtonComponent::StaticClass(), BP);
+	Button2->ButtonText = FText::FromString("Button 2");
+	UButtonComponent* Button3 = ConstructObject<UButtonComponent>(UButtonComponent::StaticClass(), BP);
+	Button3->ButtonText = FText::FromString("Button 3");
+
+	BP->WidgetTemplates.Add(Canvas);
+	BP->WidgetTemplates.Add(Vertical);
+	BP->WidgetTemplates.Add(Button1);
+	BP->WidgetTemplates.Add(Button2);
+	BP->WidgetTemplates.Add(Button3);
+
+	UCanvasPanelSlot* Slot = Canvas->AddSlot(Vertical);
+	Slot->Size.X = 100;
+	Slot->Size.Y = 100;
+	Slot->Position.X = 20;
+	Slot->Position.Y = 50;
+
+	Vertical->AddSlot(Button1);
+	Vertical->AddSlot(Button2);
+	Vertical->AddSlot(Button3);
+
+	FBlueprintEditorUtils::MarkBlueprintAsStructurallyModified(BP);
+
+	return FReply::Handled();
+}
+
 //void SUMGEditorTree::AddReferencedObjects( FReferenceCollector& Collector )
 //{
 //	Collector.AddReferencedObject(Page);
diff --git a/Engine/Source/Editor/UMGEditor/Private/SUMGEditorTree.h b/Engine/Source/Editor/UMGEditor/Private/SUMGEditorTree.h
--- a/Engine/Source/Editor/UMGEditor/Private/SUMGEditorTree.h
+++ b/Engine/Source/Editor/UMGEditor/Private/SUMGEditorTree.h
@@ -15,6 +15,8 @@ public:
 
 	void Construct(const FArguments& InArgs, TSharedPtr<FBlueprintEditor> InBlueprintEditor, USimpleConstructionScript* InSCS);
 	~SUMGEditorTree();
+
+	FReply CreateTestUI();
 	
 private:
 
diff --git a/Engine/Source/Editor/UMGEditor/Private/UMGEditorModule.cpp b/Engine/Source/Editor/UMGEditor/Private/UMGEditorModule.cpp
--- a/Engine/Source/Editor/UMGEditor/Private/UMGEditorModule.cpp
+++ b/Engine/Source/Editor/UMGEditor/Private/UMGEditorModule.cpp
@@ -8,6 +8,8 @@
 #include "AssetToolsModule.h"
 #include "UMGBlueprintEditorExtensionHook.h"
 #include "AssetTypeActions_WidgetBlueprint.h"
+#include "KismetCompilerModule.h"
+#include "WidgetBlueprintCompiler.h"
 
 #include "UMGEditor.generated.inl"
 
@@ -32,6 +34,12 @@ public:
 		// Register asset types
 		IAssetTools& AssetTools = FModuleManager::LoadModuleChecked<FAssetToolsModule>("AssetTools").Get();
 		RegisterAssetTypeAction(AssetTools, MakeShareable(new FAssetTypeActions_WidgetBlueprint()));
+
+		// Register blueprint compiler
+		IKismetCompilerInterface& KismetCompilerModule = FModuleManager::LoadModuleChecked<IKismetCompilerInterface>("KismetCompiler");
+		FBlueprintCompileDelegate sd;
+		sd.BindRaw(this, &FUMGEditorModule::CompileWidgetBlueprint);
+		KismetCompilerModule.GetCompilers().Add(sd);
 		
 		FUMGBlueprintEditorExtensionHook::InstallHooks();
 	}
@@ -59,6 +67,20 @@ public:
 		CreatedAssetTypeActions.Empty();
 	}
 
+	FReply CompileWidgetBlueprint(UBlueprint* Blueprint, const FKismetCompilerOptions& CompileOptions, FCompilerResultsLog& Results, TArray<UObject*>* ObjLoaded)
+	{
+		if ( UWidgetBlueprint* WidgetBlueprint = Cast<UWidgetBlueprint>(Blueprint) )
+		{
+			FWidgetBlueprintCompiler Compiler(WidgetBlueprint, Results, CompileOptions, ObjLoaded);
+			Compiler.Compile();
+			check(Compiler.NewClass);
+
+			return FReply::Handled();
+		}
+
+		return FReply::Unhandled();
+	}
+
 	virtual TSharedRef<IUMGEditor> CreateUMGEditor(const EToolkitMode::Type Mode, const TSharedPtr< IToolkitHost >& InitToolkitHost, UParticleSystem* ParticleSystem) OVERRIDE
 	{
 		TSharedRef<FUMGEditor> NewUMGEditor(new FUMGEditor());
diff --git a/Engine/Source/Editor/UMGEditor/Private/WidgetBlueprintCompiler.cpp b/Engine/Source/Editor/UMGEditor/Private/WidgetBlueprintCompiler.cpp
new file mode 100644
--- /dev/null
+++ b/Engine/Source/Editor/UMGEditor/Private/WidgetBlueprintCompiler.cpp
@@ -0,0 +1,153 @@
+// Copyright 1998-2014 Epic Games, Inc. All Rights Reserved.
+
+#include "UMGEditorPrivatePCH.h"
+
+#include "WidgetBlueprintCompiler.h"
+#include "Kismet2NameValidators.h"
+#include "KismetReinstanceUtilities.h"
+
+FWidgetBlueprintCompiler::FWidgetBlueprintCompiler(UWidgetBlueprint* SourceSketch, FCompilerResultsLog& InMessageLog, const FKismetCompilerOptions& InCompilerOptions, TArray<UObject*>* InObjLoaded)
+	: Super(SourceSketch, InMessageLog, InCompilerOptions, InObjLoaded)
+{
+}
+
+FWidgetBlueprintCompiler::~FWidgetBlueprintCompiler()
+{
+}
+
+void FWidgetBlueprintCompiler::ValidateWidgetNames()
+{
+	UWidgetBlueprint* Blueprint = WidgetBlueprint();
+
+	TSharedPtr<FKismetNameValidator> ParentBPNameValidator;
+	if ( Blueprint->ParentClass != NULL )
+	{
+		UBlueprint* ParentBP = Cast<UBlueprint>(Blueprint->ParentClass->ClassGeneratedBy);
+		if ( ParentBP != NULL )
+		{
+			ParentBPNameValidator = MakeShareable(new FKismetNameValidator(ParentBP));
+		}
+	}
+
+	for ( USlateWrapperComponent* Widget : Blueprint->WidgetTemplates )
+	{
+		if ( ParentBPNameValidator.IsValid() && ParentBPNameValidator->IsValid(Widget->GetName()) != EValidatorResult::Ok )
+		{
+			// TODO Support renaming items, similar to timelines.
+
+			//// Use the viewer displayed Timeline name (without the _Template suffix) because it will be added later for appropriate checks.
+			//FString TimelineName = UTimelineTemplate::TimelineTemplateNameToVariableName(TimelineTemplate->GetFName());
+
+			//FName NewName = FBlueprintEditorUtils::FindUniqueKismetName(Blueprint, TimelineName);
+			//MessageLog.Warning(*FString::Printf(*LOCTEXT("TimelineConflictWarning", "Found a timeline with a conflicting name (%s) - changed to %s.").ToString(), *TimelineTemplate->GetName(), *NewName.ToString()));
+			//FBlueprintEditorUtils::RenameTimeline(Blueprint, FName(*TimelineName), NewName);
+		}
+	}
+}
+
+template<typename TOBJ>
+struct FCullTemplateObjectsHelper
+{
+	const TArray<TOBJ*>& Templates;
+
+	FCullTemplateObjectsHelper(const TArray<TOBJ*>& InComponentTemplates)
+		: Templates(InComponentTemplates)
+	{}
+
+	bool operator()(const UObject* const RemovalCandidate) const
+	{
+		return ( NULL != Templates.FindByKey(RemovalCandidate) );
+	}
+};
+
+void FWidgetBlueprintCompiler::CleanAndSanitizeClass(UBlueprintGeneratedClass* ClassToClean, UObject*& OldCDO)
+{
+	Super::CleanAndSanitizeClass(ClassToClean, OldCDO);
+
+	// Purge all subobjects (properties, functions, params) of the class, as they will be regenerated
+	TArray<UObject*> ClassSubObjects;
+	GetObjectsWithOuter(ClassToClean, ClassSubObjects, true);
+
+	UWidgetBlueprint* Blueprint = WidgetBlueprint();
+
+	if ( 0 != Blueprint->Timelines.Num() )
+	{
+		ClassSubObjects.RemoveAllSwap(FCullTemplateObjectsHelper<USlateWrapperComponent>(Blueprint->WidgetTemplates));
+	}
+}
+
+void FWidgetBlueprintCompiler::CreateClassVariablesFromBlueprint()
+{
+	Super::CreateClassVariablesFromBlueprint();
+
+	UWidgetBlueprint* Blueprint = WidgetBlueprint();
+
+	ValidateWidgetNames();
+
+	int32 i = 0;
+	for ( USlateWrapperComponent* Widget : Blueprint->WidgetTemplates )
+	{
+		FString VariableName = Widget->GetClass()->GetName() + TEXT("_") + FString::FromInt(i);
+
+		FEdGraphPinType WidgetPinType(Schema->PC_Object, TEXT(""), Widget->GetClass(), false, false);
+		UProperty* WidgetProperty = CreateVariable(FName(*VariableName), WidgetPinType);
+		if ( WidgetProperty != NULL )
+		{
+			WidgetProperty->SetMetaData(TEXT("Category"), *Blueprint->GetName());
+			WidgetProperty->SetPropertyFlags(CPF_BlueprintVisible);
+
+			WidgetToMemberVariableMap.Add(Widget, WidgetProperty);
+		}
+
+		i++;
+	}
+
+	printf("Compile");
+}
+
+void FWidgetBlueprintCompiler::FinishCompilingClass(UClass* Class)
+{
+	Super::FinishCompilingClass(Class);
+
+	UClass* ParentClass = Class->GetSuperClass();
+	UWidgetBlueprint* Blueprint = WidgetBlueprint();
+
+	{
+		UWidgetBlueprintGeneratedClass* BPGClass = CastChecked<UWidgetBlueprintGeneratedClass>(Class);
+
+		BPGClass->WidgetTemplates.Empty();
+		BPGClass->WidgetTemplates = Blueprint->WidgetTemplates;
+	}
+}
+
+void FWidgetBlueprintCompiler::Compile()
+{
+	Super::Compile();
+
+	WidgetToMemberVariableMap.Empty();
+}
+
+void FWidgetBlueprintCompiler::EnsureProperGeneratedClass(UClass*& TargetUClass)
+{
+	if ( TargetUClass && !( (UObject*)TargetUClass )->IsA(UWidgetBlueprintGeneratedClass::StaticClass()) )
+	{
+		FKismetCompilerUtilities::ConsignToOblivion(TargetUClass, Blueprint->bIsRegeneratingOnLoad);
+		TargetUClass = NULL;
+	}
+}
+
+void FWidgetBlueprintCompiler::SpawnNewClass(const FString& NewClassName)
+{
+	NewWidgetBlueprintClass = FindObject<UWidgetBlueprintGeneratedClass>(Blueprint->GetOutermost(), *NewClassName);
+
+	if ( NewWidgetBlueprintClass == NULL )
+	{
+		NewWidgetBlueprintClass = ConstructObject<UWidgetBlueprintGeneratedClass>(UWidgetBlueprintGeneratedClass::StaticClass(), Blueprint->GetOutermost(), FName(*NewClassName), RF_Public | RF_Transactional);
+	}
+	else
+	{
+		// Already existed, but wasn't linked in the Blueprint yet due to load ordering issues
+		FBlueprintCompileReinstancer GeneratedClassReinstancer(NewWidgetBlueprintClass);
+	}
+	NewClass = NewWidgetBlueprintClass;
+}
diff --git a/Engine/Source/Editor/UMGEditor/Private/WidgetBlueprintCompiler.h b/Engine/Source/Editor/UMGEditor/Private/WidgetBlueprintCompiler.h
new file mode 100644
--- /dev/null
+++ b/Engine/Source/Editor/UMGEditor/Private/WidgetBlueprintCompiler.h
@@ -0,0 +1,44 @@
+// Copyright 1998-2014 Epic Games, Inc. All Rights Reserved.
+
+#pragma once
+
+#include "KismetCompiler.h"
+
+//////////////////////////////////////////////////////////////////////////
+// FWidgetBlueprintCompiler
+
+class FWidgetBlueprintCompiler : public FKismetCompilerContext
+{
+protected:
+	typedef FKismetCompilerContext Super;
+
+public:
+	FWidgetBlueprintCompiler(UWidgetBlueprint* SourceSketch, FCompilerResultsLog& InMessageLog, const FKismetCompilerOptions& InCompilerOptions, TArray<UObject*>* InObjLoaded);
+	virtual ~FWidgetBlueprintCompiler();
+
+	// FKismetCompilerContext
+	virtual void Compile() OVERRIDE;
+	// End FKismetCompilerContext
+
+protected:
+	UWidgetBlueprint* WidgetBlueprint() const { return Cast<UWidgetBlueprint>(Blueprint); }
+
+	void ValidateWidgetNames();
+
+	// FKismetCompilerContext
+	virtual void SpawnNewClass(const FString& NewClassName) OVERRIDE;
+	virtual void EnsureProperGeneratedClass(UClass*& TargetClass) OVERRIDE;
+	virtual void CreateClassVariablesFromBlueprint() OVERRIDE;
+	virtual void FinishCompilingClass(UClass* Class) OVERRIDE;
+	virtual void CleanAndSanitizeClass(UBlueprintGeneratedClass* ClassToClean, UObject*& OldCDO) OVERRIDE;
+	// End FKismetCompilerContext
+
+protected:
+	UWidgetBlueprintGeneratedClass* NewWidgetBlueprintClass;
+	//UWidgetBlueprint* AnimBlueprint;
+
+	// Map of properties created for widgets; to aid in debug data generation
+	TMap<class USlateWrapperComponent*, class UProperty*> WidgetToMemberVariableMap;
+
+};
+
diff --git a/Engine/Source/Editor/UMGEditor/Private/WidgetBlueprintFactory.cpp b/Engine/Source/Editor/UMGEditor/Private/WidgetBlueprintFactory.cpp
--- a/Engine/Source/Editor/UMGEditor/Private/WidgetBlueprintFactory.cpp
+++ b/Engine/Source/Editor/UMGEditor/Private/WidgetBlueprintFactory.cpp
@@ -46,7 +46,7 @@ UObject* UWidgetBlueprintFactory::FactoryCreateNew(UClass* Class, UObject* InPar
 	}
 	else
 	{
-		UWidgetBlueprint* NewBP = CastChecked<UWidgetBlueprint>(FKismetEditorUtilities::CreateBlueprint(ParentClass, InParent, Name, BlueprintType, UWidgetBlueprint::StaticClass(), CallingContext));
+		UWidgetBlueprint* NewBP = CastChecked<UWidgetBlueprint>(FKismetEditorUtilities::CreateBlueprint(ParentClass, InParent, Name, BlueprintType, UWidgetBlueprint::StaticClass(), UWidgetBlueprintGeneratedClass::StaticClass(), CallingContext));
 
 		return NewBP;
 	}
diff --git a/Engine/Source/Editor/UMGEditor/UMGEditor.Build.cs b/Engine/Source/Editor/UMGEditor/UMGEditor.Build.cs
--- a/Engine/Source/Editor/UMGEditor/UMGEditor.Build.cs
+++ b/Engine/Source/Editor/UMGEditor/UMGEditor.Build.cs
@@ -37,6 +37,8 @@ public class UMGEditor : ModuleRules
 				"AssetTools",
 				"UnrealEd", // for FAssetEditorManager
 				"KismetWidgets",
+                "KismetCompiler",
+                "BlueprintGraph",
 				"GraphEditor",
 				"Kismet",  // for FWorkflowCentricApplication
 				"PropertyEditor",
diff --git a/Engine/Source/Editor/UnrealEd/Private/Factories/ActorFactory.cpp b/Engine/Source/Editor/UnrealEd/Private/Factories/ActorFactory.cpp
--- a/Engine/Source/Editor/UnrealEd/Private/Factories/ActorFactory.cpp
+++ b/Engine/Source/Editor/UnrealEd/Private/Factories/ActorFactory.cpp
@@ -97,7 +97,7 @@ UBlueprint* UActorFactory::CreateBlueprint( UObject* Asset, UObject* Outer, cons
 	// @todo sequencer major: Needs to be overridden on any class that needs any custom setup for the new blueprint
 	//	(e.g. static mesh reference assignment.)  Basically, anywhere that PostSpawnActor() or CreateActor() is overridden,
 	//	we should consider overriding CreateBlueprint(), too.
-	UBlueprint* NewBlueprint = FKismetEditorUtilities::CreateBlueprint( NewActorClass, Outer, Name, EBlueprintType::BPTYPE_Normal, UBlueprint::StaticClass(), CallingContext );
+	UBlueprint* NewBlueprint = FKismetEditorUtilities::CreateBlueprint(NewActorClass, Outer, Name, EBlueprintType::BPTYPE_Normal, UBlueprint::StaticClass(), UBlueprintGeneratedClass::StaticClass(), CallingContext);
 	AActor* CDO = CastChecked<AActor>( NewBlueprint->GeneratedClass->ClassDefaultObject );
 	PostCreateBlueprint( Asset, CDO );
 	return NewBlueprint;
diff --git a/Engine/Source/Editor/UnrealEd/Private/Factories/AnimBlueprintFactory.cpp b/Engine/Source/Editor/UnrealEd/Private/Factories/AnimBlueprintFactory.cpp
--- a/Engine/Source/Editor/UnrealEd/Private/Factories/AnimBlueprintFactory.cpp
+++ b/Engine/Source/Editor/UnrealEd/Private/Factories/AnimBlueprintFactory.cpp
@@ -382,7 +382,7 @@ UObject* UAnimBlueprintFactory::FactoryCreateNew(UClass* Class, UObject* InParen
 	}
 	else
 	{
-		UAnimBlueprint* NewBP = CastChecked<UAnimBlueprint>(FKismetEditorUtilities::CreateBlueprint(ParentClass, InParent, Name, BlueprintType, UAnimBlueprint::StaticClass(), CallingContext));
+		UAnimBlueprint* NewBP = CastChecked<UAnimBlueprint>(FKismetEditorUtilities::CreateBlueprint(ParentClass, InParent, Name, BlueprintType, UAnimBlueprint::StaticClass(), UBlueprintGeneratedClass::StaticClass(), CallingContext));
 		NewBP->TargetSkeleton = TargetSkeleton;
 
 		// Because the BP itself didn't have the skeleton set when the initial compile occured, it's not set on the generated classes either
diff --git a/Engine/Source/Editor/UnrealEd/Private/Factories/EditorFactories.cpp b/Engine/Source/Editor/UnrealEd/Private/Factories/EditorFactories.cpp
--- a/Engine/Source/Editor/UnrealEd/Private/Factories/EditorFactories.cpp
+++ b/Engine/Source/Editor/UnrealEd/Private/Factories/EditorFactories.cpp
@@ -6112,7 +6112,7 @@ UObject* UBlueprintGeneratedClassFactory::FactoryCreateNew(UClass* Class, UObjec
 	}
 	else
 	{
-		return FKismetEditorUtilities::CreateBlueprint(ParentClass, InParent, Name, BPTYPE_Normal, UBlueprint::StaticClass());
+		return FKismetEditorUtilities::CreateBlueprint(ParentClass, InParent, Name, BPTYPE_Normal, UBlueprint::StaticClass(), UBlueprintGeneratedClass::StaticClass());
 	}
 }
 
@@ -6201,7 +6201,7 @@ UObject* UBlueprintFactory::FactoryCreateNew(UClass* Class, UObject* InParent, F
 	}
 	else
 	{
-		return FKismetEditorUtilities::CreateBlueprint(ParentClass, InParent, Name, BPTYPE_Normal, UBlueprint::StaticClass(), CallingContext);
+		return FKismetEditorUtilities::CreateBlueprint(ParentClass, InParent, Name, BPTYPE_Normal, UBlueprint::StaticClass(), UBlueprintGeneratedClass::StaticClass(), CallingContext);
 	}
 }
 
@@ -6254,7 +6254,7 @@ UObject* UBlueprintMacroFactory::FactoryCreateNew(UClass* Class, UObject* InPare
 	}
 	else
 	{
-		return FKismetEditorUtilities::CreateBlueprint(ParentClass, InParent, Name, BPTYPE_MacroLibrary, UBlueprint::StaticClass(), CallingContext);
+		return FKismetEditorUtilities::CreateBlueprint(ParentClass, InParent, Name, BPTYPE_MacroLibrary, UBlueprint::StaticClass(), UBlueprintGeneratedClass::StaticClass(), CallingContext);
 	}
 }
 
@@ -6309,7 +6309,7 @@ UObject* UBlueprintFunctionLibraryFactory::FactoryCreateNew(UClass* Class, UObje
 	}
 	else
 	{
-		return FKismetEditorUtilities::CreateBlueprint(ParentClass, InParent, Name, BPTYPE_FunctionLibrary, UBlueprint::StaticClass(), CallingContext);
+		return FKismetEditorUtilities::CreateBlueprint(ParentClass, InParent, Name, BPTYPE_FunctionLibrary, UBlueprint::StaticClass(), UBlueprintGeneratedClass::StaticClass(), CallingContext);
 	}
 }
 
@@ -6363,7 +6363,7 @@ UObject* UBlueprintInterfaceFactory::FactoryCreateNew(UClass* Class, UObject* In
 	}
 	else
 	{
-		return FKismetEditorUtilities::CreateBlueprint(ParentClass, InParent, Name, BPTYPE_Interface, UBlueprint::StaticClass(), CallingContext);
+		return FKismetEditorUtilities::CreateBlueprint(ParentClass, InParent, Name, BPTYPE_Interface, UBlueprint::StaticClass(), UBlueprintGeneratedClass::StaticClass(), CallingContext);
 	}
 }
 
diff --git a/Engine/Source/Editor/UnrealEd/Private/Kismet2/Kismet2.cpp b/Engine/Source/Editor/UnrealEd/Private/Kismet2/Kismet2.cpp
--- a/Engine/Source/Editor/UnrealEd/Private/Kismet2/Kismet2.cpp
+++ b/Engine/Source/Editor/UnrealEd/Private/Kismet2/Kismet2.cpp
@@ -173,7 +173,7 @@ void FKismetEditorUtilities::CreateDefaultEventGraphs(UBlueprint* Blueprint)
 }
 
 /** Create a new Blueprint and initialize it to a valid state. */
-UBlueprint* FKismetEditorUtilities::CreateBlueprint(UClass* ParentClass, UObject* Outer, const FName NewBPName, EBlueprintType BlueprintType, TSubclassOf<UBlueprint> BlueprintClassType, FName CallingContext)
+UBlueprint* FKismetEditorUtilities::CreateBlueprint(UClass* ParentClass, UObject* Outer, const FName NewBPName, EBlueprintType BlueprintType, TSubclassOf<UBlueprint> BlueprintClassType, TSubclassOf<UBlueprintGeneratedClass> BlueprintGeneratedClassType, FName CallingContext)
 {
 	check(FindObject<UBlueprint>(Outer, *NewBPName.ToString()) == NULL); 
 
@@ -200,7 +200,7 @@ UBlueprint* FKismetEditorUtilities::CreateBlueprint(UClass* ParentClass, UObject
 		FName NewSkelClassName, NewGenClassName;
 		NewBP->GetBlueprintClassNames(NewGenClassName, NewSkelClassName);
 		UBlueprintGeneratedClass* NewClass = ConstructObject<UBlueprintGeneratedClass>(
-			UBlueprintGeneratedClass::StaticClass(), NewBP->GetOutermost(), NewGenClassName, RF_Public|RF_Transactional);
+			*BlueprintGeneratedClassType, NewBP->GetOutermost(), NewGenClassName, RF_Public | RF_Transactional);
 		NewBP->GeneratedClass = NewClass;
 		NewClass->ClassGeneratedBy = NewBP;
 		NewClass->SetSuperStruct(ParentClass);
@@ -363,11 +363,11 @@ void FKismetEditorUtilities::CompileBlueprint(UBlueprint* BlueprintObj, bool bIs
 		// Replace instances of this class
 		ReinstanceHelper.ReinstanceObjects();
 
- 		if (!bSkipGarbageCollection)
- 		{
- 			// Garbage collect to make sure the old class and actors are disposed of
- 			CollectGarbage(GARBAGE_COLLECTION_KEEPFLAGS);
- 		}
+		if (!bSkipGarbageCollection)
+		{
+			// Garbage collect to make sure the old class and actors are disposed of
+			CollectGarbage(GARBAGE_COLLECTION_KEEPFLAGS);
+		}
 
 		// If you need to verify that all old instances are taken care of, uncomment this!
 		// ReinstanceHelper.VerifyReplacement();
@@ -537,7 +537,7 @@ UBlueprint* FKismetEditorUtilities::CreateBlueprintFromActor(const FString& Path
 			else 
 			{
 				// We don't have a factory, but we can still try to create a blueprint for this actor class
-				NewBlueprint = FKismetEditorUtilities::CreateBlueprint( Object->GetClass(), Package, BlueprintName, EBlueprintType::BPTYPE_Normal, UBlueprint::StaticClass(), FName("CreateFromActor") );
+				NewBlueprint = FKismetEditorUtilities::CreateBlueprint( Object->GetClass(), Package, BlueprintName, EBlueprintType::BPTYPE_Normal, UBlueprint::StaticClass(), UBlueprintGeneratedClass::StaticClass(), FName("CreateFromActor") );
 			}
 		}
 
@@ -650,7 +650,7 @@ public:
 			}
 
 			UPackage* Package = CreatePackage(NULL, *PackageName);
-			Blueprint = FKismetEditorUtilities::CreateBlueprint(AActor::StaticClass(), Package, *AssetName, BPTYPE_Normal, UBlueprint::StaticClass(), FName("HarvestFromActors"));
+			Blueprint = FKismetEditorUtilities::CreateBlueprint(AActor::StaticClass(), Package, *AssetName, BPTYPE_Normal, UBlueprint::StaticClass(), UBlueprintGeneratedClass::StaticClass(), FName("HarvestFromActors"));
 
 			check(Blueprint->SimpleConstructionScript != NULL);
 			SCS = Blueprint->SimpleConstructionScript;
@@ -945,7 +945,7 @@ UBlueprint* FKismetEditorUtilities::CreateBlueprintFromClass(FText InWindowTitle
 		check(Package);
 
 		// Create and init a new Blueprint
-		UBlueprint* Blueprint = FKismetEditorUtilities::CreateBlueprint(InParentClass, Package, BPName, BPTYPE_Normal, UBlueprint::StaticClass(), FName("LevelEditorActions"));
+		UBlueprint* Blueprint = FKismetEditorUtilities::CreateBlueprint(InParentClass, Package, BPName, BPTYPE_Normal, UBlueprint::StaticClass(), UBlueprintGeneratedClass::StaticClass(), FName("LevelEditorActions"));
 		if (Blueprint)
 		{
 			// Notify the asset registry
diff --git a/Engine/Source/Editor/UnrealEd/Public/Kismet2/KismetEditorUtilities.h b/Engine/Source/Editor/UnrealEd/Public/Kismet2/KismetEditorUtilities.h
--- a/Engine/Source/Editor/UnrealEd/Public/Kismet2/KismetEditorUtilities.h
+++ b/Engine/Source/Editor/UnrealEd/Public/Kismet2/KismetEditorUtilities.h
@@ -19,15 +19,16 @@ public:
 	/**
 	 * Create a new Blueprint and initialize it to a valid state.
 	 *
-	 * @param ParentClass			the parent class of the new blueprint
-	 * @param Outer					the outer object of the new blueprint
-	 * @param NewBPName				the name of the new blueprint
-	 * @param BlueprintType			the type of the new blueprint (normal, const, etc)
-	 * @param BlueprintClassType	the actual class of the blueprint asset (UBlueprint or a derived type)
-	 * @param CallingContext		the name of the calling method or module used to identify creation methods to engine analytics/usage stats (default None will be ignored)
-	 * @return						the new blueprint
+	 * @param ParentClass					the parent class of the new blueprint
+	 * @param Outer							the outer object of the new blueprint
+	 * @param NewBPName						the name of the new blueprint
+	 * @param BlueprintType					the type of the new blueprint (normal, const, etc)
+	 * @param BlueprintClassType			the actual class of the blueprint asset (UBlueprint or a derived type)
+	 * @param BlueprintGeneratedClassType	the actual generated class of the blueprint asset (UBlueprintGeneratedClass or a derived type)
+	 * @param CallingContext				the name of the calling method or module used to identify creation methods to engine analytics/usage stats (default None will be ignored)
+	 * @return								the new blueprint
 	 */
-	static UBlueprint* CreateBlueprint(UClass* ParentClass, UObject* Outer, const FName NewBPName, enum EBlueprintType BlueprintType, TSubclassOf<UBlueprint> BlueprintClassType, FName CallingContext = NAME_None);
+	static UBlueprint* CreateBlueprint(UClass* ParentClass, UObject* Outer, const FName NewBPName, enum EBlueprintType BlueprintType, TSubclassOf<UBlueprint> BlueprintClassType, TSubclassOf<UBlueprintGeneratedClass> BlueprintGeneratedClassType, FName CallingContext = NAME_None);
 
 	/** Create the correct event graphs for this blueprint */
 	static void CreateDefaultEventGraphs(UBlueprint* Blueprint);
diff --git a/Engine/Source/Runtime/Engine/Classes/Engine/BlueprintGeneratedClass.h b/Engine/Source/Runtime/Engine/Classes/Engine/BlueprintGeneratedClass.h
--- a/Engine/Source/Runtime/Engine/Classes/Engine/BlueprintGeneratedClass.h
+++ b/Engine/Source/Runtime/Engine/Classes/Engine/BlueprintGeneratedClass.h
@@ -405,12 +405,11 @@ public:
 	 */
 	static bool GetGeneratedClassesHierarchy(const UClass* InClass, TArray<const UBlueprintGeneratedClass*>& OutBPGClasses);
 
-
 	/** Find the object in the TemplateObjects array with the supplied name */
 	UActorComponent* FindComponentTemplateByName(const FName& TemplateName);
 
 	/** Create Timeline objects for this Actor based on the Timelines array*/
-	void CreateTimelinesForActor(AActor* Actor) const;
+	virtual void CreateComponentsForActor(AActor* Actor) const;
 
 	// UObject interface
 	virtual void PostLoad() OVERRIDE;
diff --git a/Engine/Source/Runtime/Engine/Private/ActorConstruction.cpp b/Engine/Source/Runtime/Engine/Private/ActorConstruction.cpp
--- a/Engine/Source/Runtime/Engine/Private/ActorConstruction.cpp
+++ b/Engine/Source/Runtime/Engine/Private/ActorConstruction.cpp
@@ -272,7 +272,7 @@ void AActor::OnConstruction(const FTransform& Transform)
 					CurrentBPGClass->SimpleConstructionScript->ExecuteScriptOnActor(this, Transform);
 				}
 				// Now that the construction scripts have been run, we can create timelines and hook them up
-				CurrentBPGClass->CreateTimelinesForActor(this);
+				CurrentBPGClass->CreateComponentsForActor(this);
 			}
 
 #if WITH_EDITOR
diff --git a/Engine/Source/Runtime/Engine/Private/BlueprintGeneratedClass.cpp b/Engine/Source/Runtime/Engine/Private/BlueprintGeneratedClass.cpp
--- a/Engine/Source/Runtime/Engine/Private/BlueprintGeneratedClass.cpp
+++ b/Engine/Source/Runtime/Engine/Private/BlueprintGeneratedClass.cpp
@@ -320,7 +320,7 @@ UActorComponent* UBlueprintGeneratedClass::FindComponentTemplateByName(const FNa
 	return NULL;
 }
 
-void UBlueprintGeneratedClass::CreateTimelinesForActor(AActor* Actor) const
+void UBlueprintGeneratedClass::CreateComponentsForActor(AActor* Actor) const
 {
 	check(Actor != NULL);
 	check(!Actor->IsTemplate());
diff --git a/Engine/Source/Runtime/Engine/Private/Level.cpp b/Engine/Source/Runtime/Engine/Private/Level.cpp
--- a/Engine/Source/Runtime/Engine/Private/Level.cpp
+++ b/Engine/Source/Runtime/Engine/Private/Level.cpp
@@ -1650,7 +1650,7 @@ ULevelScriptBlueprint* ULevel::GetLevelScriptBlueprint(bool bDontCreate)
 	if( !LevelScriptBlueprint && !bDontCreate)
 	{
 		// If no blueprint is found, create one. 
-		LevelScriptBlueprint = Cast<ULevelScriptBlueprint>(FKismetEditorUtilities::CreateBlueprint(GEngine->LevelScriptActorClass, this, FName(*LevelScriptName), BPTYPE_LevelScript, ULevelScriptBlueprint::StaticClass()));
+		LevelScriptBlueprint = Cast<ULevelScriptBlueprint>(FKismetEditorUtilities::CreateBlueprint(GEngine->LevelScriptActorClass, this, FName(*LevelScriptName), BPTYPE_LevelScript, ULevelScriptBlueprint::StaticClass(), UBlueprintGeneratedClass::StaticClass()));
 
 		// LevelScript blueprints should not be standalone
 		LevelScriptBlueprint->ClearFlags(RF_Standalone);
diff --git a/Engine/Source/Runtime/UMG/Classes/BoxPanelComponent.h b/Engine/Source/Runtime/UMG/Classes/BoxPanelComponent.h
--- a/Engine/Source/Runtime/UMG/Classes/BoxPanelComponent.h
+++ b/Engine/Source/Runtime/UMG/Classes/BoxPanelComponent.h
@@ -4,7 +4,6 @@
 
 #include "BoxPanelComponent.generated.h"
 
-
 USTRUCT()
 struct UMG_API FBoxPanelSlot
 {
@@ -63,9 +62,9 @@ class UMG_API UBoxPanelComponent : public USlateNonLeafWidgetComponent
 #endif
 
 	// USceneComponent interface
-	virtual bool HasAnySockets() const OVERRIDE;
-	virtual FTransform GetSocketTransform(FName InSocketName, ERelativeTransformSpace TransformSpace = RTS_World) const OVERRIDE;
-	virtual void QuerySupportedSockets(TArray<FComponentSocketDescription>& OutSockets) const OVERRIDE;
+	virtual bool HasAnySockets() const;
+	virtual FTransform GetSocketTransform(FName InSocketName, ERelativeTransformSpace TransformSpace = RTS_World) const;
+	virtual void QuerySupportedSockets(TArray<FComponentSocketDescription>& OutSockets) const;
 	// End of USceneComponent interface
 
 protected:
diff --git a/Engine/Source/Runtime/UMG/Classes/ButtonComponent.h b/Engine/Source/Runtime/UMG/Classes/ButtonComponent.h
--- a/Engine/Source/Runtime/UMG/Classes/ButtonComponent.h
+++ b/Engine/Source/Runtime/UMG/Classes/ButtonComponent.h
@@ -14,7 +14,7 @@ class UMG_API UButtonComponent : public USlateNonLeafWidgetComponent
 {
 	GENERATED_UCLASS_BODY()
 
-protected:
+public:
 	/** Style of the button */
 	UPROPERTY(EditAnywhere, BlueprintReadOnly, Category=Appearance)
 	USlateWidgetStyleAsset* ButtonStyle;
diff --git a/Engine/Source/Runtime/UMG/Classes/CanvasPanelComponent.h b/Engine/Source/Runtime/UMG/Classes/CanvasPanelComponent.h
--- a/Engine/Source/Runtime/UMG/Classes/CanvasPanelComponent.h
+++ b/Engine/Source/Runtime/UMG/Classes/CanvasPanelComponent.h
@@ -9,26 +9,23 @@ class UMG_API UCanvasPanelComponent : public USlateNonLeafWidgetComponent
 {
 	GENERATED_UCLASS_BODY()
 
+	/** The items placed on the canvas */
+	UPROPERTY(EditAnywhere, EditInline, Category=Slots)
+	TArray<UCanvasPanelSlot*> Slots;
+
+	/** The desired size of the canvas */
+	UPROPERTY(EditAnywhere, Category=Appearance)
+	FVector2D DesiredCanvasSize;
+
+	UCanvasPanelSlot* AddSlot(USlateWrapperComponent* Content);
+
 #if WITH_EDITOR
 	// UObject interface
 	virtual void PostEditChangeProperty(struct FPropertyChangedEvent& PropertyChangedEvent) OVERRIDE;
 	// End of UObject interface
 #endif
 
-	// USceneComponent interface
-	virtual bool HasAnySockets() const OVERRIDE;
-	virtual FTransform GetSocketTransform(FName InSocketName, ERelativeTransformSpace TransformSpace = RTS_World) const OVERRIDE;
-	virtual void QuerySupportedSockets(TArray<FComponentSocketDescription>& OutSockets) const OVERRIDE;
-	// End of USceneComponent interface
-
 protected:
-	/** The items placed on the canvas */
-	UPROPERTY(EditAnywhere, Category=Slots)
-	TArray<class UCanvasPanelSlot*> Slots;
-
-	/** The desired size of the canvas */
-	UPROPERTY(EditAnywhere, Category=Appearance)
-	FVector2D DesiredCanvasSize;
 
 	TWeakPtr<class SFixedSizeCanvas> MyCanvas;
 
diff --git a/Engine/Source/Runtime/UMG/Classes/CanvasPanelSlot.h b/Engine/Source/Runtime/UMG/Classes/CanvasPanelSlot.h
--- a/Engine/Source/Runtime/UMG/Classes/CanvasPanelSlot.h
+++ b/Engine/Source/Runtime/UMG/Classes/CanvasPanelSlot.h
@@ -4,14 +4,13 @@
 
 #include "CanvasPanelSlot.generated.h"
 
-USTRUCT()
-struct UMG_API FCanvasSlotData
+UCLASS()
+class UMG_API UCanvasPanelSlot : public UPanelSlot
 {
-	GENERATED_USTRUCT_BODY()
-
-	/** Name of the slot (use this to attach children to the slot) */
-	UPROPERTY(EditAnywhere, Category=Appearance)
-	FName SlotName;
+	GENERATED_UCLASS_BODY()
+	
+	UPROPERTY()
+	class USlateWrapperComponent* Content;
 
 	/** Position. */
 	UPROPERTY(EditAnywhere, Category=Appearance)
@@ -23,14 +22,14 @@ struct UMG_API FCanvasSlotData
 
 	/**
 	* Horizontal pivot position
-	*  Given a top aligned slot, where '+' represents the 
+	*  Given a top aligned slot, where '+' represents the
 	*  anchor point defined by PositionAttr.
-	*  
+	*
 	*   Left				Center				Right
-		+ _ _ _ _            _ _ + _ _          _ _ _ _ +
-		|		  |		   | 		   |	  |		    |
-		| _ _ _ _ |        | _ _ _ _ _ |	  | _ _ _ _ |
-	* 
+	+ _ _ _ _            _ _ + _ _          _ _ _ _ +
+	|		  |		   | 		   |	  |		    |
+	| _ _ _ _ |        | _ _ _ _ _ |	  | _ _ _ _ |
+	*
 	*  Note: FILL is NOT supported.
 	*/
 	UPROPERTY(EditAnywhere, Category=Appearance)
@@ -38,35 +37,17 @@ struct UMG_API FCanvasSlotData
 
 	/**
 	* Vertical pivot position
-	*   Given a left aligned slot, where '+' represents the 
+	*   Given a left aligned slot, where '+' represents the
 	*   anchor point defined by PositionAttr.
-	*  
+	*
 	*   Top					Center			  Bottom
-	*	+_ _ _ _ _		 _ _ _ _ _		 _ _ _ _ _ 
+	*	+_ _ _ _ _		 _ _ _ _ _		 _ _ _ _ _
 	*	|         |		| 		  |		|		  |
 	*	|         |     +		  |		|		  |
 	*	| _ _ _ _ |		| _ _ _ _ |		+ _ _ _ _ |
-	* 
+	*
 	*  Note: FILL is NOT supported.
 	*/
 	UPROPERTY(EditAnywhere, Category=Appearance)
 	TEnumAsByte<EVerticalAlignment> VerticalAlignment;
-
-	FCanvasSlotData()
-		: SlotName(NAME_None) 
-		, Position(FVector2D::ZeroVector)
-		, Size(FVector2D(1.0f, 1.0f))
-		, HorizontalAlignment(HAlign_Left)
-		, VerticalAlignment(VAlign_Top)
-	{
-	}
-};
-
-UCLASS()
-class UMG_API UCanvasPanelSlot : public UPanelSlot
-{
-	GENERATED_UCLASS_BODY()
-
-	UPROPERTY(EditAnywhere, Category=Slot)
-	FCanvasSlotData Data;
 };
\ No newline at end of file
diff --git a/Engine/Source/Runtime/UMG/Classes/GridPanelComponent.h b/Engine/Source/Runtime/UMG/Classes/GridPanelComponent.h
--- a/Engine/Source/Runtime/UMG/Classes/GridPanelComponent.h
+++ b/Engine/Source/Runtime/UMG/Classes/GridPanelComponent.h
@@ -73,9 +73,9 @@ class UMG_API UGridPanelComponent : public USlateNonLeafWidgetComponent
 #endif
 
 	// USceneComponent interface
-	virtual bool HasAnySockets() const OVERRIDE;
-	virtual FTransform GetSocketTransform(FName InSocketName, ERelativeTransformSpace TransformSpace = RTS_World) const OVERRIDE;
-	virtual void QuerySupportedSockets(TArray<FComponentSocketDescription>& OutSockets) const OVERRIDE;
+	virtual bool HasAnySockets() const;
+	virtual FTransform GetSocketTransform(FName InSocketName, ERelativeTransformSpace TransformSpace = RTS_World) const;
+	virtual void QuerySupportedSockets(TArray<FComponentSocketDescription>& OutSockets) const;
 	// End of USceneComponent interface
 
 protected:
diff --git a/Engine/Source/Runtime/UMG/Classes/HorizontalBoxComponent.h b/Engine/Source/Runtime/UMG/Classes/HorizontalBoxComponent.h
new file mode 100644
--- /dev/null
+++ b/Engine/Source/Runtime/UMG/Classes/HorizontalBoxComponent.h
@@ -0,0 +1,36 @@
+// Copyright 1998-2014 Epic Games, Inc. All Rights Reserved.
+
+#pragma once
+
+#include "HorizontalBoxComponent.generated.h"
+
+UCLASS(meta=(BlueprintSpawnableComponent), ClassGroup=UserInterface)
+class UMG_API UHorizontalBoxComponent : public USlateNonLeafWidgetComponent
+{
+	GENERATED_UCLASS_BODY()
+
+	/** The items placed on the canvas */
+	UPROPERTY(EditAnywhere, EditInline, Category=Slots)
+	TArray<UHorizontalBoxSlot*> Slots;
+
+	UHorizontalBoxSlot* AddSlot(USlateWrapperComponent* Content);
+
+#if WITH_EDITOR
+	// UObject interface
+	virtual void PostEditChangeProperty(struct FPropertyChangedEvent& PropertyChangedEvent) OVERRIDE;
+	// End of UObject interface
+#endif
+
+protected:
+
+	TWeakPtr<class SHorizontalBox> MyHorizontalBox;
+
+protected:
+	// USlateWrapperComponent interface
+	virtual TSharedRef<SWidget> RebuildWidget() OVERRIDE;
+	// End of USlateWrapperComponent interface
+
+	// USlateNonLeafWidgetComponent interface
+	virtual void OnKnownChildrenChanged() OVERRIDE;
+	// End of USlateNonLeafWidgetComponent
+};
diff --git a/Engine/Source/Runtime/UMG/Classes/HorizontalBoxSlot.h b/Engine/Source/Runtime/UMG/Classes/HorizontalBoxSlot.h
new file mode 100644
--- /dev/null
+++ b/Engine/Source/Runtime/UMG/Classes/HorizontalBoxSlot.h
@@ -0,0 +1,20 @@
+// Copyright 1998-2014 Epic Games, Inc. All Rights Reserved.
+
+#pragma once
+
+#include "HorizontalBoxSlot.generated.h"
+
+UCLASS()
+class UMG_API UHorizontalBoxSlot : public UPanelSlot
+{
+	GENERATED_UCLASS_BODY()
+	
+	UPROPERTY()
+	class USlateWrapperComponent* Content;
+
+	UPROPERTY(EditAnywhere, Category=Appearance)
+	TEnumAsByte<EHorizontalAlignment> HorizontalAlignment;
+
+	UPROPERTY(EditAnywhere, Category=Appearance)
+	TEnumAsByte<EVerticalAlignment> VerticalAlignment;
+};
\ No newline at end of file
diff --git a/Engine/Source/Runtime/UMG/Classes/SlateLeafWidgetComponent.h b/Engine/Source/Runtime/UMG/Classes/SlateLeafWidgetComponent.h
--- a/Engine/Source/Runtime/UMG/Classes/SlateLeafWidgetComponent.h
+++ b/Engine/Source/Runtime/UMG/Classes/SlateLeafWidgetComponent.h
@@ -11,6 +11,6 @@ class UMG_API USlateLeafWidgetComponent: public USlateWrapperComponent
 	GENERATED_UCLASS_BODY()
 
 	// USceneComponent interface
-	virtual bool CanAttachAsChild(USceneComponent* ChildComponent, FName SocketName) const OVERRIDE;
+	virtual bool CanAttachAsChild(USceneComponent* ChildComponent, FName SocketName) const;
 	// End of USceneComponent interface
 };
diff --git a/Engine/Source/Runtime/UMG/Classes/SlateNonLeafWidgetComponent.h b/Engine/Source/Runtime/UMG/Classes/SlateNonLeafWidgetComponent.h
--- a/Engine/Source/Runtime/UMG/Classes/SlateNonLeafWidgetComponent.h
+++ b/Engine/Source/Runtime/UMG/Classes/SlateNonLeafWidgetComponent.h
@@ -19,10 +19,6 @@ class UMG_API USlateNonLeafWidgetComponent : public USlateWrapperComponent
 	// End of UActorComponent interface
 #endif
 
-	// USceneComponent interface
-	virtual bool CanAttachAsChild(USceneComponent* ChildComponent, FName SocketName) const OVERRIDE;
-	// End of USceneComponent interface
-
 protected:
 	// This function is called when the known children list has changed; use this to repopulate child slots in derived classes
 	virtual void OnKnownChildrenChanged() {}
diff --git a/Engine/Source/Runtime/UMG/Classes/SlateWrapperComponent.h b/Engine/Source/Runtime/UMG/Classes/SlateWrapperComponent.h
--- a/Engine/Source/Runtime/UMG/Classes/SlateWrapperComponent.h
+++ b/Engine/Source/Runtime/UMG/Classes/SlateWrapperComponent.h
@@ -7,7 +7,7 @@
 
 /** This is the base class for all Slate widget wrapper components */
 UCLASS(Abstract, hideCategories=(Activation, "Components|Activation", ComponentReplication, LOD, Rendering))
-class UMG_API USlateWrapperComponent : public USceneComponent
+class UMG_API USlateWrapperComponent : public UActorComponent
 {
 	GENERATED_UCLASS_BODY()
 
diff --git a/Engine/Source/Runtime/UMG/Classes/SlateWrapperTypes.h b/Engine/Source/Runtime/UMG/Classes/SlateWrapperTypes.h
--- a/Engine/Source/Runtime/UMG/Classes/SlateWrapperTypes.h
+++ b/Engine/Source/Runtime/UMG/Classes/SlateWrapperTypes.h
@@ -35,9 +35,10 @@ namespace ESlateSizeRule
 {
 	enum Type
 	{
+		/** The container will size to fit the needs of the child widgets */
 		Automatic,
-		AspectRatio,
-		Stretch
+		/** The container will fill the percentage of the container based on the Value 0..1 */
+		Fill
 	};
 }
 
@@ -46,15 +47,16 @@ struct FSlateChildSize
 {
 	GENERATED_USTRUCT_BODY()
 
-	UPROPERTY(EditAnywhere, Category=Appearance)
-	float StretchValue;
+	//editcondition="SizeRule=ESlateSizeRule::Fill"
+	UPROPERTY(EditAnywhere, Category=Appearance, meta=(ClampMin="0", ClampMax="1"))
+	float Value;
 
 	UPROPERTY(EditAnywhere, Category=Appearance)
 	TEnumAsByte<ESlateSizeRule::Type> SizeRule;
 
 	FSlateChildSize()
-		: StretchValue(1.0f)
-		, SizeRule(ESlateSizeRule::Stretch)
+		: Value(1.0f)
+		, SizeRule(ESlateSizeRule::Fill)
 	{
 	}
 };
diff --git a/Engine/Source/Runtime/UMG/Classes/VerticalBoxComponent.h b/Engine/Source/Runtime/UMG/Classes/VerticalBoxComponent.h
new file mode 100644
--- /dev/null
+++ b/Engine/Source/Runtime/UMG/Classes/VerticalBoxComponent.h
@@ -0,0 +1,36 @@
+// Copyright 1998-2014 Epic Games, Inc. All Rights Reserved.
+
+#pragma once
+
+#include "VerticalBoxComponent.generated.h"
+
+UCLASS(meta=(BlueprintSpawnableComponent), ClassGroup=UserInterface)
+class UMG_API UVerticalBoxComponent : public USlateNonLeafWidgetComponent
+{
+	GENERATED_UCLASS_BODY()
+
+	/** The items placed on the canvas */
+	UPROPERTY(EditAnywhere, EditInline, Category=Slots)
+	TArray<UVerticalBoxSlot*> Slots;
+
+	UVerticalBoxSlot* AddSlot(USlateWrapperComponent* Content);
+
+#if WITH_EDITOR
+	// UObject interface
+	virtual void PostEditChangeProperty(struct FPropertyChangedEvent& PropertyChangedEvent) OVERRIDE;
+	// End of UObject interface
+#endif
+
+protected:
+
+	TWeakPtr<class SVerticalBox> MyVerticalBox;
+
+protected:
+	// USlateWrapperComponent interface
+	virtual TSharedRef<SWidget> RebuildWidget() OVERRIDE;
+	// End of USlateWrapperComponent interface
+
+	// USlateNonLeafWidgetComponent interface
+	virtual void OnKnownChildrenChanged() OVERRIDE;
+	// End of USlateNonLeafWidgetComponent
+};
diff --git a/Engine/Source/Runtime/UMG/Classes/VerticalBoxSlot.h b/Engine/Source/Runtime/UMG/Classes/VerticalBoxSlot.h
new file mode 100644
--- /dev/null
+++ b/Engine/Source/Runtime/UMG/Classes/VerticalBoxSlot.h
@@ -0,0 +1,29 @@
+// Copyright 1998-2014 Epic Games, Inc. All Rights Reserved.
+
+#pragma once
+
+#include "SlateWrapperTypes.h"
+
+#include "VerticalBoxSlot.generated.h"
+
+UCLASS()
+class UMG_API UVerticalBoxSlot : public UPanelSlot
+{
+	GENERATED_UCLASS_BODY()
+	
+	UPROPERTY()
+	class USlateWrapperComponent* Content;
+
+	UPROPERTY(EditAnywhere, Category=Appearance)
+	FMargin Padding;
+
+	/** How much space this slot should occupy in the direction of the panel. */
+	UPROPERTY(EditAnywhere, Category=Appearance)
+	FSlateChildSize Size;
+
+	UPROPERTY(EditAnywhere, Category=Appearance)
+	TEnumAsByte<EHorizontalAlignment> HorizontalAlignment;
+
+	UPROPERTY(EditAnywhere, Category=Appearance)
+	TEnumAsByte<EVerticalAlignment> VerticalAlignment;
+};
diff --git a/Engine/Source/Runtime/UMG/Classes/WidgetBlueprint.h b/Engine/Source/Runtime/UMG/Classes/WidgetBlueprint.h
--- a/Engine/Source/Runtime/UMG/Classes/WidgetBlueprint.h
+++ b/Engine/Source/Runtime/UMG/Classes/WidgetBlueprint.h
@@ -9,7 +9,10 @@ class UMG_API UWidgetBlueprint : public UBlueprint
 {
 	GENERATED_UCLASS_BODY()
 
-	//TODO Add Hierarchy.	
+public:
+	/** Array of templates for timelines that should be created */
+	UPROPERTY()
+	TArray<class USlateWrapperComponent*> WidgetTemplates;
 	
 #if WITH_EDITOR
 	// UBlueprint interface
@@ -18,5 +21,7 @@ class UMG_API UWidgetBlueprint : public UBlueprint
 		return false;
 	}
 	// End of UBlueprint interface
+
+	static bool ValidateGeneratedClass(const UClass* InClass);
 #endif
 };
\ No newline at end of file
diff --git a/Engine/Source/Runtime/UMG/Classes/WidgetBlueprintGeneratedClass.h b/Engine/Source/Runtime/UMG/Classes/WidgetBlueprintGeneratedClass.h
--- a/Engine/Source/Runtime/UMG/Classes/WidgetBlueprintGeneratedClass.h
+++ b/Engine/Source/Runtime/UMG/Classes/WidgetBlueprintGeneratedClass.h
@@ -9,4 +9,8 @@ class UMG_API UWidgetBlueprintGeneratedClass : public UBlueprintGeneratedClass
 {
 	GENERATED_UCLASS_BODY()
 
+	UPROPERTY()
+	TArray<class USlateWrapperComponent*> WidgetTemplates;
+
+	virtual void CreateComponentsForActor(AActor* Actor) const OVERRIDE;
 };
\ No newline at end of file
diff --git a/Engine/Source/Runtime/UMG/Private/BorderComponent.cpp b/Engine/Source/Runtime/UMG/Private/BorderComponent.cpp
--- a/Engine/Source/Runtime/UMG/Private/BorderComponent.cpp
+++ b/Engine/Source/Runtime/UMG/Private/BorderComponent.cpp
@@ -2,7 +2,7 @@
 
 #include "UMGPrivatePCH.h"
 
-#define LOCTEXT_NAMESPACE "SlateComponentWrappers"
+#define LOCTEXT_NAMESPACE "UMG"
 
 /////////////////////////////////////////////////////
 // UBorderComponent
diff --git a/Engine/Source/Runtime/UMG/Private/BoxPanelComponent.cpp b/Engine/Source/Runtime/UMG/Private/BoxPanelComponent.cpp
--- a/Engine/Source/Runtime/UMG/Private/BoxPanelComponent.cpp
+++ b/Engine/Source/Runtime/UMG/Private/BoxPanelComponent.cpp
@@ -119,56 +119,56 @@ TSharedRef<SWidget> UBoxPanelComponent::RebuildWidget()
 
 void UBoxPanelComponent::OnKnownChildrenChanged()
 {
-	TSharedPtr<SConcreteBoxPanel> Box = MyBox.Pin();
-	if (Box.IsValid())
-	{
-		Box->ClearChildren();
-
-		// Add slots
-		TMap<FName, int32> SlotNameToIndex;
-
-		for (int32 SlotIndex = 0; SlotIndex < Slots.Num(); ++SlotIndex)
-		{
-			const FBoxPanelSlot& SlotConfig = Slots[SlotIndex];
-
-			SlotNameToIndex.Add(SlotConfig.SlotName, SlotIndex);
-
-			FSizeParam SlotSize = USlateWrapperComponent::ConvertSerializedSizeParamToRuntime(SlotConfig.SlotSize);
-
-			auto NewSlot = Box->AddSlot()
-				.Padding(SlotConfig.Padding)
-				.HAlign(SlotConfig.HorizontalAlignment)
-				.VAlign(SlotConfig.VerticalAlignment)
-				.SetSizeParam(SlotSize)
-				.MaxWidth(SlotConfig.bHasMaximumSize ? SlotConfig.MaxSize : 0.0f)
-			[
-				SNullWidget::NullWidget
-			];
-		}
-
-		// Place widgets in their slots, and add anything left over to the end
-		for (int32 ComponentIndex = 0; ComponentIndex < AttachChildren.Num(); ++ComponentIndex)
-		{
-			if (USlateWrapperComponent* Wrapper = Cast<USlateWrapperComponent>(AttachChildren[ComponentIndex]))
-			{
-				if (Wrapper->IsRegistered())
-				{
-					if (int32* pSlotIndex = SlotNameToIndex.Find(Wrapper->AttachSocketName))
-					{
-						auto Foo = (TPanelChildren<SHorizontalBox::FSlot>*)(Box->GetChildren());
-						(*Foo)[*pSlotIndex].Widget = Wrapper->GetWidget();
-					}
-					else
-					{
-						Box->AddSlot()
-						[
-							Wrapper->GetWidget()
-						];
-					}
-				}
-			}
-		}
-	}
+	//TSharedPtr<SConcreteBoxPanel> Box = MyBox.Pin();
+	//if (Box.IsValid())
+	//{
+	//	Box->ClearChildren();
+
+	//	// Add slots
+	//	TMap<FName, int32> SlotNameToIndex;
+
+	//	for (int32 SlotIndex = 0; SlotIndex < Slots.Num(); ++SlotIndex)
+	//	{
+	//		const FBoxPanelSlot& SlotConfig = Slots[SlotIndex];
+
+	//		SlotNameToIndex.Add(SlotConfig.SlotName, SlotIndex);
+
+	//		FSizeParam SlotSize = USlateWrapperComponent::ConvertSerializedSizeParamToRuntime(SlotConfig.SlotSize);
+
+	//		auto NewSlot = Box->AddSlot()
+	//			.Padding(SlotConfig.Padding)
+	//			.HAlign(SlotConfig.HorizontalAlignment)
+	//			.VAlign(SlotConfig.VerticalAlignment)
+	//			.SetSizeParam(SlotSize)
+	//			.MaxWidth(SlotConfig.bHasMaximumSize ? SlotConfig.MaxSize : 0.0f)
+	//		[
+	//			SNullWidget::NullWidget
+	//		];
+	//	}
+
+	//	// Place widgets in their slots, and add anything left over to the end
+	//	for (int32 ComponentIndex = 0; ComponentIndex < AttachChildren.Num(); ++ComponentIndex)
+	//	{
+	//		if (USlateWrapperComponent* Wrapper = Cast<USlateWrapperComponent>(AttachChildren[ComponentIndex]))
+	//		{
+	//			if (Wrapper->IsRegistered())
+	//			{
+	//				if (int32* pSlotIndex = SlotNameToIndex.Find(Wrapper->AttachSocketName))
+	//				{
+	//					auto Foo = (TPanelChildren<SHorizontalBox::FSlot>*)(Box->GetChildren());
+	//					(*Foo)[*pSlotIndex].Widget = Wrapper->GetWidget();
+	//				}
+	//				else
+	//				{
+	//					Box->AddSlot()
+	//					[
+	//						Wrapper->GetWidget()
+	//					];
+	//				}
+	//			}
+	//		}
+	//	}
+	//}
 }
 
 bool UBoxPanelComponent::HasAnySockets() const
diff --git a/Engine/Source/Runtime/UMG/Private/ButtonComponent.cpp b/Engine/Source/Runtime/UMG/Private/ButtonComponent.cpp
--- a/Engine/Source/Runtime/UMG/Private/ButtonComponent.cpp
+++ b/Engine/Source/Runtime/UMG/Private/ButtonComponent.cpp
@@ -2,7 +2,7 @@
 
 #include "UMGPrivatePCH.h"
 
-#define LOCTEXT_NAMESPACE "SlateComponentWrappers"
+#define LOCTEXT_NAMESPACE "UMG"
 
 /////////////////////////////////////////////////////
 // UButtonComponent
diff --git a/Engine/Source/Runtime/UMG/Private/CanvasPanelComponent.cpp b/Engine/Source/Runtime/UMG/Private/CanvasPanelComponent.cpp
--- a/Engine/Source/Runtime/UMG/Private/CanvasPanelComponent.cpp
+++ b/Engine/Source/Runtime/UMG/Private/CanvasPanelComponent.cpp
@@ -45,6 +45,8 @@ TSharedRef<SWidget> UCanvasPanelComponent::RebuildWidget()
 	TSharedRef<SFixedSizeCanvas> NewCanvas = SNew(SFixedSizeCanvas, DesiredCanvasSize);
 	MyCanvas = NewCanvas;
 
+	OnKnownChildrenChanged();
+
 	return NewCanvas;
 }
 
@@ -55,73 +57,34 @@ void UCanvasPanelComponent::OnKnownChildrenChanged()
 	{
 		Canvas->ClearChildren();
 
-		// Add slots
-		TMap<FName, int32> SlotNameToIndex;
-
 		for (int32 SlotIndex = 0; SlotIndex < Slots.Num(); ++SlotIndex)
 		{
-			const UCanvasPanelSlot* Slot = Slots[SlotIndex];
+			UCanvasPanelSlot* Slot = Slots[SlotIndex];
 			if ( Slot == NULL )
 			{
-				continue;
+				Slots[SlotIndex] = Slot = ConstructObject<UCanvasPanelSlot>(UCanvasPanelSlot::StaticClass(), this);
 			}
 
-			const FCanvasSlotData& SlotConfig = Slot->Data;
-
-			SlotNameToIndex.Add(SlotConfig.SlotName, SlotIndex);
-
 			auto NewSlot = Canvas->AddSlot()
-				.Position(SlotConfig.Position)
-				.Size(SlotConfig.Size)
-				.HAlign(SlotConfig.HorizontalAlignment)
-				.VAlign(SlotConfig.VerticalAlignment)
+				.Position(Slot->Position)
+				.Size(Slot->Size)
+				.HAlign(Slot->HorizontalAlignment)
+				.VAlign(Slot->VerticalAlignment)
 			[
-				SNullWidget::NullWidget
+				Slot->Content == NULL ? SNullWidget::NullWidget : Slot->Content->GetWidget()
 			];
 		}
-
-		// Place widgets in their slots, and add anything left over to the end
-		for (int32 ComponentIndex = 0; ComponentIndex < AttachChildren.Num(); ++ComponentIndex)
-		{
-			if (USlateWrapperComponent* Wrapper = Cast<USlateWrapperComponent>(AttachChildren[ComponentIndex]))
-			{
-				if (Wrapper->IsRegistered())
-				{
-					if (int32* pSlotIndex = SlotNameToIndex.Find(Wrapper->AttachSocketName))
-					{
-						auto ExistingSlot = (TPanelChildren<SCanvas::FSlot>*)(Canvas->GetChildren());
-						(*ExistingSlot)[*pSlotIndex].Widget = Wrapper->GetWidget();
-					}
-					else
-					{
-						Canvas->AddSlot()
-						[
-							Wrapper->GetWidget()
-						];
-					}
-				}
-			}
-		}
 	}
 }
 
-bool UCanvasPanelComponent::HasAnySockets() const
+UCanvasPanelSlot* UCanvasPanelComponent::AddSlot(USlateWrapperComponent* Content)
 {
-	return true;
-}
+	UCanvasPanelSlot* Slot = ConstructObject<UCanvasPanelSlot>(UCanvasPanelSlot::StaticClass(), this);
+	Slot->Content = Content;
+	
+	Slots.Add(Slot);
 
-FTransform UCanvasPanelComponent::GetSocketTransform(FName InSocketName, ERelativeTransformSpace TransformSpace) const
-{
-	return FTransform::Identity;
-}
-
-void UCanvasPanelComponent::QuerySupportedSockets(TArray<FComponentSocketDescription>& OutSockets) const
-{
-	for (int32 SlotIndex = 0; SlotIndex < Slots.Num(); ++SlotIndex)
-	{
-		FName SocketName = Slots[SlotIndex]->Data.SlotName;
-		new (OutSockets) FComponentSocketDescription(SocketName, EComponentSocketType::Socket);
-	}
+	return Slot;
 }
 
 #if WITH_EDITOR
@@ -133,17 +96,22 @@ void UCanvasPanelComponent::PostEditChangeProperty(struct FPropertyChangedEvent&
 	TSet<FName> UniqueSlotNames;
 	for (int32 SlotIndex = 0; SlotIndex < Slots.Num(); ++SlotIndex)
 	{
-		FName& SlotName = Slots[SlotIndex]->Data.SlotName;
-
-		if ((SlotName == NAME_None) || UniqueSlotNames.Contains(SlotName))
+		if ( Slots[SlotIndex] == NULL )
 		{
-			do 
-			{
-				SlotName = FName(TEXT("Slot"), SlotNumbering++);
-			} while (UniqueSlotNames.Contains(SlotName));
+			Slots[SlotIndex] = ConstructObject<UCanvasPanelSlot>(UCanvasPanelSlot::StaticClass(), this);
 		}
 
-		UniqueSlotNames.Add(SlotName);
+		//FName& SlotName = Slots[SlotIndex].SlotName;
+
+		//if ((SlotName == NAME_None) || UniqueSlotNames.Contains(SlotName))
+		//{
+		//	do 
+		//	{
+		//		SlotName = FName(TEXT("Slot"), SlotNumbering++);
+		//	} while (UniqueSlotNames.Contains(SlotName));
+		//}
+
+		//UniqueSlotNames.Add(SlotName);
 	}
 }
 #endif
\ No newline at end of file
diff --git a/Engine/Source/Runtime/UMG/Private/CanvasPanelSlot.cpp b/Engine/Source/Runtime/UMG/Private/CanvasPanelSlot.cpp
--- a/Engine/Source/Runtime/UMG/Private/CanvasPanelSlot.cpp
+++ b/Engine/Source/Runtime/UMG/Private/CanvasPanelSlot.cpp
@@ -8,4 +8,8 @@
 UCanvasPanelSlot::UCanvasPanelSlot(const FPostConstructInitializeProperties& PCIP)
 	: Super(PCIP)
 {
-}
\ No newline at end of file
+	Position = FVector2D::ZeroVector;
+	Size = FVector2D(1.0f, 1.0f);
+	HorizontalAlignment = HAlign_Left;
+	VerticalAlignment = VAlign_Top;
+}
diff --git a/Engine/Source/Runtime/UMG/Private/CheckBoxComponent.cpp b/Engine/Source/Runtime/UMG/Private/CheckBoxComponent.cpp
--- a/Engine/Source/Runtime/UMG/Private/CheckBoxComponent.cpp
+++ b/Engine/Source/Runtime/UMG/Private/CheckBoxComponent.cpp
@@ -2,7 +2,7 @@
 
 #include "UMGPrivatePCH.h"
 
-#define LOCTEXT_NAMESPACE "SlateComponentWrappers"
+#define LOCTEXT_NAMESPACE "UMG"
 
 /////////////////////////////////////////////////////
 // UCheckBoxComponent
diff --git a/Engine/Source/Runtime/UMG/Private/CircularThrobberComponent.cpp b/Engine/Source/Runtime/UMG/Private/CircularThrobberComponent.cpp
--- a/Engine/Source/Runtime/UMG/Private/CircularThrobberComponent.cpp
+++ b/Engine/Source/Runtime/UMG/Private/CircularThrobberComponent.cpp
@@ -2,7 +2,7 @@
 
 #include "UMGPrivatePCH.h"
 
-#define LOCTEXT_NAMESPACE "SlateComponentWrappers"
+#define LOCTEXT_NAMESPACE "UMG"
 
 /////////////////////////////////////////////////////
 // UCircularThrobberComponent
diff --git a/Engine/Source/Runtime/UMG/Private/EditableTextBlockComponent.cpp b/Engine/Source/Runtime/UMG/Private/EditableTextBlockComponent.cpp
--- a/Engine/Source/Runtime/UMG/Private/EditableTextBlockComponent.cpp
+++ b/Engine/Source/Runtime/UMG/Private/EditableTextBlockComponent.cpp
@@ -2,7 +2,7 @@
 
 #include "UMGPrivatePCH.h"
 
-#define LOCTEXT_NAMESPACE "SlateComponentWrappers"
+#define LOCTEXT_NAMESPACE "UMG"
 
 /////////////////////////////////////////////////////
 // UEditableTextBlockComponent
diff --git a/Engine/Source/Runtime/UMG/Private/GridPanelComponent.cpp b/Engine/Source/Runtime/UMG/Private/GridPanelComponent.cpp
--- a/Engine/Source/Runtime/UMG/Private/GridPanelComponent.cpp
+++ b/Engine/Source/Runtime/UMG/Private/GridPanelComponent.cpp
@@ -24,68 +24,68 @@ TSharedRef<SWidget> UGridPanelComponent::RebuildWidget()
 
 void UGridPanelComponent::OnKnownChildrenChanged()
 {
-	TSharedPtr<SGridPanel> Grid = MyGrid.Pin();
-	if (Grid.IsValid())
-	{
-		// Create a mapping of slot names to slot indices
-		TMap<FName, int32> SlotNameToIndex;
-		for (int32 SlotIndex = 0; SlotIndex < Slots.Num(); ++SlotIndex)
-		{
-			const FGridPanelSlot& SlotConfig = Slots[SlotIndex];
-			SlotNameToIndex.Add(SlotConfig.SlotName, SlotIndex);
-		}
-		
-		// Make sure all children are assigned to a slot
-		TArray<USlateWrapperComponent*> WrappedComponentsInSlotOrder;
-		WrappedComponentsInSlotOrder.AddZeroed(Slots.Num());
-
-		for (int32 ComponentIndex = 0; ComponentIndex < AttachChildren.Num(); ++ComponentIndex)
-		{
-			if (USlateWrapperComponent* Wrapper = Cast<USlateWrapperComponent>(AttachChildren[ComponentIndex]))
-			{
-				if (Wrapper->IsRegistered())
-				{
-					WrappedComponentsInSlotOrder.Add(Wrapper);
-
-					if (int32* pSlotIndex = SlotNameToIndex.Find(Wrapper->AttachSocketName))
-					{
-						WrappedComponentsInSlotOrder[*pSlotIndex] = Wrapper;
-					}
-					else
-					{
-						FGridPanelSlot& NewSlot = *new (Slots) FGridPanelSlot();
-						EnsureSlotIsUnique(NewSlot);
-						WrappedComponentsInSlotOrder.Add(Wrapper);
-
-						Wrapper->AttachSocketName = NewSlot.SlotName;
-					}
-				}						
-			}
-		}
-
-		// Add slots to the grid panel
-		Grid->ClearChildren();
-		for (int32 SlotIndex = 0; SlotIndex < Slots.Num(); ++SlotIndex)
-		{
-			const FGridPanelSlot& SlotConfig = Slots[SlotIndex];
-
-			TSharedRef<SWidget> ContentWidget = SNullWidget::NullWidget;
-			if (WrappedComponentsInSlotOrder[SlotIndex] != NULL)
-			{
-				ContentWidget = WrappedComponentsInSlotOrder[SlotIndex]->GetWidget();
-			}
-
-			Grid->AddSlot(SlotConfig.Column, SlotConfig.Row)
-				.RowSpan(SlotConfig.RowSpan)
-				.ColumnSpan(SlotConfig.ColumnSpan)
-				.Padding(SlotConfig.Padding)
-				.HAlign(SlotConfig.HorizontalAlignment)
-				.VAlign(SlotConfig.VerticalAlignment)
-			[
-				ContentWidget
-			];
-		}
-	}
+	//TSharedPtr<SGridPanel> Grid = MyGrid.Pin();
+	//if (Grid.IsValid())
+	//{
+	//	// Create a mapping of slot names to slot indices
+	//	TMap<FName, int32> SlotNameToIndex;
+	//	for (int32 SlotIndex = 0; SlotIndex < Slots.Num(); ++SlotIndex)
+	//	{
+	//		const FGridPanelSlot& SlotConfig = Slots[SlotIndex];
+	//		SlotNameToIndex.Add(SlotConfig.SlotName, SlotIndex);
+	//	}
+	//	
+	//	// Make sure all children are assigned to a slot
+	//	TArray<USlateWrapperComponent*> WrappedComponentsInSlotOrder;
+	//	WrappedComponentsInSlotOrder.AddZeroed(Slots.Num());
+
+	//	for (int32 ComponentIndex = 0; ComponentIndex < AttachChildren.Num(); ++ComponentIndex)
+	//	{
+	//		if (USlateWrapperComponent* Wrapper = Cast<USlateWrapperComponent>(AttachChildren[ComponentIndex]))
+	//		{
+	//			if (Wrapper->IsRegistered())
+	//			{
+	//				WrappedComponentsInSlotOrder.Add(Wrapper);
+
+	//				if (int32* pSlotIndex = SlotNameToIndex.Find(Wrapper->AttachSocketName))
+	//				{
+	//					WrappedComponentsInSlotOrder[*pSlotIndex] = Wrapper;
+	//				}
+	//				else
+	//				{
+	//					FGridPanelSlot& NewSlot = *new (Slots) FGridPanelSlot();
+	//					EnsureSlotIsUnique(NewSlot);
+	//					WrappedComponentsInSlotOrder.Add(Wrapper);
+
+	//					Wrapper->AttachSocketName = NewSlot.SlotName;
+	//				}
+	//			}						
+	//		}
+	//	}
+
+	//	// Add slots to the grid panel
+	//	Grid->ClearChildren();
+	//	for (int32 SlotIndex = 0; SlotIndex < Slots.Num(); ++SlotIndex)
+	//	{
+	//		const FGridPanelSlot& SlotConfig = Slots[SlotIndex];
+
+	//		TSharedRef<SWidget> ContentWidget = SNullWidget::NullWidget;
+	//		if (WrappedComponentsInSlotOrder[SlotIndex] != NULL)
+	//		{
+	//			ContentWidget = WrappedComponentsInSlotOrder[SlotIndex]->GetWidget();
+	//		}
+
+	//		Grid->AddSlot(SlotConfig.Column, SlotConfig.Row)
+	//			.RowSpan(SlotConfig.RowSpan)
+	//			.ColumnSpan(SlotConfig.ColumnSpan)
+	//			.Padding(SlotConfig.Padding)
+	//			.HAlign(SlotConfig.HorizontalAlignment)
+	//			.VAlign(SlotConfig.VerticalAlignment)
+	//		[
+	//			ContentWidget
+	//		];
+	//	}
+	//}
 }
 
 bool UGridPanelComponent::HasAnySockets() const
diff --git a/Engine/Source/Runtime/UMG/Private/HorizontalBoxComponent.cpp b/Engine/Source/Runtime/UMG/Private/HorizontalBoxComponent.cpp
new file mode 100644
--- /dev/null
+++ b/Engine/Source/Runtime/UMG/Private/HorizontalBoxComponent.cpp
@@ -0,0 +1,74 @@
+// Copyright 1998-2014 Epic Games, Inc. All Rights Reserved.
+
+#include "UMGPrivatePCH.h"
+
+/////////////////////////////////////////////////////
+// UHorizontalBoxComponent
+
+UHorizontalBoxComponent::UHorizontalBoxComponent(const FPostConstructInitializeProperties& PCIP)
+	: Super(PCIP)
+{
+}
+
+TSharedRef<SWidget> UHorizontalBoxComponent::RebuildWidget()
+{
+	TSharedRef<SHorizontalBox> NewCanvas = SNew(SHorizontalBox);
+	MyHorizontalBox = NewCanvas;
+
+	OnKnownChildrenChanged();
+
+	return NewCanvas;
+}
+
+void UHorizontalBoxComponent::OnKnownChildrenChanged()
+{
+	TSharedPtr<SHorizontalBox> Canvas = MyHorizontalBox.Pin();
+	if (Canvas.IsValid())
+	{
+		Canvas->ClearChildren();
+
+		for (int32 SlotIndex = 0; SlotIndex < Slots.Num(); ++SlotIndex)
+		{
+			UHorizontalBoxSlot* Slot = Slots[SlotIndex];
+			if ( Slot == NULL )
+			{
+				Slots[SlotIndex] = Slot = ConstructObject<UHorizontalBoxSlot>(UHorizontalBoxSlot::StaticClass(), this);
+			}
+
+			auto NewSlot = Canvas->AddSlot()
+				.AutoWidth()
+				.HAlign(Slot->HorizontalAlignment)
+				.VAlign(Slot->VerticalAlignment)
+			[
+				Slot->Content == NULL ? SNullWidget::NullWidget : Slot->Content->GetWidget()
+			];
+		}
+	}
+}
+
+UHorizontalBoxSlot* UHorizontalBoxComponent::AddSlot(USlateWrapperComponent* Content)
+{
+	UHorizontalBoxSlot* Slot = ConstructObject<UHorizontalBoxSlot>(UHorizontalBoxSlot::StaticClass(), this);
+	Slot->Content = Content;
+	
+	Slots.Add(Slot);
+
+	return Slot;
+}
+
+#if WITH_EDITOR
+void UHorizontalBoxComponent::PostEditChangeProperty(struct FPropertyChangedEvent& PropertyChangedEvent)
+{
+	// Ensure the slots have unique names
+	int32 SlotNumbering = 1;
+
+	TSet<FName> UniqueSlotNames;
+	for (int32 SlotIndex = 0; SlotIndex < Slots.Num(); ++SlotIndex)
+	{
+		if ( Slots[SlotIndex] == NULL )
+		{
+			Slots[SlotIndex] = ConstructObject<UHorizontalBoxSlot>(UHorizontalBoxSlot::StaticClass(), this);
+		}
+	}
+}
+#endif
diff --git a/Engine/Source/Runtime/UMG/Private/HorizontalBoxSlot.cpp b/Engine/Source/Runtime/UMG/Private/HorizontalBoxSlot.cpp
new file mode 100644
--- /dev/null
+++ b/Engine/Source/Runtime/UMG/Private/HorizontalBoxSlot.cpp
@@ -0,0 +1,13 @@
+// Copyright 1998-2014 Epic Games, Inc. All Rights Reserved.
+
+#include "UMGPrivatePCH.h"
+
+/////////////////////////////////////////////////////
+// UHorizontalBoxSlot
+
+UHorizontalBoxSlot::UHorizontalBoxSlot(const FPostConstructInitializeProperties& PCIP)
+	: Super(PCIP)
+{
+	HorizontalAlignment = HAlign_Left;
+	VerticalAlignment = VAlign_Top;
+}
diff --git a/Engine/Source/Runtime/UMG/Private/ImageComponent.cpp b/Engine/Source/Runtime/UMG/Private/ImageComponent.cpp
--- a/Engine/Source/Runtime/UMG/Private/ImageComponent.cpp
+++ b/Engine/Source/Runtime/UMG/Private/ImageComponent.cpp
@@ -2,7 +2,7 @@
 
 #include "UMGPrivatePCH.h"
 
-#define LOCTEXT_NAMESPACE "SlateComponentWrappers"
+#define LOCTEXT_NAMESPACE "UMG"
 
 /////////////////////////////////////////////////////
 // UImageComponent
diff --git a/Engine/Source/Runtime/UMG/Private/ScrollBoxComponent.cpp b/Engine/Source/Runtime/UMG/Private/ScrollBoxComponent.cpp
--- a/Engine/Source/Runtime/UMG/Private/ScrollBoxComponent.cpp
+++ b/Engine/Source/Runtime/UMG/Private/ScrollBoxComponent.cpp
@@ -2,7 +2,7 @@
 
 #include "UMGPrivatePCH.h"
 
-#define LOCTEXT_NAMESPACE "SlateComponentWrappers"
+#define LOCTEXT_NAMESPACE "UMG"
 
 /////////////////////////////////////////////////////
 // UScrollBoxComponent
diff --git a/Engine/Source/Runtime/UMG/Private/SlateNonLeafWidgetComponent.cpp b/Engine/Source/Runtime/UMG/Private/SlateNonLeafWidgetComponent.cpp
--- a/Engine/Source/Runtime/UMG/Private/SlateNonLeafWidgetComponent.cpp
+++ b/Engine/Source/Runtime/UMG/Private/SlateNonLeafWidgetComponent.cpp
@@ -7,7 +7,7 @@
 #include "UObjectToken.h"
 #endif
 
-#define LOCTEXT_NAMESPACE "SlateComponentWrappers"
+#define LOCTEXT_NAMESPACE "UMG"
 
 /////////////////////////////////////////////////////
 // USlateNonLeafWidgetComponent
@@ -25,57 +25,57 @@ USlateNonLeafWidgetComponent::USlateNonLeafWidgetComponent(const FPostConstructI
 
 void USlateNonLeafWidgetComponent::TickComponent(float DeltaTime, enum ELevelTick TickType, FActorComponentTickFunction* ThisTickFunction)
 {
-	bool bNeedRebuild = false;
-	TSet<USlateWrapperComponent*> RemainingChildren;
-	for (auto SetIt = KnownChildren.CreateIterator(); SetIt; ++SetIt)
-	{
-		if (USlateWrapperComponent* Wrapper = (*SetIt).Get())
-		{
-			if (Wrapper->IsRegistered())
-			{
-				RemainingChildren.Add(Wrapper);
-			}
-		}
-	}
-
-	for (int32 ComponentIndex = 0; (ComponentIndex < AttachChildren.Num()) && !bNeedRebuild; ++ComponentIndex)
-	{
-		if (USlateWrapperComponent* Wrapper = Cast<USlateWrapperComponent>(AttachChildren[ComponentIndex]))
-		{
-			if (Wrapper->IsRegistered())
-			{
-				if (KnownChildren.Contains(Wrapper))
-				{
-					RemainingChildren.Remove(Wrapper);
-				}
-				else
-				{
-					bNeedRebuild = true;
-				}
-			}
-		}
-	}
-
-	bNeedRebuild = bNeedRebuild || (RemainingChildren.Num() > 0);
-
-	if (bNeedRebuild)
-	{
-		// Recreate the list of known children
-		KnownChildren.Empty();
-		for (int32 ComponentIndex = 0; ComponentIndex < AttachChildren.Num(); ++ComponentIndex)
-		{
-			if (USlateWrapperComponent* Wrapper = Cast<USlateWrapperComponent>(AttachChildren[ComponentIndex]))
-			{
-				if (Wrapper->IsRegistered())
-				{
-					KnownChildren.Add(Wrapper);
-				}
-			}
-		}
-
-		// Let derived classes do something with the list
-		OnKnownChildrenChanged();
-	}
+	//bool bNeedRebuild = false;
+	//TSet<USlateWrapperComponent*> RemainingChildren;
+	//for (auto SetIt = KnownChildren.CreateIterator(); SetIt; ++SetIt)
+	//{
+	//	if (USlateWrapperComponent* Wrapper = (*SetIt).Get())
+	//	{
+	//		if (Wrapper->IsRegistered())
+	//		{
+	//			RemainingChildren.Add(Wrapper);
+	//		}
+	//	}
+	//}
+
+	//for (int32 ComponentIndex = 0; (ComponentIndex < AttachChildren.Num()) && !bNeedRebuild; ++ComponentIndex)
+	//{
+	//	if (USlateWrapperComponent* Wrapper = Cast<USlateWrapperComponent>(AttachChildren[ComponentIndex]))
+	//	{
+	//		if (Wrapper->IsRegistered())
+	//		{
+	//			if (KnownChildren.Contains(Wrapper))
+	//			{
+	//				RemainingChildren.Remove(Wrapper);
+	//			}
+	//			else
+	//			{
+	//				bNeedRebuild = true;
+	//			}
+	//		}
+	//	}
+	//}
+
+	//bNeedRebuild = bNeedRebuild || (RemainingChildren.Num() > 0);
+
+	//if (bNeedRebuild)
+	//{
+	//	// Recreate the list of known children
+	//	KnownChildren.Empty();
+	//	for (int32 ComponentIndex = 0; ComponentIndex < AttachChildren.Num(); ++ComponentIndex)
+	//	{
+	//		if (USlateWrapperComponent* Wrapper = Cast<USlateWrapperComponent>(AttachChildren[ComponentIndex]))
+	//		{
+	//			if (Wrapper->IsRegistered())
+	//			{
+	//				KnownChildren.Add(Wrapper);
+	//			}
+	//		}
+	//	}
+
+	//	// Let derived classes do something with the list
+	//	OnKnownChildrenChanged();
+	//}
 }
 
 USlateWrapperComponent* USlateNonLeafWidgetComponent::GetFirstWrappedChild() const
@@ -123,12 +123,6 @@ void USlateNonLeafWidgetComponent::CheckForErrors()
 
 #endif
 
-bool USlateNonLeafWidgetComponent::CanAttachAsChild(USceneComponent* ChildComponent, FName SocketName) const
-{
-	const bool bIsSlateWidgetWrapper = ChildComponent->IsA(USlateWrapperComponent::StaticClass());
-	return bIsSlateWidgetWrapper && (bCanHaveMultipleChildren ? true : (AttachChildren.Num() == 0));
-}
-
 /////////////////////////////////////////////////////
 
 #undef LOCTEXT_NAMESPACE
diff --git a/Engine/Source/Runtime/UMG/Private/SlateWrapperComponent.cpp b/Engine/Source/Runtime/UMG/Private/SlateWrapperComponent.cpp
--- a/Engine/Source/Runtime/UMG/Private/SlateWrapperComponent.cpp
+++ b/Engine/Source/Runtime/UMG/Private/SlateWrapperComponent.cpp
@@ -12,7 +12,6 @@ USlateWrapperComponent::USlateWrapperComponent(const FPostConstructInitializePro
 
 void USlateWrapperComponent::OnRegister()
 {
-	MyWidget = RebuildWidget();
 	Super::OnRegister();
 }
 
@@ -24,14 +23,12 @@ void USlateWrapperComponent::OnUnregister()
 
 TSharedRef<SWidget> USlateWrapperComponent::GetWidget()
 {
-	if (MyWidget.IsValid())
+	if ( !MyWidget.IsValid() )
 	{
-		return MyWidget.ToSharedRef();
-	}
-	else
-	{
-		return SNullWidget::NullWidget;
+		MyWidget = RebuildWidget();
 	}
+
+	return MyWidget.ToSharedRef();
 }
 
 TSharedRef<SWidget> USlateWrapperComponent::RebuildWidget()
@@ -67,9 +64,7 @@ FSizeParam USlateWrapperComponent::ConvertSerializedSizeParamToRuntime(const FSl
 	default:
 	case ESlateSizeRule::Automatic:
 		return FAuto();
-	case ESlateSizeRule::AspectRatio:
-		return FAspectRatio();
-	case ESlateSizeRule::Stretch:
-		return FStretch(Input.StretchValue);
+	case ESlateSizeRule::Fill:
+		return FStretch(Input.Value);
 	}
 }
\ No newline at end of file
diff --git a/Engine/Source/Runtime/UMG/Private/SpacerComponent.cpp b/Engine/Source/Runtime/UMG/Private/SpacerComponent.cpp
--- a/Engine/Source/Runtime/UMG/Private/SpacerComponent.cpp
+++ b/Engine/Source/Runtime/UMG/Private/SpacerComponent.cpp
@@ -2,7 +2,7 @@
 
 #include "UMGPrivatePCH.h"
 
-#define LOCTEXT_NAMESPACE "SlateComponentWrappers"
+#define LOCTEXT_NAMESPACE "UMG"
 
 /////////////////////////////////////////////////////
 // USpacerComponent
diff --git a/Engine/Source/Runtime/UMG/Private/SpinningImageComponent.cpp b/Engine/Source/Runtime/UMG/Private/SpinningImageComponent.cpp
--- a/Engine/Source/Runtime/UMG/Private/SpinningImageComponent.cpp
+++ b/Engine/Source/Runtime/UMG/Private/SpinningImageComponent.cpp
@@ -2,7 +2,7 @@
 
 #include "UMGPrivatePCH.h"
 
-#define LOCTEXT_NAMESPACE "SlateComponentWrappers"
+#define LOCTEXT_NAMESPACE "UMG"
 
 /////////////////////////////////////////////////////
 // UImageComponent
diff --git a/Engine/Source/Runtime/UMG/Private/TextBlockComponent.cpp b/Engine/Source/Runtime/UMG/Private/TextBlockComponent.cpp
--- a/Engine/Source/Runtime/UMG/Private/TextBlockComponent.cpp
+++ b/Engine/Source/Runtime/UMG/Private/TextBlockComponent.cpp
@@ -2,7 +2,7 @@
 
 #include "UMGPrivatePCH.h"
 
-#define LOCTEXT_NAMESPACE "SlateComponentWrappers"
+#define LOCTEXT_NAMESPACE "UMG"
 
 /////////////////////////////////////////////////////
 // UTextBlockComponent
diff --git a/Engine/Source/Runtime/UMG/Private/ThrobberComponent.cpp b/Engine/Source/Runtime/UMG/Private/ThrobberComponent.cpp
--- a/Engine/Source/Runtime/UMG/Private/ThrobberComponent.cpp
+++ b/Engine/Source/Runtime/UMG/Private/ThrobberComponent.cpp
@@ -2,7 +2,7 @@
 
 #include "UMGPrivatePCH.h"
 
-#define LOCTEXT_NAMESPACE "SlateComponentWrappers"
+#define LOCTEXT_NAMESPACE "UMG"
 
 /////////////////////////////////////////////////////
 // UThrobberComponent
diff --git a/Engine/Source/Runtime/UMG/Private/UMGModule.cpp b/Engine/Source/Runtime/UMG/Private/UMGModule.cpp
--- a/Engine/Source/Runtime/UMG/Private/UMGModule.cpp
+++ b/Engine/Source/Runtime/UMG/Private/UMGModule.cpp
@@ -18,7 +18,6 @@ public:
 	/** Called right after the module DLL has been loaded and the module object has been created */
 	virtual void StartupModule() OVERRIDE
 	{
-		printf("ALIVE");
 	}
 
 	/** Called before the module is unloaded, right before the module object is destroyed. */
diff --git a/Engine/Source/Runtime/UMG/Private/UserWidget.cpp b/Engine/Source/Runtime/UMG/Private/UserWidget.cpp
--- a/Engine/Source/Runtime/UMG/Private/UserWidget.cpp
+++ b/Engine/Source/Runtime/UMG/Private/UserWidget.cpp
@@ -44,15 +44,17 @@ void AUserWidget::RebuildWrapperWidget()
 	
 	TArray<USlateWrapperComponent*> SlateWrapperComponents;
 	GetComponents(SlateWrapperComponents);
+
 	// Place all of our top-level children Slate wrapped components into the overlay
 	for (int32 ComponentIndex = 0; ComponentIndex < SlateWrapperComponents.Num(); ++ComponentIndex)
 	{
-		if (SlateWrapperComponents[ComponentIndex]->IsRegistered() && (SlateWrapperComponents[ComponentIndex]->AttachParent == NULL))
+		if (SlateWrapperComponents[ComponentIndex]->IsRegistered())
 		{
 			MyWrapperWidget->AddSlot()
 			[
 				SlateWrapperComponents[ComponentIndex]->GetWidget()
 			];
+			break;
 		}
 	}
 	
diff --git a/Engine/Source/Runtime/UMG/Private/VerticalBoxComponent.cpp b/Engine/Source/Runtime/UMG/Private/VerticalBoxComponent.cpp
new file mode 100644
--- /dev/null
+++ b/Engine/Source/Runtime/UMG/Private/VerticalBoxComponent.cpp
@@ -0,0 +1,74 @@
+// Copyright 1998-2014 Epic Games, Inc. All Rights Reserved.
+
+#include "UMGPrivatePCH.h"
+
+/////////////////////////////////////////////////////
+// UVerticalBoxComponent
+
+UVerticalBoxComponent::UVerticalBoxComponent(const FPostConstructInitializeProperties& PCIP)
+	: Super(PCIP)
+{
+}
+
+TSharedRef<SWidget> UVerticalBoxComponent::RebuildWidget()
+{
+	TSharedRef<SVerticalBox> NewCanvas = SNew(SVerticalBox);
+	MyVerticalBox = NewCanvas;
+
+	OnKnownChildrenChanged();
+
+	return NewCanvas;
+}
+
+void UVerticalBoxComponent::OnKnownChildrenChanged()
+{
+	TSharedPtr<SVerticalBox> Canvas = MyVerticalBox.Pin();
+	if (Canvas.IsValid())
+	{
+		Canvas->ClearChildren();
+
+		for (int32 SlotIndex = 0; SlotIndex < Slots.Num(); ++SlotIndex)
+		{
+			UVerticalBoxSlot* Slot = Slots[SlotIndex];
+			if ( Slot == NULL )
+			{
+				Slots[SlotIndex] = Slot = ConstructObject<UVerticalBoxSlot>(UVerticalBoxSlot::StaticClass(), this);
+			}
+
+			auto NewSlot = Canvas->AddSlot()
+				.AutoHeight()
+				.HAlign(Slot->HorizontalAlignment)
+				.VAlign(Slot->VerticalAlignment)
+			[
+				Slot->Content == NULL ? SNullWidget::NullWidget : Slot->Content->GetWidget()
+			];
+		}
+	}
+}
+
+UVerticalBoxSlot* UVerticalBoxComponent::AddSlot(USlateWrapperComponent* Content)
+{
+	UVerticalBoxSlot* Slot = ConstructObject<UVerticalBoxSlot>(UVerticalBoxSlot::StaticClass(), this);
+	Slot->Content = Content;
+	
+	Slots.Add(Slot);
+
+	return Slot;
+}
+
+#if WITH_EDITOR
+void UVerticalBoxComponent::PostEditChangeProperty(struct FPropertyChangedEvent& PropertyChangedEvent)
+{
+	// Ensure the slots have unique names
+	int32 SlotNumbering = 1;
+
+	TSet<FName> UniqueSlotNames;
+	for (int32 SlotIndex = 0; SlotIndex < Slots.Num(); ++SlotIndex)
+	{
+		if ( Slots[SlotIndex] == NULL )
+		{
+			Slots[SlotIndex] = ConstructObject<UVerticalBoxSlot>(UVerticalBoxSlot::StaticClass(), this);
+		}
+	}
+}
+#endif
\ No newline at end of file
diff --git a/Engine/Source/Runtime/UMG/Private/VerticalBoxSlot.cpp b/Engine/Source/Runtime/UMG/Private/VerticalBoxSlot.cpp
new file mode 100644
--- /dev/null
+++ b/Engine/Source/Runtime/UMG/Private/VerticalBoxSlot.cpp
@@ -0,0 +1,13 @@
+// Copyright 1998-2014 Epic Games, Inc. All Rights Reserved.
+
+#include "UMGPrivatePCH.h"
+
+/////////////////////////////////////////////////////
+// UVerticalBoxSlot
+
+UVerticalBoxSlot::UVerticalBoxSlot(const FPostConstructInitializeProperties& PCIP)
+	: Super(PCIP)
+{
+	HorizontalAlignment = HAlign_Left;
+	VerticalAlignment = VAlign_Top;
+}
diff --git a/Engine/Source/Runtime/UMG/Private/WidgetBlueprint.cpp b/Engine/Source/Runtime/UMG/Private/WidgetBlueprint.cpp
--- a/Engine/Source/Runtime/UMG/Private/WidgetBlueprint.cpp
+++ b/Engine/Source/Runtime/UMG/Private/WidgetBlueprint.cpp
@@ -2,7 +2,7 @@
 
 #include "UMGPrivatePCH.h"
 
-#define LOCTEXT_NAMESPACE "SlateComponentWrappers"
+#define LOCTEXT_NAMESPACE "UMG"
 
 /////////////////////////////////////////////////////
 // UWidgetBlueprint
@@ -10,5 +10,29 @@
 UWidgetBlueprint::UWidgetBlueprint(const FPostConstructInitializeProperties& PCIP)
 	: Super(PCIP)
 {
-	printf("TEST");
+}
+
+bool UWidgetBlueprint::ValidateGeneratedClass(const UClass* InClass)
+{
+	bool Result = Super::ValidateGeneratedClass(InClass);
+
+	//for ( auto CompIt = Blueprint->Timelines.CreateConstIterator(); CompIt; ++CompIt )
+	//{
+	//	const UTimelineTemplate* Template = ( *CompIt );
+	//	if ( !ensure(Template && ( Template->GetOuter() == GeneratedClass )) )
+	//	{
+	//		return false;
+	//	}
+	//}
+
+	//for ( auto CompIt = GeneratedClass->Timelines.CreateConstIterator(); CompIt; ++CompIt )
+	//{
+	//	const UTimelineTemplate* Template = ( *CompIt );
+	//	if ( !ensure(Template && ( Template->GetOuter() == GeneratedClass )) )
+	//	{
+	//		return false;
+	//	}
+	//}
+
+	return Result;
 }
\ No newline at end of file
diff --git a/Engine/Source/Runtime/UMG/Private/WidgetBlueprintGeneratedClass.cpp b/Engine/Source/Runtime/UMG/Private/WidgetBlueprintGeneratedClass.cpp
--- a/Engine/Source/Runtime/UMG/Private/WidgetBlueprintGeneratedClass.cpp
+++ b/Engine/Source/Runtime/UMG/Private/WidgetBlueprintGeneratedClass.cpp
@@ -2,7 +2,7 @@
 
 #include "UMGPrivatePCH.h"
 
-#define LOCTEXT_NAMESPACE "SlateComponentWrappers"
+#define LOCTEXT_NAMESPACE "UMG"
 
 /////////////////////////////////////////////////////
 // UWidgetBlueprintGeneratedClass
@@ -10,4 +10,66 @@
 UWidgetBlueprintGeneratedClass::UWidgetBlueprintGeneratedClass(const FPostConstructInitializeProperties& PCIP)
 	: Super(PCIP)
 {
+}
+
+void UWidgetBlueprintGeneratedClass::CreateComponentsForActor(AActor* Actor) const
+{
+	Super::CreateComponentsForActor(Actor);
+
+	// Duplicate the graph, keeping track of what was duplicated
+	TMap<UObject*, UObject*> DuplicatedObjectList;
+
+	TArray<USlateWrapperComponent*> ClonedWidgets;
+
+	// Iterate over each timeline template
+	for ( const USlateWrapperComponent* Template : WidgetTemplates )
+	{
+		FObjectDuplicationParameters Parameters(const_cast<USlateWrapperComponent*>( Template ), Actor);
+		Parameters.CreatedObjects = &DuplicatedObjectList;
+
+		USlateWrapperComponent* ClonedWidget = CastChecked<USlateWrapperComponent>(StaticDuplicateObjectEx(Parameters));
+		ClonedWidgets.Add(ClonedWidget);
+	}
+
+	// Replace references to old classes and instances on this object with the corresponding new ones
+	for ( USlateWrapperComponent* Widget : ClonedWidgets )
+	{
+		FArchiveReplaceObjectRef<UObject> ReplaceInCDOAr(Widget, DuplicatedObjectList, /*bNullPrivateRefs=*/ false, /*bIgnoreOuterRef=*/ false, /*bIgnoreArchetypeRef=*/ false);
+	}
+
+	// Iterate over each timeline template
+	for ( USlateWrapperComponent* Widget : ClonedWidgets )
+	{
+		// Not fatal if NULL, but shouldn't happen
+		if ( !ensure(Widget != NULL) )
+		{
+			continue;
+		}
+
+		FName NewName = *( FString::Printf(TEXT("WidgetComponent__%d"), Actor->SerializedComponents.Num()) );
+
+		////Serialize object properties using write/read operations.
+		//TArray<uint8> SavedProperties;
+		//FObjectWriter Writer(ComponentTemplate, SavedProperties);
+		//FObjectReader(ClonedComponent, SavedProperties);
+
+
+
+		//USlateWrapperComponent* Widget = ConstructObject<USlateWrapperComponent>(Template->GetClass(), Actor, NewName);
+		Widget->bCreatedByConstructionScript = true; // Indicate it comes from a blueprint so it gets cleared when we rerun construction scripts
+		Actor->SerializedComponents.Add(Widget); // Add to array so it gets saved
+		Widget->SetNetAddressable();	// This component has a stable name that can be referenced for replication
+
+		// Find property with the same name as the template and assign the new Timeline to it
+		UClass* ActorClass = Actor->GetClass();
+		FName thing = Widget->GetClass()->GetFName();
+		//UTimelineTemplate::TimelineTemplateNameToVariableName(Template->GetFName())
+		UObjectPropertyBase* Prop = FindField<UObjectPropertyBase>(ActorClass, thing);
+		if ( Prop )
+		{
+			Prop->SetObjectPropertyValue_InContainer(Actor, Widget);
+		}
+
+		Widget->RegisterComponent();
+	}
 }
\ No newline at end of file
