* Checkpionting UMG work, getting a basic treeview going with the browsable nodes.
* Several changes to containers, adding ways to traverse the hierarchy.
* Temporary fix for not being able to reference variables due to a variable name generation change.

diff --git a/Engine/Source/Editor/UMGEditor/Private/SUMGEditorTree.cpp b/Engine/Source/Editor/UMGEditor/Private/SUMGEditorTree.cpp
--- a/Engine/Source/Editor/UMGEditor/Private/SUMGEditorTree.cpp
+++ b/Engine/Source/Editor/UMGEditor/Private/SUMGEditorTree.cpp
@@ -11,15 +11,16 @@
 #include "SceneViewport.h"
 
 #include "BlueprintEditor.h"
+#include "SKismetInspector.h"
 
 void SUMGEditorTree::Construct(const FArguments& InArgs, TSharedPtr<FBlueprintEditor> InBlueprintEditor, USimpleConstructionScript* InSCS)
 {
 	BlueprintEditor = InBlueprintEditor;
 
-	UBlueprint* BP = InBlueprintEditor->GetBlueprintObj();
-	//BP->ComponentTemplates
+	UWidgetBlueprint* Blueprint = GetBlueprint();
+	RootWidgets.Add(Blueprint->WidgetTemplates[0]);
 
-	//FCoreDelegates::OnObjectPropertyChanged.Add( FCoreDelegates::FOnObjectPropertyChanged::FDelegate::CreateRaw(this, &SUMGEditorTree::OnObjectPropertyChanged) );
+	FCoreDelegates::OnObjectPropertyChanged.Add( FCoreDelegates::FOnObjectPropertyChanged::FDelegate::CreateRaw(this, &SUMGEditorTree::OnObjectPropertyChanged) );
 
 	ChildSlot
 	[
@@ -38,19 +39,94 @@ void SUMGEditorTree::Construct(const FArguments& InArgs, TSharedPtr<FBlueprintEd
 		+ SVerticalBox::Slot()
 		.FillHeight(1.0f)
 		[
-			SAssignNew(PageTreeView, STreeView<TSharedPtr<FString>>)
+			SAssignNew(WidgetTreeView, STreeView< USlateWrapperComponent* >)
 			.ItemHeight(20.0f)
-			//.OnGenerateRow(this, &SDeviceProcesses::HandleProcessTreeViewGenerateRow)
-			//.OnGetChildren(this, &SDeviceProcesses::HandleProcessTreeViewGetChildren)
-			//.SelectionMode(ESelectionMode::Multi)
-			//.TreeItemsSource(&ProcessList)
+			.SelectionMode(ESelectionMode::Single)
+			.OnGetChildren(this, &SUMGEditorTree::WidgetHierarchy_OnGetChildren)
+			.OnGenerateRow(this, &SUMGEditorTree::WidgetHierarchy_OnGenerateRow)
+			.OnSelectionChanged(this, &SUMGEditorTree::WidgetHierarchy_OnSelectionChanged)
+			.TreeItemsSource(&RootWidgets)
 		]
 	];
 }
 
 SUMGEditorTree::~SUMGEditorTree()
 {
-	//FCoreDelegates::OnObjectPropertyChanged.Remove( FCoreDelegates::FOnObjectPropertyChanged::FDelegate::CreateRaw(this, &SUMGEditorTree::OnObjectPropertyChanged) );
+	FCoreDelegates::OnObjectPropertyChanged.Remove( FCoreDelegates::FOnObjectPropertyChanged::FDelegate::CreateRaw(this, &SUMGEditorTree::OnObjectPropertyChanged) );
+}
+
+UWidgetBlueprint* SUMGEditorTree::GetBlueprint() const
+{
+	if ( BlueprintEditor.IsValid() )
+	{
+		UBlueprint* BP = BlueprintEditor.Pin()->GetBlueprintObj();
+		return Cast<UWidgetBlueprint>(BP);
+	}
+
+	return NULL;
+}
+
+
+void SUMGEditorTree::OnObjectPropertyChanged(UObject* ObjectBeingModified)
+{
+	if ( !ensure(ObjectBeingModified) )
+	{
+		return;
+	}
+}
+
+void SUMGEditorTree::ShowDetailsForObjects(TArray<USlateWrapperComponent*> Widgets)
+{
+	// Convert the selection set to an array of UObject* pointers
+	FString InspectorTitle;
+	TArray<UObject*> InspectorObjects;
+	InspectorObjects.Empty(Widgets.Num());
+	for ( USlateWrapperComponent* Widget : Widgets )
+	{
+		//if ( NodePtr->CanEditDefaults() )
+		{
+			InspectorTitle = "Widget";// Widget->GetDisplayString();
+			InspectorObjects.Add(Widget);
+		}
+	}
+
+	// Update the details panel
+	SKismetInspector::FShowDetailsOptions Options(InspectorTitle, true);
+	BlueprintEditor.Pin()->GetInspector()->ShowDetailsForObjects(InspectorObjects, Options);
+}
+
+void SUMGEditorTree::WidgetHierarchy_OnGetChildren(USlateWrapperComponent* InParent, TArray< USlateWrapperComponent* >& OutChildren)
+{
+	USlateNonLeafWidgetComponent* Widget = Cast<USlateNonLeafWidgetComponent>(InParent);
+	if ( Widget )
+	{
+		for ( int32 i = 0; i < Widget->GetChildrenCount(); i++ )
+		{
+			OutChildren.Add( Widget->GetChildAt(i) );
+		}
+	}
+}
+
+TSharedRef< ITableRow > SUMGEditorTree::WidgetHierarchy_OnGenerateRow(USlateWrapperComponent* InItem, const TSharedRef<STableViewBase>& OwnerTable)
+{
+	return
+		SNew(STableRow< USlateWrapperComponent* >, OwnerTable)
+		.Padding(2.0f)
+//		.OnDragDetected(this, &SUMGEditorWidgetTemplates::OnDraggingWidgetTemplateItem)
+		[
+			SNew(STextBlock)
+			.Text(InItem->GetFName().ToString())
+		];
+}
+
+void SUMGEditorTree::WidgetHierarchy_OnSelectionChanged(USlateWrapperComponent* SelectedItem, ESelectInfo::Type SelectInfo)
+{
+	if ( SelectInfo != ESelectInfo::Direct )
+	{
+		TArray<USlateWrapperComponent*> Items;
+		Items.Add(SelectedItem);
+		ShowDetailsForObjects(Items);
+	}
 }
 
 FReply SUMGEditorTree::CreateTestUI()
diff --git a/Engine/Source/Editor/UMGEditor/Private/SUMGEditorTree.h b/Engine/Source/Editor/UMGEditor/Private/SUMGEditorTree.h
--- a/Engine/Source/Editor/UMGEditor/Private/SUMGEditorTree.h
+++ b/Engine/Source/Editor/UMGEditor/Private/SUMGEditorTree.h
@@ -14,7 +14,18 @@ public:
 	SLATE_END_ARGS()
 
 	void Construct(const FArguments& InArgs, TSharedPtr<FBlueprintEditor> InBlueprintEditor, USimpleConstructionScript* InSCS);
-	~SUMGEditorTree();
+	virtual ~SUMGEditorTree();
+
+private:
+	UWidgetBlueprint* GetBlueprint() const;
+
+	void OnObjectPropertyChanged(UObject* ObjectBeingModified);
+
+	void ShowDetailsForObjects(TArray<USlateWrapperComponent*> Widgets);
+
+	void WidgetHierarchy_OnGetChildren(USlateWrapperComponent* InParent, TArray< USlateWrapperComponent* >& OutChildren);
+	TSharedRef< ITableRow > WidgetHierarchy_OnGenerateRow(USlateWrapperComponent* InItem, const TSharedRef<STableViewBase>& OwnerTable);
+	void WidgetHierarchy_OnSelectionChanged(USlateWrapperComponent* SelectedItem, ESelectInfo::Type SelectInfo);
 
 	FReply CreateTestUI();
 	
@@ -22,5 +33,7 @@ private:
 
 	TWeakPtr<class FBlueprintEditor> BlueprintEditor;
 
-	TSharedPtr< STreeView< TSharedPtr< FString > > > PageTreeView;
+	TArray< USlateWrapperComponent* > RootWidgets;
+
+	TSharedPtr< STreeView< USlateWrapperComponent* > > WidgetTreeView;
 };
diff --git a/Engine/Source/Editor/UMGEditor/Private/SUMGEditorWidgetTemplates.cpp b/Engine/Source/Editor/UMGEditor/Private/SUMGEditorWidgetTemplates.cpp
new file mode 100644
--- /dev/null
+++ b/Engine/Source/Editor/UMGEditor/Private/SUMGEditorWidgetTemplates.cpp
@@ -0,0 +1,156 @@
+// Copyright 1998-2014 Epic Games, Inc. All Rights Reserved.
+
+#include "UMGEditorPrivatePCH.h"
+
+#include "SUMGEditorWidgetTemplates.h"
+#include "UMGEditor.h"
+#include "UMGEditorViewportClient.h"
+#include "UMGEditorActions.h"
+
+#include "PreviewScene.h"
+#include "SceneViewport.h"
+
+#include "BlueprintEditor.h"
+
+#include "DecoratedDragDropOp.h"
+
+class FWidgetTemplateDragDropOp : public FDecoratedDragDropOp
+{
+public:
+	DRAG_DROP_OPERATOR_TYPE(FWidgetTemplateDragDropOp, FDecoratedDragDropOp)
+
+	TSharedPtr<FWidgetTemplateDescriptor> Template;
+
+	void Init(const TSharedPtr<FWidgetTemplateDescriptor>& InTemplate)
+	{
+		Template = InTemplate;
+	}
+};
+
+
+
+class SWidgetTemplateItem : public SCompoundWidget
+{
+public:
+
+	SLATE_BEGIN_ARGS(SWidgetTemplateItem) {}
+	SLATE_END_ARGS()
+
+	/**
+	* Constructs this widget
+	*
+	* @param InArgs    Declaration from which to construct the widget
+	*/
+	void Construct(const FArguments& InArgs, TSharedPtr<FWidgetTemplateDescriptor> InTemplate)
+	{
+		Template = InTemplate;
+
+		ChildSlot
+			[
+				SNew(SHorizontalBox)
+
+				+ SHorizontalBox::Slot()
+				.AutoWidth()
+				.Padding(2.0f)
+				[
+					SNew(SImage)
+					.Image(Template->Icon.GetIcon())
+				]
+				
+				+ SHorizontalBox::Slot()
+				.AutoWidth()
+				.Padding(2.0f)
+				[
+					SNew(STextBlock)
+					.Text(Template->Name)
+				]
+			];
+	}
+
+private:
+	TSharedPtr<FWidgetTemplateDescriptor> Template;
+};
+
+void SUMGEditorWidgetTemplates::Construct(const FArguments& InArgs, TSharedPtr<FBlueprintEditor> InBlueprintEditor, USimpleConstructionScript* InSCS)
+{
+	BlueprintEditor = InBlueprintEditor;
+
+	UBlueprint* BP = InBlueprintEditor->GetBlueprintObj();
+
+	BuildWidgetList();
+
+	//FCoreDelegates::OnObjectPropertyChanged.Add( FCoreDelegates::FOnObjectPropertyChanged::FDelegate::CreateRaw(this, &SUMGEditorWidgetTemplates::OnObjectPropertyChanged) );
+
+	ChildSlot
+	[
+		SNew(SVerticalBox)
+
+		+ SVerticalBox::Slot()
+		.FillHeight(1.0f)
+		[
+			SAssignNew(WidgetsListView, SListView< TSharedPtr<FWidgetTemplateDescriptor> >)
+			.OnGenerateRow(this, &SUMGEditorWidgetTemplates::OnGenerateWidgetTemplateItem)
+			.ListItemsSource(&WidgetTemplates)
+			.SelectionMode(ESelectionMode::Single)
+		]
+	];
+}
+
+SUMGEditorWidgetTemplates::~SUMGEditorWidgetTemplates()
+{
+	//FCoreDelegates::OnObjectPropertyChanged.Remove( FCoreDelegates::FOnObjectPropertyChanged::FDelegate::CreateRaw(this, &SUMGEditorWidgetTemplates::OnObjectPropertyChanged) );
+}
+
+UWidgetBlueprint* SUMGEditorWidgetTemplates::GetBlueprint() const
+{
+	if ( BlueprintEditor.IsValid() )
+	{
+		UBlueprint* BP = BlueprintEditor.Pin()->GetBlueprintObj();
+		return Cast<UWidgetBlueprint>(BP);
+	}
+
+	return NULL;
+}
+
+void SUMGEditorWidgetTemplates::BuildWidgetList()
+{
+	WidgetTemplates.Reset();
+
+	for ( TObjectIterator<UClass> ClassIt; ClassIt; ++ClassIt )
+	{
+		UClass* WidgetClass = *ClassIt;
+		if ( WidgetClass->IsChildOf(USlateWrapperComponent::StaticClass()) )
+		{
+			TSharedPtr<FWidgetTemplateDescriptor> Template = MakeShareable(new FWidgetTemplateDescriptor());
+			Template->Name = WidgetClass->GetDisplayNameText();
+			Template->ToolTip = WidgetClass->GetDisplayNameText();
+			Template->WidgetClass = WidgetClass;
+
+			WidgetTemplates.Add(Template);
+		}
+	}
+}
+
+TSharedRef<ITableRow> SUMGEditorWidgetTemplates::OnGenerateWidgetTemplateItem(TSharedPtr<FWidgetTemplateDescriptor> Item, const TSharedRef<STableViewBase>& OwnerTable)
+{
+	return
+		SNew(STableRow< TSharedPtr<FWidgetTemplateDescriptor> >, OwnerTable)
+		.Padding(2.0f)
+		.OnDragDetected(this, &SUMGEditorWidgetTemplates::OnDraggingWidgetTemplateItem)
+		[
+			SNew(SWidgetTemplateItem, Item)
+		];
+}
+
+FReply SUMGEditorWidgetTemplates::OnDraggingWidgetTemplateItem(const FGeometry& MyGeometry, const FPointerEvent& MouseEvent)
+{
+	FReply Reply = FReply::Unhandled();
+
+	if ( WidgetsListView->GetNumItemsSelected() > 0 )
+	{
+		TSharedPtr<FWidgetTemplateDescriptor> Template = WidgetsListView->GetSelectedItems()[0];
+//		Reply = FReply::Handled().BeginDragDrop(FWidgetTemplateDragDropOp::New(Template));
+	}
+
+	return Reply;
+}
diff --git a/Engine/Source/Editor/UMGEditor/Private/SUMGEditorWidgetTemplates.h b/Engine/Source/Editor/UMGEditor/Private/SUMGEditorWidgetTemplates.h
new file mode 100644
--- /dev/null
+++ b/Engine/Source/Editor/UMGEditor/Private/SUMGEditorWidgetTemplates.h
@@ -0,0 +1,48 @@
+// Copyright 1998-2014 Epic Games, Inc. All Rights Reserved.
+
+#pragma once
+
+#include "SCompoundWidget.h"
+#include "IUMGEditor.h"
+#include "BlueprintEditor.h"
+
+struct FWidgetTemplateDescriptor
+{
+	FText Name;
+	
+	FText ToolTip;
+
+	FSlateIcon Icon;
+
+	TSubclassOf<USlateWrapperComponent> WidgetClass;
+	
+	FWidgetTemplateDescriptor()
+	{
+	}
+};
+
+class SUMGEditorWidgetTemplates : public SCompoundWidget
+{
+public:
+	SLATE_BEGIN_ARGS( SUMGEditorWidgetTemplates ){}
+	SLATE_END_ARGS()
+
+	void Construct(const FArguments& InArgs, TSharedPtr<FBlueprintEditor> InBlueprintEditor, USimpleConstructionScript* InSCS);
+	virtual ~SUMGEditorWidgetTemplates();
+
+private:
+	UWidgetBlueprint* GetBlueprint() const;
+
+	void BuildWidgetList();
+
+	TSharedRef<ITableRow> OnGenerateWidgetTemplateItem(TSharedPtr<FWidgetTemplateDescriptor> Item, const TSharedRef<STableViewBase>& OwnerTable);
+
+	FReply OnDraggingWidgetTemplateItem(const FGeometry& MyGeometry, const FPointerEvent& MouseEvent);
+
+private:
+
+	TWeakPtr<class FBlueprintEditor> BlueprintEditor;
+
+	TArray< TSharedPtr<FWidgetTemplateDescriptor> > WidgetTemplates;
+	TSharedPtr< SListView< TSharedPtr<FWidgetTemplateDescriptor> > > WidgetsListView;
+};
diff --git a/Engine/Source/Editor/UMGEditor/Private/UMGBlueprintEditorExtensionHook.cpp b/Engine/Source/Editor/UMGEditor/Private/UMGBlueprintEditorExtensionHook.cpp
--- a/Engine/Source/Editor/UMGEditor/Private/UMGBlueprintEditorExtensionHook.cpp
+++ b/Engine/Source/Editor/UMGEditor/Private/UMGBlueprintEditorExtensionHook.cpp
@@ -280,6 +280,37 @@ public:
 	}
 };
 
+
+/////////////////////////////////////////////////////
+// FWidgetTemplatesSummoner
+
+static const FName UMGWidgetTemplatesTabID(TEXT("WidgetTemplates"));
+
+struct FWidgetTemplatesSummoner : public FWorkflowTabFactory
+{
+protected:
+	TWeakPtr<class FBlueprintEditor> BlueprintEditor;
+public:
+	FWidgetTemplatesSummoner(TSharedPtr<class FBlueprintEditor> InBlueprintEditor)
+		: FWorkflowTabFactory(UMGWidgetTemplatesTabID, InBlueprintEditor)
+		, BlueprintEditor(InBlueprintEditor)
+	{
+		TabLabel = LOCTEXT("WidgetTemplatesTabLabel", "Templates");
+		//TabIcon = FEditorStyle::GetBrush("Kismet.Tabs.Components");
+
+		bIsSingleton = true;
+
+		ViewMenuDescription = LOCTEXT("WidgetTemplates_ViewMenu_Desc", "Templates");
+		ViewMenuTooltip = LOCTEXT("WidgetTemplates_ViewMenu_ToolTip", "Show the Templates");
+	}
+
+	virtual TSharedRef<SWidget> CreateTabBody(const FWorkflowTabSpawnInfo& Info) const OVERRIDE
+	{
+		return SNew(SUMGEditorWidgetTemplates, BlueprintEditor.Pin(), BlueprintEditor.Pin()->GetBlueprintObj()->SimpleConstructionScript);
+	}
+};
+
+
 //////////////////////////////////////////////////////////////////////////
 // FComponentsEditorModeOverride
 
@@ -295,9 +326,11 @@ public:
 		//BlueprintComponentsTabFactories.RegisterFactory(MakeShareable(new FDefaultsEditorSummoner(InBlueprintEditor)));
 		//BlueprintComponentsTabFactories.RegisterFactory(MakeShareable(new FFindResultsSummoner(InBlueprintEditor)));
 
+		//BlueprintComponentsTabFactories.UnregisterFactory(FBlueprintEditorTabs::ConstructionScriptEditorID);
 		BlueprintComponentsTabFactories.UnregisterFactory(FBlueprintEditorTabs::SCSViewportID);
 		BlueprintComponentsTabFactories.RegisterFactory(MakeShareable(new FSlatePreviewSummoner(InBlueprintEditor)));
 		BlueprintComponentsTabFactories.RegisterFactory(MakeShareable(new FSlateTreeSummoner(InBlueprintEditor)));
+		BlueprintComponentsTabFactories.RegisterFactory(MakeShareable(new FWidgetTemplatesSummoner(InBlueprintEditor)));
 
 		TabLayout = FTabManager::NewLayout( "Standalone_UMGEditor_Layout_v2" )
 		->AddArea
@@ -324,7 +357,7 @@ public:
 					(
 						FTabManager::NewStack()
 						->SetSizeCoefficient( 0.5f )
-						->AddTab( FBlueprintEditorTabs::ConstructionScriptEditorID, ETabState::OpenedTab )
+						->AddTab( UMGWidgetTemplatesTabID, ETabState::OpenedTab )
 					)
 				)
 				->Split
@@ -359,14 +392,27 @@ public:
 				)
 			)
 		);
+	}
+
+	virtual void PreDeactivateMode() OVERRIDE
+	{
+		FBlueprintComponentsApplicationMode::PreDeactivateMode();
 
+		GetBlueprintEditor()->GetInspector()->EnableComponentDetailsCustomization(false);
 	}
 
-	UBlueprint* GetBlueprint() const
+	virtual void PostActivateMode() OVERRIDE
+	{
+		FBlueprintComponentsApplicationMode::PostActivateMode();
+
+		GetBlueprintEditor()->GetInspector()->EnableComponentDetailsCustomization(false);
+	}
+
+	UWidgetBlueprint* GetBlueprint() const
 	{
 		if (FBlueprintEditor* Editor = MyBlueprintEditor.Pin().Get())
 		{
-			return Editor->GetBlueprintObj();
+			return Cast<UWidgetBlueprint>(Editor->GetBlueprintObj());
 		}
 		else
 		{
@@ -378,13 +424,6 @@ public:
 	{
 		return MyBlueprintEditor.Pin();
 	}
-
-	void Foo()
-	{
-		if (AActor* PreviewActor = GetPreviewActor())
-		{
-		}
-	}
 };
 
 //////////////////////////////////////////////////////////////////////////
diff --git a/Engine/Source/Editor/UMGEditor/Private/WidgetBlueprintEditor.cpp b/Engine/Source/Editor/UMGEditor/Private/WidgetBlueprintEditor.cpp
new file mode 100644
--- /dev/null
+++ b/Engine/Source/Editor/UMGEditor/Private/WidgetBlueprintEditor.cpp
@@ -0,0 +1,17 @@
+// Copyright 1998-2014 Epic Games, Inc. All Rights Reserved.
+
+#include "UMGEditorPrivatePCH.h"
+
+#include "WidgetBlueprintEditor.h"
+
+#define LOCTEXT_NAMESPACE "UMG"
+
+FWidgetBlueprintEditor::FWidgetBlueprintEditor()
+{
+}
+
+FWidgetBlueprintEditor::~FWidgetBlueprintEditor()
+{
+}
+
+#undef LOCTEXT_NAMESPACE
diff --git a/Engine/Source/Editor/UMGEditor/Private/WidgetBlueprintEditor.h b/Engine/Source/Editor/UMGEditor/Private/WidgetBlueprintEditor.h
new file mode 100644
--- /dev/null
+++ b/Engine/Source/Editor/UMGEditor/Private/WidgetBlueprintEditor.h
@@ -0,0 +1,16 @@
+// Copyright 1998-2014 Epic Games, Inc. All Rights Reserved.
+
+#pragma once
+
+#include "Editor/Kismet/Public/BlueprintEditor.h"
+
+/**
+ * widget blueprint editor (extends Blueprint editor)
+ */
+class FWidgetBlueprintEditor : public FBlueprintEditor
+{
+public:
+	FWidgetBlueprintEditor();
+
+	virtual ~FWidgetBlueprintEditor();
+};
diff --git a/Engine/Source/Runtime/UMG/Classes/CanvasPanelComponent.h b/Engine/Source/Runtime/UMG/Classes/CanvasPanelComponent.h
--- a/Engine/Source/Runtime/UMG/Classes/CanvasPanelComponent.h
+++ b/Engine/Source/Runtime/UMG/Classes/CanvasPanelComponent.h
@@ -19,6 +19,11 @@ class UMG_API UCanvasPanelComponent : public USlateNonLeafWidgetComponent
 
 	UCanvasPanelSlot* AddSlot(USlateWrapperComponent* Content);
 
+	// USlateNonLeafWidgetComponent
+	virtual int32 GetChildrenCount() const OVERRIDE;
+	virtual USlateWrapperComponent* GetChildAt(int32 Index) const OVERRIDE;
+	// End USlateNonLeafWidgetComponent
+
 #if WITH_EDITOR
 	// UObject interface
 	virtual void PostEditChangeProperty(struct FPropertyChangedEvent& PropertyChangedEvent) OVERRIDE;
diff --git a/Engine/Source/Runtime/UMG/Classes/CanvasPanelSlot.h b/Engine/Source/Runtime/UMG/Classes/CanvasPanelSlot.h
--- a/Engine/Source/Runtime/UMG/Classes/CanvasPanelSlot.h
+++ b/Engine/Source/Runtime/UMG/Classes/CanvasPanelSlot.h
@@ -8,9 +8,6 @@ UCLASS()
 class UMG_API UCanvasPanelSlot : public UPanelSlot
 {
 	GENERATED_UCLASS_BODY()
-	
-	UPROPERTY()
-	class USlateWrapperComponent* Content;
 
 	/** Position. */
 	UPROPERTY(EditAnywhere, Category=Appearance)
diff --git a/Engine/Source/Runtime/UMG/Classes/HorizontalBoxComponent.h b/Engine/Source/Runtime/UMG/Classes/HorizontalBoxComponent.h
--- a/Engine/Source/Runtime/UMG/Classes/HorizontalBoxComponent.h
+++ b/Engine/Source/Runtime/UMG/Classes/HorizontalBoxComponent.h
@@ -15,6 +15,11 @@ class UMG_API UHorizontalBoxComponent : public USlateNonLeafWidgetComponent
 
 	UHorizontalBoxSlot* AddSlot(USlateWrapperComponent* Content);
 
+	// USlateNonLeafWidgetComponent
+	virtual int32 GetChildrenCount() const OVERRIDE;
+	virtual USlateWrapperComponent* GetChildAt(int32 Index) const OVERRIDE;
+	// End USlateNonLeafWidgetComponent
+
 #if WITH_EDITOR
 	// UObject interface
 	virtual void PostEditChangeProperty(struct FPropertyChangedEvent& PropertyChangedEvent) OVERRIDE;
diff --git a/Engine/Source/Runtime/UMG/Classes/HorizontalBoxSlot.h b/Engine/Source/Runtime/UMG/Classes/HorizontalBoxSlot.h
--- a/Engine/Source/Runtime/UMG/Classes/HorizontalBoxSlot.h
+++ b/Engine/Source/Runtime/UMG/Classes/HorizontalBoxSlot.h
@@ -2,6 +2,8 @@
 
 #pragma once
 
+#include "SlateWrapperTypes.h"
+
 #include "HorizontalBoxSlot.generated.h"
 
 UCLASS()
@@ -9,8 +11,12 @@ class UMG_API UHorizontalBoxSlot : public UPanelSlot
 {
 	GENERATED_UCLASS_BODY()
 	
-	UPROPERTY()
-	class USlateWrapperComponent* Content;
+	UPROPERTY(EditAnywhere, Category=Appearance)
+	FMargin Padding;
+
+	/** How much space this slot should occupy in the direction of the panel. */
+	UPROPERTY(EditAnywhere, Category=Appearance)
+	FSlateChildSize Size;
 
 	UPROPERTY(EditAnywhere, Category=Appearance)
 	TEnumAsByte<EHorizontalAlignment> HorizontalAlignment;
diff --git a/Engine/Source/Runtime/UMG/Classes/PanelSlot.h b/Engine/Source/Runtime/UMG/Classes/PanelSlot.h
--- a/Engine/Source/Runtime/UMG/Classes/PanelSlot.h
+++ b/Engine/Source/Runtime/UMG/Classes/PanelSlot.h
@@ -8,4 +8,7 @@ UCLASS()
 class UMG_API UPanelSlot : public UObject
 {
 	GENERATED_UCLASS_BODY()
+
+	UPROPERTY()
+	class USlateWrapperComponent* Content;
 };
diff --git a/Engine/Source/Runtime/UMG/Classes/SlateNonLeafWidgetComponent.h b/Engine/Source/Runtime/UMG/Classes/SlateNonLeafWidgetComponent.h
--- a/Engine/Source/Runtime/UMG/Classes/SlateNonLeafWidgetComponent.h
+++ b/Engine/Source/Runtime/UMG/Classes/SlateNonLeafWidgetComponent.h
@@ -19,6 +19,9 @@ class UMG_API USlateNonLeafWidgetComponent : public USlateWrapperComponent
 	// End of UActorComponent interface
 #endif
 
+	virtual int32 GetChildrenCount() const { return 0; }
+	virtual USlateWrapperComponent* GetChildAt(int32 Index) const { return NULL; }
+
 protected:
 	// This function is called when the known children list has changed; use this to repopulate child slots in derived classes
 	virtual void OnKnownChildrenChanged() {}
diff --git a/Engine/Source/Runtime/UMG/Classes/VerticalBoxComponent.h b/Engine/Source/Runtime/UMG/Classes/VerticalBoxComponent.h
--- a/Engine/Source/Runtime/UMG/Classes/VerticalBoxComponent.h
+++ b/Engine/Source/Runtime/UMG/Classes/VerticalBoxComponent.h
@@ -15,6 +15,11 @@ class UMG_API UVerticalBoxComponent : public USlateNonLeafWidgetComponent
 
 	UVerticalBoxSlot* AddSlot(USlateWrapperComponent* Content);
 
+	// USlateNonLeafWidgetComponent
+	virtual int32 GetChildrenCount() const OVERRIDE;
+	virtual USlateWrapperComponent* GetChildAt(int32 Index) const OVERRIDE;
+	// End USlateNonLeafWidgetComponent
+
 #if WITH_EDITOR
 	// UObject interface
 	virtual void PostEditChangeProperty(struct FPropertyChangedEvent& PropertyChangedEvent) OVERRIDE;
diff --git a/Engine/Source/Runtime/UMG/Classes/VerticalBoxSlot.h b/Engine/Source/Runtime/UMG/Classes/VerticalBoxSlot.h
--- a/Engine/Source/Runtime/UMG/Classes/VerticalBoxSlot.h
+++ b/Engine/Source/Runtime/UMG/Classes/VerticalBoxSlot.h
@@ -11,9 +11,6 @@ class UMG_API UVerticalBoxSlot : public UPanelSlot
 {
 	GENERATED_UCLASS_BODY()
 	
-	UPROPERTY()
-	class USlateWrapperComponent* Content;
-
 	UPROPERTY(EditAnywhere, Category=Appearance)
 	FMargin Padding;
 
diff --git a/Engine/Source/Runtime/UMG/Private/CanvasPanelComponent.cpp b/Engine/Source/Runtime/UMG/Private/CanvasPanelComponent.cpp
--- a/Engine/Source/Runtime/UMG/Private/CanvasPanelComponent.cpp
+++ b/Engine/Source/Runtime/UMG/Private/CanvasPanelComponent.cpp
@@ -40,6 +40,16 @@ UCanvasPanelComponent::UCanvasPanelComponent(const FPostConstructInitializePrope
 	DesiredCanvasSize = FVector2D(128.0f, 128.0f);
 }
 
+int32 UCanvasPanelComponent::GetChildrenCount() const
+{
+	return Slots.Num();
+}
+
+USlateWrapperComponent* UCanvasPanelComponent::GetChildAt(int32 Index) const
+{
+	return Slots[Index]->Content;
+}
+
 TSharedRef<SWidget> UCanvasPanelComponent::RebuildWidget()
 {
 	TSharedRef<SFixedSizeCanvas> NewCanvas = SNew(SFixedSizeCanvas, DesiredCanvasSize);
diff --git a/Engine/Source/Runtime/UMG/Private/HorizontalBoxComponent.cpp b/Engine/Source/Runtime/UMG/Private/HorizontalBoxComponent.cpp
--- a/Engine/Source/Runtime/UMG/Private/HorizontalBoxComponent.cpp
+++ b/Engine/Source/Runtime/UMG/Private/HorizontalBoxComponent.cpp
@@ -10,6 +10,16 @@ UHorizontalBoxComponent::UHorizontalBoxComponent(const FPostConstructInitializeP
 {
 }
 
+int32 UHorizontalBoxComponent::GetChildrenCount() const
+{
+	return Slots.Num();
+}
+
+USlateWrapperComponent* UHorizontalBoxComponent::GetChildAt(int32 Index) const
+{
+	return Slots[Index]->Content;
+}
+
 TSharedRef<SWidget> UHorizontalBoxComponent::RebuildWidget()
 {
 	TSharedRef<SHorizontalBox> NewCanvas = SNew(SHorizontalBox);
diff --git a/Engine/Source/Runtime/UMG/Private/VerticalBoxComponent.cpp b/Engine/Source/Runtime/UMG/Private/VerticalBoxComponent.cpp
--- a/Engine/Source/Runtime/UMG/Private/VerticalBoxComponent.cpp
+++ b/Engine/Source/Runtime/UMG/Private/VerticalBoxComponent.cpp
@@ -10,6 +10,16 @@ UVerticalBoxComponent::UVerticalBoxComponent(const FPostConstructInitializePrope
 {
 }
 
+int32 UVerticalBoxComponent::GetChildrenCount() const
+{
+	return Slots.Num();
+}
+
+USlateWrapperComponent* UVerticalBoxComponent::GetChildAt(int32 Index) const
+{
+	return Slots[Index]->Content;
+}
+
 TSharedRef<SWidget> UVerticalBoxComponent::RebuildWidget()
 {
 	TSharedRef<SVerticalBox> NewCanvas = SNew(SVerticalBox);
diff --git a/Engine/Source/Runtime/UMG/Private/WidgetBlueprintGeneratedClass.cpp b/Engine/Source/Runtime/UMG/Private/WidgetBlueprintGeneratedClass.cpp
--- a/Engine/Source/Runtime/UMG/Private/WidgetBlueprintGeneratedClass.cpp
+++ b/Engine/Source/Runtime/UMG/Private/WidgetBlueprintGeneratedClass.cpp
@@ -37,7 +37,7 @@ void UWidgetBlueprintGeneratedClass::CreateComponentsForActor(AActor* Actor) con
 		FArchiveReplaceObjectRef<UObject> ReplaceInCDOAr(Widget, DuplicatedObjectList, /*bNullPrivateRefs=*/ false, /*bIgnoreOuterRef=*/ false, /*bIgnoreArchetypeRef=*/ false);
 	}
 
-	// Iterate over each timeline template
+	int32 i = 0;
 	for ( USlateWrapperComponent* Widget : ClonedWidgets )
 	{
 		// Not fatal if NULL, but shouldn't happen
@@ -54,6 +54,8 @@ void UWidgetBlueprintGeneratedClass::CreateComponentsForActor(AActor* Actor) con
 		//FObjectReader(ClonedComponent, SavedProperties);
 
 
+		FString VariableName = Widget->GetClass()->GetName() + TEXT("_") + FString::FromInt(i);
+
 
 		//USlateWrapperComponent* Widget = ConstructObject<USlateWrapperComponent>(Template->GetClass(), Actor, NewName);
 		Widget->bCreatedByConstructionScript = true; // Indicate it comes from a blueprint so it gets cleared when we rerun construction scripts
@@ -62,15 +64,16 @@ void UWidgetBlueprintGeneratedClass::CreateComponentsForActor(AActor* Actor) con
 
 		// Find property with the same name as the template and assign the new Timeline to it
 		UClass* ActorClass = Actor->GetClass();
-		FName thing = Widget->GetClass()->GetFName();
 		//UTimelineTemplate::TimelineTemplateNameToVariableName(Template->GetFName())
-		UObjectPropertyBase* Prop = FindField<UObjectPropertyBase>(ActorClass, thing);
+		UObjectPropertyBase* Prop = FindField<UObjectPropertyBase>(ActorClass, *VariableName);
 		if ( Prop )
 		{
 			Prop->SetObjectPropertyValue_InContainer(Actor, Widget);
 		}
 
 		Widget->RegisterComponent();
+
+		i++;
 	}
 }
 
