UE4 UMG

* Now properly updating the transient slot property to be able to edit parent container slots on widgets.
* Slots now obey their size param for Vertical and Horizontal boxes.

diff --git a/Engine/Source/Editor/UMGEditor/Private/SUMGEditorTree.cpp b/Engine/Source/Editor/UMGEditor/Private/SUMGEditorTree.cpp
--- a/Engine/Source/Editor/UMGEditor/Private/SUMGEditorTree.cpp
+++ b/Engine/Source/Editor/UMGEditor/Private/SUMGEditorTree.cpp
@@ -18,7 +18,7 @@ void SUMGEditorTree::Construct(const FArguments& InArgs, TSharedPtr<FBlueprintEd
 	BlueprintEditor = InBlueprintEditor;
 
 	UWidgetBlueprint* Blueprint = GetBlueprint();
-	RootWidgets.Add(Blueprint->WidgetTemplates[0]);
+	Blueprint->OnChanged().AddSP(this, &SUMGEditorTree::OnBlueprintChanged);
 
 	FCoreDelegates::OnObjectPropertyChanged.Add( FCoreDelegates::FOnObjectPropertyChanged::FDelegate::CreateRaw(this, &SUMGEditorTree::OnObjectPropertyChanged) );
 
@@ -48,10 +48,15 @@ void SUMGEditorTree::Construct(const FArguments& InArgs, TSharedPtr<FBlueprintEd
 			.TreeItemsSource(&RootWidgets)
 		]
 	];
+
+	RefreshTree();
 }
 
 SUMGEditorTree::~SUMGEditorTree()
 {
+	UWidgetBlueprint* Blueprint = GetBlueprint();
+	Blueprint->OnChanged().RemoveAll(this);
+
 	FCoreDelegates::OnObjectPropertyChanged.Remove( FCoreDelegates::FOnObjectPropertyChanged::FDelegate::CreateRaw(this, &SUMGEditorTree::OnObjectPropertyChanged) );
 }
 
@@ -66,6 +71,13 @@ UWidgetBlueprint* SUMGEditorTree::GetBlueprint() const
 	return NULL;
 }
 
+void SUMGEditorTree::OnBlueprintChanged(UBlueprint* InBlueprint)
+{
+	if ( InBlueprint )
+	{
+		RefreshTree();
+	}
+}
 
 void SUMGEditorTree::OnObjectPropertyChanged(UObject* ObjectBeingModified)
 {
@@ -90,6 +102,8 @@ void SUMGEditorTree::ShowDetailsForObjects(TArray<USlateWrapperComponent*> Widge
 		}
 	}
 
+	UWidgetBlueprint* Blueprint = GetBlueprint();
+
 	// Update the details panel
 	SKismetInspector::FShowDetailsOptions Options(InspectorTitle, true);
 	BlueprintEditor.Pin()->GetInspector()->ShowDetailsForObjects(InspectorObjects, Options);
@@ -167,6 +181,17 @@ FReply SUMGEditorTree::CreateTestUI()
 	return FReply::Handled();
 }
 
+void SUMGEditorTree::RefreshTree()
+{
+	RootWidgets.Reset();
+
+	UWidgetBlueprint* Blueprint = GetBlueprint();
+	if ( Blueprint->WidgetTemplates.Num() > 0 )
+	{
+		RootWidgets.Add(Blueprint->WidgetTemplates[0]);
+	}
+}
+
 //void SUMGEditorTree::AddReferencedObjects( FReferenceCollector& Collector )
 //{
 //	Collector.AddReferencedObject(Page);
diff --git a/Engine/Source/Editor/UMGEditor/Private/SUMGEditorTree.h b/Engine/Source/Editor/UMGEditor/Private/SUMGEditorTree.h
--- a/Engine/Source/Editor/UMGEditor/Private/SUMGEditorTree.h
+++ b/Engine/Source/Editor/UMGEditor/Private/SUMGEditorTree.h
@@ -19,6 +19,7 @@ public:
 private:
 	UWidgetBlueprint* GetBlueprint() const;
 
+	void OnBlueprintChanged(UBlueprint* InBlueprint);
 	void OnObjectPropertyChanged(UObject* ObjectBeingModified);
 
 	void ShowDetailsForObjects(TArray<USlateWrapperComponent*> Widgets);
@@ -28,6 +29,7 @@ private:
 	void WidgetHierarchy_OnSelectionChanged(USlateWrapperComponent* SelectedItem, ESelectInfo::Type SelectInfo);
 
 	FReply CreateTestUI();
+	void RefreshTree();
 	
 private:
 
diff --git a/Engine/Source/Editor/UMGEditor/Private/WidgetBlueprintCompiler.cpp b/Engine/Source/Editor/UMGEditor/Private/WidgetBlueprintCompiler.cpp
--- a/Engine/Source/Editor/UMGEditor/Private/WidgetBlueprintCompiler.cpp
+++ b/Engine/Source/Editor/UMGEditor/Private/WidgetBlueprintCompiler.cpp
@@ -70,7 +70,7 @@ void FWidgetBlueprintCompiler::CleanAndSanitizeClass(UBlueprintGeneratedClass* C
 
 	UWidgetBlueprint* Blueprint = WidgetBlueprint();
 
-	if ( 0 != Blueprint->Timelines.Num() )
+	if ( 0 != Blueprint->WidgetTemplates.Num() )
 	{
 		ClassSubObjects.RemoveAllSwap(FCullTemplateObjectsHelper<USlateWrapperComponent>(Blueprint->WidgetTemplates));
 	}
diff --git a/Engine/Source/Runtime/UMG/Classes/ButtonComponent.h b/Engine/Source/Runtime/UMG/Classes/ButtonComponent.h
--- a/Engine/Source/Runtime/UMG/Classes/ButtonComponent.h
+++ b/Engine/Source/Runtime/UMG/Classes/ButtonComponent.h
@@ -16,7 +16,7 @@ class UMG_API UButtonComponent : public USlateNonLeafWidgetComponent
 
 public:
 	/** Style of the button */
-	UPROPERTY(EditAnywhere, BlueprintReadOnly, Category=Appearance)
+	UPROPERTY(EditAnywhere, BlueprintReadOnly, Category=Style)
 	USlateWidgetStyleAsset* ButtonStyle;
 
 	/** Horizontal positioning of the content within the button */
@@ -40,7 +40,7 @@ public:
 	FText ButtonText;
 
 	/** Text to display as the content of the button */
-	UPROPERTY(EditAnywhere, BlueprintReadOnly, Category=Appearance, meta=(EditCondition="bDisplayText"))
+	UPROPERTY(EditAnywhere, BlueprintReadOnly, Category=Style, meta=( EditCondition="bDisplayText" ))
 	USlateWidgetStyleAsset* ButtonTextStyle;
 
 	/** The scaling factor for the button border */
diff --git a/Engine/Source/Runtime/UMG/Classes/CanvasPanelComponent.h b/Engine/Source/Runtime/UMG/Classes/CanvasPanelComponent.h
--- a/Engine/Source/Runtime/UMG/Classes/CanvasPanelComponent.h
+++ b/Engine/Source/Runtime/UMG/Classes/CanvasPanelComponent.h
@@ -28,6 +28,10 @@ class UMG_API UCanvasPanelComponent : public USlateNonLeafWidgetComponent
 	// UObject interface
 	virtual void PostEditChangeProperty(struct FPropertyChangedEvent& PropertyChangedEvent) OVERRIDE;
 	// End of UObject interface
+
+	// USlateWrapperComponent interface
+	virtual void ConnectEditorData() OVERRIDE;
+	// End USlateWrapperComponent interface
 #endif
 
 protected:
diff --git a/Engine/Source/Runtime/UMG/Classes/HorizontalBoxComponent.h b/Engine/Source/Runtime/UMG/Classes/HorizontalBoxComponent.h
--- a/Engine/Source/Runtime/UMG/Classes/HorizontalBoxComponent.h
+++ b/Engine/Source/Runtime/UMG/Classes/HorizontalBoxComponent.h
@@ -24,6 +24,10 @@ class UMG_API UHorizontalBoxComponent : public USlateNonLeafWidgetComponent
 	// UObject interface
 	virtual void PostEditChangeProperty(struct FPropertyChangedEvent& PropertyChangedEvent) OVERRIDE;
 	// End of UObject interface
+
+	// USlateWrapperComponent interface
+	virtual void ConnectEditorData() OVERRIDE;
+	// End USlateWrapperComponent interface
 #endif
 
 protected:
diff --git a/Engine/Source/Runtime/UMG/Classes/SlateWrapperComponent.h b/Engine/Source/Runtime/UMG/Classes/SlateWrapperComponent.h
--- a/Engine/Source/Runtime/UMG/Classes/SlateWrapperComponent.h
+++ b/Engine/Source/Runtime/UMG/Classes/SlateWrapperComponent.h
@@ -12,7 +12,7 @@ class UMG_API USlateWrapperComponent : public UActorComponent
 	GENERATED_UCLASS_BODY()
 
 #if WITH_EDITOR
-	UPROPERTY(Transient, EditAnywhere, BlueprintReadWrite, Category=Layout)
+	UPROPERTY(Transient, EditInline, EditAnywhere, BlueprintReadWrite, Category=Layout)
 	UPanelSlot* Slot;
 #endif
 
@@ -23,6 +23,10 @@ class UMG_API USlateWrapperComponent : public UActorComponent
 
 	TSharedRef<SWidget> GetWidget();
 
+#if WITH_EDITOR
+	virtual void ConnectEditorData() { }
+#endif
+
 	// Utility methods
 	//@TODO: Should move elsewhere
 	static EVisibility ConvertSerializedVisibilityToRuntime(TEnumAsByte<ESlateVisibility::Type> Input);
diff --git a/Engine/Source/Runtime/UMG/Classes/VerticalBoxComponent.h b/Engine/Source/Runtime/UMG/Classes/VerticalBoxComponent.h
--- a/Engine/Source/Runtime/UMG/Classes/VerticalBoxComponent.h
+++ b/Engine/Source/Runtime/UMG/Classes/VerticalBoxComponent.h
@@ -24,6 +24,10 @@ class UMG_API UVerticalBoxComponent : public USlateNonLeafWidgetComponent
 	// UObject interface
 	virtual void PostEditChangeProperty(struct FPropertyChangedEvent& PropertyChangedEvent) OVERRIDE;
 	// End of UObject interface
+
+	// USlateWrapperComponent interface
+	virtual void ConnectEditorData() OVERRIDE;
+	// End USlateWrapperComponent interface
 #endif
 
 protected:
diff --git a/Engine/Source/Runtime/UMG/Classes/WidgetBlueprint.h b/Engine/Source/Runtime/UMG/Classes/WidgetBlueprint.h
--- a/Engine/Source/Runtime/UMG/Classes/WidgetBlueprint.h
+++ b/Engine/Source/Runtime/UMG/Classes/WidgetBlueprint.h
@@ -13,6 +13,8 @@ public:
 	/** Array of templates for timelines that should be created */
 	UPROPERTY()
 	TArray<class USlateWrapperComponent*> WidgetTemplates;
+
+	virtual void PostLoad() OVERRIDE;
 	
 #if WITH_EDITOR
 	// UBlueprint interface
diff --git a/Engine/Source/Runtime/UMG/Private/CanvasPanelComponent.cpp b/Engine/Source/Runtime/UMG/Private/CanvasPanelComponent.cpp
--- a/Engine/Source/Runtime/UMG/Private/CanvasPanelComponent.cpp
+++ b/Engine/Source/Runtime/UMG/Private/CanvasPanelComponent.cpp
@@ -91,6 +91,10 @@ UCanvasPanelSlot* UCanvasPanelComponent::AddSlot(USlateWrapperComponent* Content
 {
 	UCanvasPanelSlot* Slot = ConstructObject<UCanvasPanelSlot>(UCanvasPanelSlot::StaticClass(), this);
 	Slot->Content = Content;
+
+#if WITH_EDITOR
+	Content->Slot = Slot;
+#endif
 	
 	Slots.Add(Slot);
 
@@ -98,6 +102,15 @@ UCanvasPanelSlot* UCanvasPanelComponent::AddSlot(USlateWrapperComponent* Content
 }
 
 #if WITH_EDITOR
+
+void UCanvasPanelComponent::ConnectEditorData()
+{
+	for ( UCanvasPanelSlot* Slot : Slots )
+	{
+		Slot->Content->Slot = Slot;
+	}
+}
+
 void UCanvasPanelComponent::PostEditChangeProperty(struct FPropertyChangedEvent& PropertyChangedEvent)
 {
 	// Ensure the slots have unique names
diff --git a/Engine/Source/Runtime/UMG/Private/HorizontalBoxComponent.cpp b/Engine/Source/Runtime/UMG/Private/HorizontalBoxComponent.cpp
--- a/Engine/Source/Runtime/UMG/Private/HorizontalBoxComponent.cpp
+++ b/Engine/Source/Runtime/UMG/Private/HorizontalBoxComponent.cpp
@@ -45,13 +45,15 @@ void UHorizontalBoxComponent::OnKnownChildrenChanged()
 				Slots[SlotIndex] = Slot = ConstructObject<UHorizontalBoxSlot>(UHorizontalBoxSlot::StaticClass(), this);
 			}
 
-			auto NewSlot = Canvas->AddSlot()
+			auto& NewSlot = Canvas->AddSlot()
 				.AutoWidth()
 				.HAlign(Slot->HorizontalAlignment)
 				.VAlign(Slot->VerticalAlignment)
 			[
 				Slot->Content == NULL ? SNullWidget::NullWidget : Slot->Content->GetWidget()
 			];
+
+			NewSlot.SizeParam = USlateWrapperComponent::ConvertSerializedSizeParamToRuntime(Slot->Size);
 		}
 	}
 }
@@ -60,6 +62,10 @@ UHorizontalBoxSlot* UHorizontalBoxComponent::AddSlot(USlateWrapperComponent* Con
 {
 	UHorizontalBoxSlot* Slot = ConstructObject<UHorizontalBoxSlot>(UHorizontalBoxSlot::StaticClass(), this);
 	Slot->Content = Content;
+
+#if WITH_EDITOR
+	Content->Slot = Slot;
+#endif
 	
 	Slots.Add(Slot);
 
@@ -67,6 +73,14 @@ UHorizontalBoxSlot* UHorizontalBoxComponent::AddSlot(USlateWrapperComponent* Con
 }
 
 #if WITH_EDITOR
+void UHorizontalBoxComponent::ConnectEditorData()
+{
+	for ( UHorizontalBoxSlot* Slot : Slots )
+	{
+		Slot->Content->Slot = Slot;
+	}
+}
+
 void UHorizontalBoxComponent::PostEditChangeProperty(struct FPropertyChangedEvent& PropertyChangedEvent)
 {
 	// Ensure the slots have unique names
diff --git a/Engine/Source/Runtime/UMG/Private/VerticalBoxComponent.cpp b/Engine/Source/Runtime/UMG/Private/VerticalBoxComponent.cpp
--- a/Engine/Source/Runtime/UMG/Private/VerticalBoxComponent.cpp
+++ b/Engine/Source/Runtime/UMG/Private/VerticalBoxComponent.cpp
@@ -45,13 +45,14 @@ void UVerticalBoxComponent::OnKnownChildrenChanged()
 				Slots[SlotIndex] = Slot = ConstructObject<UVerticalBoxSlot>(UVerticalBoxSlot::StaticClass(), this);
 			}
 
-			auto NewSlot = Canvas->AddSlot()
-				.AutoHeight()
+			auto& NewSlot = Canvas->AddSlot()
 				.HAlign(Slot->HorizontalAlignment)
 				.VAlign(Slot->VerticalAlignment)
 			[
 				Slot->Content == NULL ? SNullWidget::NullWidget : Slot->Content->GetWidget()
 			];
+
+			NewSlot.SizeParam = USlateWrapperComponent::ConvertSerializedSizeParamToRuntime(Slot->Size);
 		}
 	}
 }
@@ -60,6 +61,10 @@ UVerticalBoxSlot* UVerticalBoxComponent::AddSlot(USlateWrapperComponent* Content
 {
 	UVerticalBoxSlot* Slot = ConstructObject<UVerticalBoxSlot>(UVerticalBoxSlot::StaticClass(), this);
 	Slot->Content = Content;
+
+#if WITH_EDITOR
+	Content->Slot = Slot;
+#endif
 	
 	Slots.Add(Slot);
 
@@ -67,6 +72,14 @@ UVerticalBoxSlot* UVerticalBoxComponent::AddSlot(USlateWrapperComponent* Content
 }
 
 #if WITH_EDITOR
+void UVerticalBoxComponent::ConnectEditorData()
+{
+	for ( UVerticalBoxSlot* Slot : Slots )
+	{
+		Slot->Content->Slot = Slot;
+	}
+}
+
 void UVerticalBoxComponent::PostEditChangeProperty(struct FPropertyChangedEvent& PropertyChangedEvent)
 {
 	// Ensure the slots have unique names
diff --git a/Engine/Source/Runtime/UMG/Private/WidgetBlueprint.cpp b/Engine/Source/Runtime/UMG/Private/WidgetBlueprint.cpp
--- a/Engine/Source/Runtime/UMG/Private/WidgetBlueprint.cpp
+++ b/Engine/Source/Runtime/UMG/Private/WidgetBlueprint.cpp
@@ -12,6 +12,16 @@ UWidgetBlueprint::UWidgetBlueprint(const FPostConstructInitializeProperties& PCI
 {
 }
 
+void UWidgetBlueprint::PostLoad()
+{
+	Super::PostLoad();
+
+	for ( USlateWrapperComponent* Widget : WidgetTemplates )
+	{
+		Widget->ConnectEditorData();
+	}
+}
+
 bool UWidgetBlueprint::ValidateGeneratedClass(const UClass* InClass)
 {
 	bool Result = Super::ValidateGeneratedClass(InClass);
diff --git a/Engine/Source/Runtime/UMG/Private/WidgetBlueprintGeneratedClass.cpp b/Engine/Source/Runtime/UMG/Private/WidgetBlueprintGeneratedClass.cpp
--- a/Engine/Source/Runtime/UMG/Private/WidgetBlueprintGeneratedClass.cpp
+++ b/Engine/Source/Runtime/UMG/Private/WidgetBlueprintGeneratedClass.cpp
@@ -73,6 +73,10 @@ void UWidgetBlueprintGeneratedClass::CreateComponentsForActor(AActor* Actor) con
 
 		Widget->RegisterComponent();
 
+#if WITH_EDITOR
+		Widget->ConnectEditorData();
+#endif
+
 		i++;
 	}
 }
