UMG UE4

* Splitting out the designer into its own file and renaming it.
* Fixing the bug where the preview widget also adds itself to PIE, instead of its current world's viewport.

diff --git a/Engine/Source/Editor/UMGEditor/Private/SUMGDesigner.cpp b/Engine/Source/Editor/UMGEditor/Private/SUMGDesigner.cpp
new file mode 100644
--- /dev/null
+++ b/Engine/Source/Editor/UMGEditor/Private/SUMGDesigner.cpp
@@ -0,0 +1,172 @@
+// Copyright 1998-2014 Epic Games, Inc. All Rights Reserved.
+
+#include "UMGEditorPrivatePCH.h"
+
+#include "SUMGDesigner.h"
+#include "BlueprintEditor.h"
+
+#define LOCTEXT_NAMESPACE "UMG"
+
+//class SDesignerWidget : public SBorder
+//{
+//	SLATE_BEGIN_ARGS(SDesignerWidget) {}
+//		/** Slot for the wrapped content (optional) */
+//		SLATE_DEFAULT_SLOT(FArguments, Content)
+//	SLATE_END_ARGS()
+//
+//	void Construct(const FArguments& InArgs, TSharedPtr<FBlueprintEditor> InBlueprintEditor)
+//	{
+//		LastPreviewActor = NULL;
+//		BlueprintEditor = InBlueprintEditor;
+//
+//		SBorder::Construct(SBorder::FArguments()
+//			.BorderImage(FStyleDefaults::GetNoBrush())
+//			.Padding(0.0f)
+//			[
+//
+//				InArgs._Content.Widget
+//			]);
+//		}
+//	}
+//
+//	//virtual FReply OnMouseButtonDown(const FGeometry& MyGeometry, const FPointerEvent& MouseEvent) OVERRIDE;
+//	//virtual FReply OnMouseButtonUp(const FGeometry& MyGeometry, const FPointerEvent& MouseEvent) OVERRIDE;
+///*virtual bool OnHitTest(const FGeometry& MyGeometry, FVector2D InAbsoluteCursorPosition) OVERRIDE;
+//{
+//return true;
+//}*/
+//
+//protected:
+//	TWeakPtr<FBlueprintEditor> BlueprintEditor;
+//	TWeakObjectPtr<AActor> LastPreviewActor;
+//};
+
+
+//////////////////////////////////////////////////////////////////////////
+
+FWorkflowApplicationModeExtender BlueprintEditorExtenderDelegate;
+
+static bool LocateWidgetsUnderCursor_Helper(FArrangedWidget& Candidate, FVector2D InAbsoluteCursorLocation, FArrangedChildren& OutWidgetsUnderCursor, bool bIgnoreEnabledStatus)
+{
+	const bool bCandidateUnderCursor =
+		// Candidate is physically under the cursor
+		Candidate.Geometry.IsUnderLocation(InAbsoluteCursorLocation) &&
+		// Candidate actually considers itself hit by this test
+		( Candidate.Widget->OnHitTest(Candidate.Geometry, InAbsoluteCursorLocation) );
+
+	bool bHitAnyWidget = false;
+	if ( bCandidateUnderCursor )
+	{
+		// The candidate widget is under the mouse
+		OutWidgetsUnderCursor.AddWidget(Candidate);
+
+		// Check to see if we were asked to still allow children to be hit test visible
+		bool bHitChildWidget = false;
+		if ( ( Candidate.Widget->GetVisibility().AreChildrenHitTestVisible() ) != 0 )
+		{
+			FArrangedChildren ArrangedChildren(OutWidgetsUnderCursor.GetFilter());
+			Candidate.Widget->ArrangeChildren(Candidate.Geometry, ArrangedChildren);
+
+			// A widget's children are implicitly Z-ordered from first to last
+			for ( int32 ChildIndex = ArrangedChildren.Num() - 1; !bHitChildWidget && ChildIndex >= 0; --ChildIndex )
+			{
+				FArrangedWidget& SomeChild = ArrangedChildren(ChildIndex);
+				bHitChildWidget = ( SomeChild.Widget->IsEnabled() || bIgnoreEnabledStatus ) && LocateWidgetsUnderCursor_Helper(SomeChild, InAbsoluteCursorLocation, OutWidgetsUnderCursor, bIgnoreEnabledStatus);
+			}
+		}
+
+		// If we hit a child widget or we hit our candidate widget then we'll append our widgets
+		const bool bHitCandidateWidget = Candidate.Widget->GetVisibility().IsHitTestVisible();
+		bHitAnyWidget = bHitChildWidget || bHitCandidateWidget;
+		if ( !bHitAnyWidget )
+		{
+			// No child widgets were hit, and even though the cursor was over our candidate widget, the candidate
+			// widget was not hit-testable, so we won't report it
+			check(OutWidgetsUnderCursor.Last() == Candidate);
+			OutWidgetsUnderCursor.Remove(OutWidgetsUnderCursor.Num() - 1);
+		}
+	}
+
+	return bHitAnyWidget;
+}
+
+/////////////////////////////////////////////////////
+// SUMGDesigner
+
+void SUMGDesigner::Construct(const FArguments& InArgs, TSharedPtr<FBlueprintEditor> InBlueprintEditor)
+{
+	LastPreviewActor = NULL;
+	BlueprintEditor = InBlueprintEditor;
+
+	ChildSlot
+	[
+		SNew(SOverlay)
+		+ SOverlay::Slot()
+		.HAlign(HAlign_Fill)
+		.VAlign(VAlign_Fill)
+		[
+			SAssignNew(PreviewSurface, SBorder)
+			.Visibility(EVisibility::HitTestInvisible)
+		]
+
+		+ SOverlay::Slot()
+		.HAlign(HAlign_Fill)
+		.VAlign(VAlign_Fill)
+		[
+			SNew(STextBlock)
+		]
+	];
+}
+
+FReply SUMGDesigner::OnMouseButtonDown(const FGeometry& MyGeometry, const FPointerEvent& MouseEvent)
+{
+	FVector2D LocalMouseCoordinates = MyGeometry.AbsoluteToLocal(MouseEvent.GetScreenSpacePosition());
+
+	FArrangedChildren Children(EVisibility::All);
+
+	PreviewSurface->SetVisibility(EVisibility::Visible);
+	FArrangedWidget WindowWidgetGeometry(PreviewSurface.ToSharedRef(), MyGeometry);
+	LocateWidgetsUnderCursor_Helper(WindowWidgetGeometry, MouseEvent.GetScreenSpacePosition(), Children, true);
+
+	PreviewSurface->SetVisibility(EVisibility::HitTestInvisible);
+
+	return FReply::Handled();
+}
+
+void SUMGDesigner::Tick(const FGeometry& AllottedGeometry, const double InCurrentTime, const float InDeltaTime)
+{
+	AActor* PreviewActor = BlueprintEditor.Pin()->GetPreviewActor();
+	if (PreviewActor != LastPreviewActor.Get())
+	{
+		LastPreviewActor = PreviewActor;
+	}
+
+	if (AUserWidget* WidgetActor = Cast<AUserWidget>(PreviewActor))
+	{
+		TSharedRef<SWidget> CurrentWidget = WidgetActor->GetWidget();
+
+		if (CurrentWidget != LastPreviewWidget.Pin())
+		{
+			LastPreviewWidget = CurrentWidget;
+			PreviewSurface->SetContent(CurrentWidget);
+		}
+	}
+	else
+	{
+		ChildSlot
+		[
+			SNew(SHorizontalBox)
+			+SHorizontalBox::Slot()
+			.HAlign(HAlign_Center)
+			.VAlign(VAlign_Center)
+			[
+				SNew(STextBlock)
+				.Text(LOCTEXT("NoWrappedWidget", "No actor; Open the viewport and tab back"))
+			]
+		];
+	}
+
+	SCompoundWidget::Tick(AllottedGeometry, InCurrentTime, InDeltaTime);
+}
+
+#undef LOCTEXT_NAMESPACE
diff --git a/Engine/Source/Editor/UMGEditor/Private/SUMGDesigner.h b/Engine/Source/Editor/UMGEditor/Private/SUMGDesigner.h
new file mode 100644
--- /dev/null
+++ b/Engine/Source/Editor/UMGEditor/Private/SUMGDesigner.h
@@ -0,0 +1,27 @@
+// Copyright 1998-2014 Epic Games, Inc. All Rights Reserved.
+
+#pragma once
+
+#include "SCompoundWidget.h"
+
+class SUMGDesigner : public SCompoundWidget
+{
+public:
+
+	SLATE_BEGIN_ARGS( SUMGDesigner ) {}
+	SLATE_END_ARGS()
+
+	void Construct(const FArguments& InArgs, TSharedPtr<FBlueprintEditor> InBlueprintEditor);
+
+	virtual FReply OnMouseButtonDown(const FGeometry& MyGeometry, const FPointerEvent& MouseEvent) OVERRIDE;
+
+	// SWidget interface
+	virtual void Tick(const FGeometry& AllottedGeometry, const double InCurrentTime, const float InDeltaTime) OVERRIDE;
+	// End of Swidget interface
+protected:
+	TWeakPtr<FBlueprintEditor> BlueprintEditor;
+	TWeakObjectPtr<AActor> LastPreviewActor;
+	TWeakPtr<SWidget> LastPreviewWidget;
+
+	TSharedPtr<SBorder> PreviewSurface;
+};
diff --git a/Engine/Source/Editor/UMGEditor/Private/UMGBlueprintEditorExtensionHook.cpp b/Engine/Source/Editor/UMGEditor/Private/UMGBlueprintEditorExtensionHook.cpp
--- a/Engine/Source/Editor/UMGEditor/Private/UMGBlueprintEditorExtensionHook.cpp
+++ b/Engine/Source/Editor/UMGEditor/Private/UMGBlueprintEditorExtensionHook.cpp
@@ -11,189 +11,10 @@
 #include "SSCSEditorViewport.h"
 
 #include "SUMGEditorTree.h"
+#include "SUMGDesigner.h"
 
 #define LOCTEXT_NAMESPACE "UMG_EXTENSION"
 
-//class SDesignerWidget : public SBorder
-//{
-//	SLATE_BEGIN_ARGS(SDesignerWidget) {}
-//		/** Slot for the wrapped content (optional) */
-//		SLATE_DEFAULT_SLOT(FArguments, Content)
-//	SLATE_END_ARGS()
-//
-//	void Construct(const FArguments& InArgs, TSharedPtr<FBlueprintEditor> InBlueprintEditor)
-//	{
-//		LastPreviewActor = NULL;
-//		BlueprintEditor = InBlueprintEditor;
-//
-//		SBorder::Construct(SBorder::FArguments()
-//			.BorderImage(FStyleDefaults::GetNoBrush())
-//			.Padding(0.0f)
-//			[
-//
-//				InArgs._Content.Widget
-//			]);
-//		}
-//	}
-//
-//	//virtual FReply OnMouseButtonDown(const FGeometry& MyGeometry, const FPointerEvent& MouseEvent) OVERRIDE;
-//	//virtual FReply OnMouseButtonUp(const FGeometry& MyGeometry, const FPointerEvent& MouseEvent) OVERRIDE;
-///*virtual bool OnHitTest(const FGeometry& MyGeometry, FVector2D InAbsoluteCursorPosition) OVERRIDE;
-//{
-//return true;
-//}*/
-//
-//protected:
-//	TWeakPtr<FBlueprintEditor> BlueprintEditor;
-//	TWeakObjectPtr<AActor> LastPreviewActor;
-//};
-
-
-//////////////////////////////////////////////////////////////////////////
-
-FWorkflowApplicationModeExtender BlueprintEditorExtenderDelegate;
-
-static bool LocateWidgetsUnderCursor_Helper(FArrangedWidget& Candidate, FVector2D InAbsoluteCursorLocation, FArrangedChildren& OutWidgetsUnderCursor, bool bIgnoreEnabledStatus)
-{
-	const bool bCandidateUnderCursor =
-		// Candidate is physically under the cursor
-		Candidate.Geometry.IsUnderLocation(InAbsoluteCursorLocation) &&
-		// Candidate actually considers itself hit by this test
-		( Candidate.Widget->OnHitTest(Candidate.Geometry, InAbsoluteCursorLocation) );
-
-	bool bHitAnyWidget = false;
-	if ( bCandidateUnderCursor )
-	{
-		// The candidate widget is under the mouse
-		OutWidgetsUnderCursor.AddWidget(Candidate);
-
-		// Check to see if we were asked to still allow children to be hit test visible
-		bool bHitChildWidget = false;
-		if ( ( Candidate.Widget->GetVisibility().AreChildrenHitTestVisible() ) != 0 )
-		{
-			FArrangedChildren ArrangedChildren(OutWidgetsUnderCursor.GetFilter());
-			Candidate.Widget->ArrangeChildren(Candidate.Geometry, ArrangedChildren);
-
-			// A widget's children are implicitly Z-ordered from first to last
-			for ( int32 ChildIndex = ArrangedChildren.Num() - 1; !bHitChildWidget && ChildIndex >= 0; --ChildIndex )
-			{
-				FArrangedWidget& SomeChild = ArrangedChildren(ChildIndex);
-				bHitChildWidget = ( SomeChild.Widget->IsEnabled() || bIgnoreEnabledStatus ) && LocateWidgetsUnderCursor_Helper(SomeChild, InAbsoluteCursorLocation, OutWidgetsUnderCursor, bIgnoreEnabledStatus);
-			}
-		}
-
-		// If we hit a child widget or we hit our candidate widget then we'll append our widgets
-		const bool bHitCandidateWidget = Candidate.Widget->GetVisibility().IsHitTestVisible();
-		bHitAnyWidget = bHitChildWidget || bHitCandidateWidget;
-		if ( !bHitAnyWidget )
-		{
-			// No child widgets were hit, and even though the cursor was over our candidate widget, the candidate
-			// widget was not hit-testable, so we won't report it
-			check(OutWidgetsUnderCursor.Last() == Candidate);
-			OutWidgetsUnderCursor.Remove(OutWidgetsUnderCursor.Num() - 1);
-		}
-	}
-
-	return bHitAnyWidget;
-}
-
-/////////////////////////////////////////////////////
-// SSlateBlueprintPreview
-
-class SSlateBlueprintPreview : public SCompoundWidget
-{
-public:
-
-	SLATE_BEGIN_ARGS( SSlateBlueprintPreview ) {}
-	SLATE_END_ARGS()
-
-	void Construct(const FArguments& InArgs, TSharedPtr<FBlueprintEditor> InBlueprintEditor)
-	{
-		LastPreviewActor = NULL;
-		BlueprintEditor = InBlueprintEditor;
-
-		ChildSlot
-		[
-			SNew(SOverlay)
-			+ SOverlay::Slot()
-			.HAlign(HAlign_Fill)
-			.VAlign(VAlign_Fill)
-			[
-				SAssignNew(PreviewSurface, SBorder)
-				.Visibility(EVisibility::HitTestInvisible)
-			]
-
-			+ SOverlay::Slot()
-			.HAlign(HAlign_Fill)
-			.VAlign(VAlign_Fill)
-			[
-				SNew(STextBlock)
-			]
-		];
-	}
-
-	virtual FReply OnMouseButtonDown(const FGeometry& MyGeometry, const FPointerEvent& MouseEvent) OVERRIDE
-	{
-		FVector2D LocalMouseCoordinates = MyGeometry.AbsoluteToLocal(MouseEvent.GetScreenSpacePosition());
-
-		FArrangedChildren Children(EVisibility::All);
-
-		PreviewSurface->SetVisibility(EVisibility::Visible);
-		FArrangedWidget WindowWidgetGeometry(PreviewSurface.ToSharedRef(), MyGeometry);
-		LocateWidgetsUnderCursor_Helper(WindowWidgetGeometry, MouseEvent.GetScreenSpacePosition(), Children, true);
-
-		PreviewSurface->SetVisibility(EVisibility::HitTestInvisible);
-
-		return FReply::Handled();
-	}
-
-	// SWidget interface
-	virtual void Tick(const FGeometry& AllottedGeometry, const double InCurrentTime, const float InDeltaTime) OVERRIDE;
-	// End of Swidget interface
-protected:
-	TWeakPtr<FBlueprintEditor> BlueprintEditor;
-	TWeakObjectPtr<AActor> LastPreviewActor;
-	TWeakPtr<SWidget> LastPreviewWidget;
-
-	TSharedPtr<SBorder> PreviewSurface;
-};
-
-void SSlateBlueprintPreview::Tick(const FGeometry& AllottedGeometry, const double InCurrentTime, const float InDeltaTime)
-{
-	AActor* PreviewActor = BlueprintEditor.Pin()->GetPreviewActor();
-	if (PreviewActor != LastPreviewActor.Get())
-	{
-		LastPreviewActor = PreviewActor;
-	}
-
-	if (AUserWidget* WidgetActor = Cast<AUserWidget>(PreviewActor))
-	{
-		TSharedRef<SWidget> CurrentWidget = WidgetActor->GetWidget();
-
-		if (CurrentWidget != LastPreviewWidget.Pin())
-		{
-			LastPreviewWidget = CurrentWidget;
-			PreviewSurface->SetContent(CurrentWidget);
-		}
-	}
-	else
-	{
-		ChildSlot
-		[
-			SNew(SHorizontalBox)
-			+SHorizontalBox::Slot()
-			.HAlign(HAlign_Center)
-			.VAlign(VAlign_Center)
-			[
-				SNew(STextBlock)
-				.Text(LOCTEXT("NoWrappedWidget", "No actor; Open the viewport and tab back"))
-			]
-		];
-	}
-
-	SCompoundWidget::Tick(AllottedGeometry, InCurrentTime, InDeltaTime);
-}
-
 /////////////////////////////////////////////////////
 // FSlatePreviewSummoner
 
@@ -233,13 +54,14 @@ public:
 			+ SHorizontalBox::Slot()
 			.FillWidth(0)
 			[
+				// HACK: Super hack to make sure we have a preview actor.
 				Result.ToSharedRef()
 			]
 
 			+ SHorizontalBox::Slot()
 			.FillWidth(1)
 			[
-				SNew(SSlateBlueprintPreview, BlueprintEditorPtr)
+				SNew(SUMGDesigner, BlueprintEditorPtr)
 			];
 		}
 		else
diff --git a/Engine/Source/Runtime/UMG/Private/UserWidget.cpp b/Engine/Source/Runtime/UMG/Private/UserWidget.cpp
--- a/Engine/Source/Runtime/UMG/Private/UserWidget.cpp
+++ b/Engine/Source/Runtime/UMG/Private/UserWidget.cpp
@@ -58,10 +58,14 @@ void AUserWidget::RebuildWrapperWidget()
 		}
 	}
 	
-	if (GEngine->GameViewport != NULL)
+	// If this is a game world add the widget to the current worlds viewport.
+	UWorld* World = GetWorld();
+	if ( World && World->IsGameWorld() )
 	{
-		GEngine->GameViewport->AddViewportWidgetContent(MyWrapperWidget.ToSharedRef());
+		UGameViewportClient* Viewport = World->GetGameViewport();
+		Viewport->AddViewportWidgetContent(MyWrapperWidget.ToSharedRef());
 	}
+
 //@TODO: Figure out how hot-reloading/etc... should work with slate wrapped components
 	//SAssignNew(ScoreboardWidgetContainer,SWeakWidget)
 	//.PossiblyNullContent(ScoreboardWidgetOverlay));
