UMG UE4

* Fixing crash on exit.

diff --git a/Engine/Source/Editor/UMGEditor/Private/SUMGEditorTree.cpp b/Engine/Source/Editor/UMGEditor/Private/SUMGEditorTree.cpp
--- a/Engine/Source/Editor/UMGEditor/Private/SUMGEditorTree.cpp
+++ b/Engine/Source/Editor/UMGEditor/Private/SUMGEditorTree.cpp
@@ -55,7 +55,10 @@ void SUMGEditorTree::Construct(const FArguments& InArgs, TSharedPtr<FBlueprintEd
 SUMGEditorTree::~SUMGEditorTree()
 {
 	UWidgetBlueprint* Blueprint = GetBlueprint();
-	Blueprint->OnChanged().RemoveAll(this);
+	if ( Blueprint )
+	{
+		Blueprint->OnChanged().RemoveAll(this);
+	}
 
 	FCoreDelegates::OnObjectPropertyChanged.Remove( FCoreDelegates::FOnObjectPropertyChanged::FDelegate::CreateRaw(this, &SUMGEditorTree::OnObjectPropertyChanged) );
 }
