UMG UE4

* Fixing layout bugs.
* Cleaning up several things in the codebase.
* Adding some additional plugs in AUserWidget that will allow me to do picking and get the corresponding component from an SWidget pointer.

diff --git a/Engine/Source/Editor/UMGEditor/Private/SUMGDesigner.cpp b/Engine/Source/Editor/UMGEditor/Private/SUMGDesigner.cpp
--- a/Engine/Source/Editor/UMGEditor/Private/SUMGDesigner.cpp
+++ b/Engine/Source/Editor/UMGEditor/Private/SUMGDesigner.cpp
@@ -3,7 +3,9 @@
 #include "UMGEditorPrivatePCH.h"
 
 #include "SUMGDesigner.h"
+
 #include "BlueprintEditor.h"
+#include "SKismetInspector.h"
 
 #define LOCTEXT_NAMESPACE "UMG"
 
@@ -95,7 +97,7 @@ static bool LocateWidgetsUnderCursor_Helper(FArrangedWidget& Candidate, FVector2
 
 void SUMGDesigner::Construct(const FArguments& InArgs, TSharedPtr<FBlueprintEditor> InBlueprintEditor)
 {
-	LastPreviewActor = NULL;
+	PreviewWidgetActor = NULL;
 	BlueprintEditor = InBlueprintEditor;
 
 	ChildSlot
@@ -118,6 +120,57 @@ void SUMGDesigner::Construct(const FArguments& InArgs, TSharedPtr<FBlueprintEdit
 	];
 }
 
+UWidgetBlueprint* SUMGDesigner::GetBlueprint() const
+{
+	if ( BlueprintEditor.IsValid() )
+	{
+		UBlueprint* BP = BlueprintEditor.Pin()->GetBlueprintObj();
+		return Cast<UWidgetBlueprint>(BP);
+	}
+
+	return NULL;
+}
+
+void SUMGDesigner::OnBlueprintChanged(UBlueprint* InBlueprint)
+{
+	if ( InBlueprint )
+	{
+		
+	}
+}
+
+void SUMGDesigner::OnObjectPropertyChanged(UObject* ObjectBeingModified)
+{
+	if ( !ensure(ObjectBeingModified) )
+	{
+		return;
+	}
+}
+
+void SUMGDesigner::ShowDetailsForObjects(TArray<USlateWrapperComponent*> Widgets)
+{
+	// @TODO COde duplication
+
+	// Convert the selection set to an array of UObject* pointers
+	FString InspectorTitle;
+	TArray<UObject*> InspectorObjects;
+	InspectorObjects.Empty(Widgets.Num());
+	for ( USlateWrapperComponent* Widget : Widgets )
+	{
+		//if ( NodePtr->CanEditDefaults() )
+		{
+			InspectorTitle = "Widget";// Widget->GetDisplayString();
+			InspectorObjects.Add(Widget);
+		}
+	}
+
+	UWidgetBlueprint* Blueprint = GetBlueprint();
+
+	// Update the details panel
+	SKismetInspector::FShowDetailsOptions Options(InspectorTitle, true);
+	BlueprintEditor.Pin()->GetInspector()->ShowDetailsForObjects(InspectorObjects, Options);
+}
+
 FReply SUMGDesigner::OnMouseButtonDown(const FGeometry& MyGeometry, const FPointerEvent& MouseEvent)
 {
 	FVector2D LocalMouseCoordinates = MyGeometry.AbsoluteToLocal(MouseEvent.GetScreenSpacePosition());
@@ -130,24 +183,41 @@ FReply SUMGDesigner::OnMouseButtonDown(const FGeometry& MyGeometry, const FPoint
 
 	PreviewSurface->SetVisibility(EVisibility::HitTestInvisible);
 
+	AUserWidget* WidgetActor = PreviewWidgetActor.Get();
+	if ( WidgetActor )
+	{
+		USlateWrapperComponent* PreviewHandle = NULL;
+		for ( int32 ChildIndex = Children.Num() - 1; ChildIndex >= 0; ChildIndex-- )
+		{
+			FArrangedWidget& Child = Children.GetInternalArray()[ChildIndex];
+			PreviewHandle = WidgetActor->GetWidgetHandle(Child.Widget);
+			if ( PreviewHandle )
+			{
+				break;
+			}
+		}
+
+		UWidgetBlueprint* Blueprint = GetBlueprint();
+		// TODO we have the widgets runtime handle, we need to look at the handle's name and find the corresponding
+		// design time handle and select it!
+	}
+
 	return FReply::Handled();
 }
 
 void SUMGDesigner::Tick(const FGeometry& AllottedGeometry, const double InCurrentTime, const float InDeltaTime)
 {
 	AActor* PreviewActor = BlueprintEditor.Pin()->GetPreviewActor();
-	if (PreviewActor != LastPreviewActor.Get())
-	{
-		LastPreviewActor = PreviewActor;
-	}
+	AUserWidget* WidgetActor = Cast<AUserWidget>(PreviewActor);
+	PreviewWidgetActor = WidgetActor;
 
-	if (AUserWidget* WidgetActor = Cast<AUserWidget>(PreviewActor))
+	if (WidgetActor)
 	{
-		TSharedRef<SWidget> CurrentWidget = WidgetActor->GetWidget();
+		TSharedRef<SWidget> CurrentWidget = WidgetActor->GetRootWidget();
 
-		if (CurrentWidget != LastPreviewWidget.Pin())
+		if ( CurrentWidget != PreviewWidget.Pin() )
 		{
-			LastPreviewWidget = CurrentWidget;
+			PreviewWidget = CurrentWidget;
 			PreviewSurface->SetContent(CurrentWidget);
 		}
 	}
diff --git a/Engine/Source/Editor/UMGEditor/Private/SUMGDesigner.h b/Engine/Source/Editor/UMGEditor/Private/SUMGDesigner.h
--- a/Engine/Source/Editor/UMGEditor/Private/SUMGDesigner.h
+++ b/Engine/Source/Editor/UMGEditor/Private/SUMGDesigner.h
@@ -18,10 +18,17 @@ public:
 	// SWidget interface
 	virtual void Tick(const FGeometry& AllottedGeometry, const double InCurrentTime, const float InDeltaTime) OVERRIDE;
 	// End of Swidget interface
-protected:
+
+private:
+	UWidgetBlueprint* GetBlueprint() const;
+	void OnBlueprintChanged(UBlueprint* InBlueprint);
+	void OnObjectPropertyChanged(UObject* ObjectBeingModified);
+	void ShowDetailsForObjects(TArray<USlateWrapperComponent*> Widgets);
+
+private:
 	TWeakPtr<FBlueprintEditor> BlueprintEditor;
-	TWeakObjectPtr<AActor> LastPreviewActor;
-	TWeakPtr<SWidget> LastPreviewWidget;
+	TWeakObjectPtr<AUserWidget> PreviewWidgetActor;
+	TWeakPtr<SWidget> PreviewWidget;
 
 	TSharedPtr<SBorder> PreviewSurface;
 };
diff --git a/Engine/Source/Runtime/UMG/Classes/SlateWrapperComponent.h b/Engine/Source/Runtime/UMG/Classes/SlateWrapperComponent.h
--- a/Engine/Source/Runtime/UMG/Classes/SlateWrapperComponent.h
+++ b/Engine/Source/Runtime/UMG/Classes/SlateWrapperComponent.h
@@ -12,7 +12,7 @@ class UMG_API USlateWrapperComponent : public UActorComponent
 	GENERATED_UCLASS_BODY()
 
 #if WITH_EDITOR
-	UPROPERTY(Transient, EditInline, EditAnywhere, BlueprintReadWrite, Category=Layout)
+	UPROPERTY(Transient, EditInline, EditAnywhere, BlueprintReadWrite, Category=Layout, meta=( ShowOnlyInnerProperties ))
 	UPanelSlot* Slot;
 #endif
 
diff --git a/Engine/Source/Runtime/UMG/Classes/UserWidget.h b/Engine/Source/Runtime/UMG/Classes/UserWidget.h
--- a/Engine/Source/Runtime/UMG/Classes/UserWidget.h
+++ b/Engine/Source/Runtime/UMG/Classes/UserWidget.h
@@ -15,12 +15,14 @@ class UMG_API AUserWidget : public AActor
 	virtual void RerunConstructionScripts() OVERRIDE;
 	// End of AActor interface
 
-	TSharedRef<SWidget> GetWidget();
+	USlateWrapperComponent* GetWidgetHandle(TSharedRef<SWidget> InWidget);
+
+	TSharedRef<SWidget> GetRootWidget();
 
 private:
-	TSharedPtr<SOverlay> MyWrapperWidget;
+	TSharedPtr<SWidget> RootWidget;
+	TMap< SWidget*, USlateWrapperComponent* > WidgetToComponent;
 
 private:
-	void EnsureWidgetExists();
 	void RebuildWrapperWidget();
 };
diff --git a/Engine/Source/Runtime/UMG/Private/HorizontalBoxComponent.cpp b/Engine/Source/Runtime/UMG/Private/HorizontalBoxComponent.cpp
--- a/Engine/Source/Runtime/UMG/Private/HorizontalBoxComponent.cpp
+++ b/Engine/Source/Runtime/UMG/Private/HorizontalBoxComponent.cpp
@@ -46,7 +46,7 @@ void UHorizontalBoxComponent::OnKnownChildrenChanged()
 			}
 
 			auto& NewSlot = Canvas->AddSlot()
-				.AutoWidth()
+				.Padding(Slot->Padding)
 				.HAlign(Slot->HorizontalAlignment)
 				.VAlign(Slot->VerticalAlignment)
 			[
diff --git a/Engine/Source/Runtime/UMG/Private/UserWidget.cpp b/Engine/Source/Runtime/UMG/Private/UserWidget.cpp
--- a/Engine/Source/Runtime/UMG/Private/UserWidget.cpp
+++ b/Engine/Source/Runtime/UMG/Private/UserWidget.cpp
@@ -8,17 +8,17 @@
 AUserWidget::AUserWidget(const FPostConstructInitializeProperties& PCIP)
 	: Super(PCIP)
 {
+	RootWidget = SNullWidget::NullWidget;
 }
 
 void AUserWidget::PostActorCreated()
 {
-	EnsureWidgetExists();
 	Super::PostActorCreated();
 }
 
 void AUserWidget::Destroyed()
 {
-	MyWrapperWidget = NULL;
+	RootWidget.Reset();
 	Super::Destroyed();
 }
 
@@ -29,33 +29,32 @@ void AUserWidget::RerunConstructionScripts()
 	RebuildWrapperWidget();
 }
 
-void AUserWidget::EnsureWidgetExists()
+USlateWrapperComponent* AUserWidget::GetWidgetHandle(TSharedRef<SWidget> InWidget)
 {
-	if (!MyWrapperWidget.IsValid())
-	{
-		MyWrapperWidget = SNew(SOverlay);
-	}
+	return WidgetToComponent.FindRef(&InWidget.Get());
 }
 
 void AUserWidget::RebuildWrapperWidget()
 {
-	EnsureWidgetExists();
-	MyWrapperWidget->ClearChildren();
+	WidgetToComponent.Reset();
+	RootWidget = SNullWidget::NullWidget;
 	
 	TArray<USlateWrapperComponent*> SlateWrapperComponents;
 	GetComponents(SlateWrapperComponents);
 
+	// Add the first component to the root of the widget surface.
+	if ( SlateWrapperComponents.Num() > 0 )
+	{
+		RootWidget = SlateWrapperComponents[0]->GetWidget();
+	}
+
 	// Place all of our top-level children Slate wrapped components into the overlay
 	for (int32 ComponentIndex = 0; ComponentIndex < SlateWrapperComponents.Num(); ++ComponentIndex)
 	{
-		if (SlateWrapperComponents[ComponentIndex]->IsRegistered())
-		{
-			MyWrapperWidget->AddSlot()
-			[
-				SlateWrapperComponents[ComponentIndex]->GetWidget()
-			];
-			break;
-		}
+		USlateWrapperComponent* Handle = SlateWrapperComponents[ComponentIndex];
+		TSharedRef<SWidget> Widget = Handle->GetWidget();
+
+		WidgetToComponent.Add(&Widget.Get(), Handle);
 	}
 	
 	// If this is a game world add the widget to the current worlds viewport.
@@ -63,16 +62,11 @@ void AUserWidget::RebuildWrapperWidget()
 	if ( World && World->IsGameWorld() )
 	{
 		UGameViewportClient* Viewport = World->GetGameViewport();
-		Viewport->AddViewportWidgetContent(MyWrapperWidget.ToSharedRef());
+		Viewport->AddViewportWidgetContent(RootWidget.ToSharedRef());
 	}
-
-//@TODO: Figure out how hot-reloading/etc... should work with slate wrapped components
-	//SAssignNew(ScoreboardWidgetContainer,SWeakWidget)
-	//.PossiblyNullContent(ScoreboardWidgetOverlay));
 }
 
-TSharedRef<SWidget> AUserWidget::GetWidget()
+TSharedRef<SWidget> AUserWidget::GetRootWidget()
 {
-	EnsureWidgetExists();
-	return MyWrapperWidget.ToSharedRef();
+	return RootWidget.ToSharedRef();
 }
diff --git a/Engine/Source/Runtime/UMG/Private/VerticalBoxComponent.cpp b/Engine/Source/Runtime/UMG/Private/VerticalBoxComponent.cpp
--- a/Engine/Source/Runtime/UMG/Private/VerticalBoxComponent.cpp
+++ b/Engine/Source/Runtime/UMG/Private/VerticalBoxComponent.cpp
@@ -46,6 +46,7 @@ void UVerticalBoxComponent::OnKnownChildrenChanged()
 			}
 
 			auto& NewSlot = Canvas->AddSlot()
+				.Padding(Slot->Padding)
 				.HAlign(Slot->HorizontalAlignment)
 				.VAlign(Slot->VerticalAlignment)
 			[
diff --git a/Engine/Source/Runtime/UMG/Private/WidgetBlueprintGeneratedClass.cpp b/Engine/Source/Runtime/UMG/Private/WidgetBlueprintGeneratedClass.cpp
--- a/Engine/Source/Runtime/UMG/Private/WidgetBlueprintGeneratedClass.cpp
+++ b/Engine/Source/Runtime/UMG/Private/WidgetBlueprintGeneratedClass.cpp
@@ -37,6 +37,8 @@ void UWidgetBlueprintGeneratedClass::CreateComponentsForActor(AActor* Actor) con
 		FArchiveReplaceObjectRef<UObject> ReplaceInCDOAr(Widget, DuplicatedObjectList, /*bNullPrivateRefs=*/ false, /*bIgnoreOuterRef=*/ false, /*bIgnoreArchetypeRef=*/ false);
 	}
 
+	AUserWidget* WidgetActor = CastChecked<AUserWidget>(Actor);
+
 	int32 i = 0;
 	for ( USlateWrapperComponent* Widget : ClonedWidgets )
 	{
