Continued UMG work, fixing non-unity build, beginning to get drag/drop working.  Designer selection highlights / picking.

diff --git a/Engine/Source/Editor/UMGEditor/Private/SUMGDesigner.cpp b/Engine/Source/Editor/UMGEditor/Private/SUMGDesigner.cpp
--- a/Engine/Source/Editor/UMGEditor/Private/SUMGDesigner.cpp
+++ b/Engine/Source/Editor/UMGEditor/Private/SUMGDesigner.cpp
@@ -6,6 +6,7 @@
 
 #include "BlueprintEditor.h"
 #include "SKismetInspector.h"
+#include "PrimitiveTypes.h"
 
 #define LOCTEXT_NAMESPACE "UMG"
 
@@ -46,8 +47,6 @@
 
 //////////////////////////////////////////////////////////////////////////
 
-FWorkflowApplicationModeExtender BlueprintEditorExtenderDelegate;
-
 static bool LocateWidgetsUnderCursor_Helper(FArrangedWidget& Candidate, FVector2D InAbsoluteCursorLocation, FArrangedChildren& OutWidgetsUnderCursor, bool bIgnoreEnabledStatus)
 {
 	const bool bCandidateUnderCursor =
@@ -193,6 +192,7 @@ FReply SUMGDesigner::OnMouseButtonDown(const FGeometry& MyGeometry, const FPoint
 			PreviewHandle = WidgetActor->GetWidgetHandle(Child.Widget);
 			if ( PreviewHandle )
 			{
+				SelectedWidget = Child.Widget;
 				break;
 			}
 		}
@@ -205,6 +205,58 @@ FReply SUMGDesigner::OnMouseButtonDown(const FGeometry& MyGeometry, const FPoint
 	return FReply::Handled();
 }
 
+bool SUMGDesigner::GetLocallyArrangedWidget(TSharedRef<SWidget> Widget, FArrangedWidget& ArrangedWidget) const
+{
+	// We can't screenshot the widget unless there's a valid window handle to draw it in.
+	TSharedPtr<SWindow> WidgetWindow = FSlateApplication::Get().FindWidgetWindow(Widget);
+	if ( !WidgetWindow.IsValid() )
+	{
+		return false;
+	}
+
+	TSharedRef<SWindow> CurrentWindowRef = WidgetWindow.ToSharedRef();
+
+	FWidgetPath WidgetPath;
+	if ( FSlateApplication::Get().GeneratePathToWidgetUnchecked(Widget, WidgetPath) )
+	{
+		ArrangedWidget = WidgetPath.FindArrangedWidget(Widget);
+		ArrangedWidget.Geometry.AbsolutePosition -= CurrentWindowRef->GetPositionInScreen();
+		return true;
+	}
+
+	return false;
+}
+
+int32 SUMGDesigner::OnPaint(const FGeometry& AllottedGeometry, const FSlateRect& MyClippingRect, FSlateWindowElementList& OutDrawElements, int32 LayerId, const FWidgetStyle& InWidgetStyle, bool bParentEnabled) const
+{
+	SCompoundWidget::OnPaint(AllottedGeometry, MyClippingRect, OutDrawElements, LayerId, InWidgetStyle, bParentEnabled);
+
+	if ( SelectedWidget.IsValid() )
+	{
+		TSharedRef<SWidget> Widget = SelectedWidget.Pin().ToSharedRef();
+
+		FArrangedWidget ArrangedWidget(SNullWidget::NullWidget, FGeometry());
+		GetLocallyArrangedWidget(Widget, ArrangedWidget);
+
+		const FLinearColor Tint(0, 1, 0);
+
+		FSlateDrawElement::MakeBox(
+			OutDrawElements,
+			++LayerId,
+			FPaintGeometry(
+			ArrangedWidget.Geometry.AbsolutePosition,
+			ArrangedWidget.Geometry.Size * ArrangedWidget.Geometry.Scale,
+			ArrangedWidget.Geometry.Scale),
+			FCoreStyle::Get().GetBrush(TEXT("Debug.Border")),
+			MyClippingRect,
+			ESlateDrawEffect::None,
+			Tint
+		);
+	}
+
+	return LayerId;
+}
+
 void SUMGDesigner::Tick(const FGeometry& AllottedGeometry, const double InCurrentTime, const float InDeltaTime)
 {
 	AActor* PreviewActor = BlueprintEditor.Pin()->GetPreviewActor();
diff --git a/Engine/Source/Editor/UMGEditor/Private/SUMGDesigner.h b/Engine/Source/Editor/UMGEditor/Private/SUMGDesigner.h
--- a/Engine/Source/Editor/UMGEditor/Private/SUMGDesigner.h
+++ b/Engine/Source/Editor/UMGEditor/Private/SUMGDesigner.h
@@ -16,6 +16,7 @@ public:
 	virtual FReply OnMouseButtonDown(const FGeometry& MyGeometry, const FPointerEvent& MouseEvent) OVERRIDE;
 
 	// SWidget interface
+	virtual int32 OnPaint(const FGeometry& AllottedGeometry, const FSlateRect& MyClippingRect, FSlateWindowElementList& OutDrawElements, int32 LayerId, const FWidgetStyle& InWidgetStyle, bool bParentEnabled) const OVERRIDE;
 	virtual void Tick(const FGeometry& AllottedGeometry, const double InCurrentTime, const float InDeltaTime) OVERRIDE;
 	// End of Swidget interface
 
@@ -25,10 +26,14 @@ private:
 	void OnObjectPropertyChanged(UObject* ObjectBeingModified);
 	void ShowDetailsForObjects(TArray<USlateWrapperComponent*> Widgets);
 
+	bool GetLocallyArrangedWidget(TSharedRef<SWidget> Widget, FArrangedWidget& ArrangedWidget) const;
+
 private:
 	TWeakPtr<FBlueprintEditor> BlueprintEditor;
 	TWeakObjectPtr<AUserWidget> PreviewWidgetActor;
 	TWeakPtr<SWidget> PreviewWidget;
 
+	TWeakPtr<SWidget> SelectedWidget;
+
 	TSharedPtr<SBorder> PreviewSurface;
 };
diff --git a/Engine/Source/Editor/UMGEditor/Private/SUMGEditorWidgetTemplates.cpp b/Engine/Source/Editor/UMGEditor/Private/SUMGEditorWidgetTemplates.cpp
--- a/Engine/Source/Editor/UMGEditor/Private/SUMGEditorWidgetTemplates.cpp
+++ b/Engine/Source/Editor/UMGEditor/Private/SUMGEditorWidgetTemplates.cpp
@@ -13,21 +13,7 @@
 #include "BlueprintEditor.h"
 
 #include "DecoratedDragDropOp.h"
-
-class FWidgetTemplateDragDropOp : public FDecoratedDragDropOp
-{
-public:
-	DRAG_DROP_OPERATOR_TYPE(FWidgetTemplateDragDropOp, FDecoratedDragDropOp)
-
-	TSharedPtr<FWidgetTemplateDescriptor> Template;
-
-	void Init(const TSharedPtr<FWidgetTemplateDescriptor>& InTemplate)
-	{
-		Template = InTemplate;
-	}
-};
-
-
+#include "WidgetTemplateDragDropOp.h"
 
 class SWidgetTemplateItem : public SCompoundWidget
 {
@@ -144,13 +130,11 @@ TSharedRef<ITableRow> SUMGEditorWidgetTemplates::OnGenerateWidgetTemplateItem(TS
 
 FReply SUMGEditorWidgetTemplates::OnDraggingWidgetTemplateItem(const FGeometry& MyGeometry, const FPointerEvent& MouseEvent)
 {
-	FReply Reply = FReply::Unhandled();
-
 	if ( WidgetsListView->GetNumItemsSelected() > 0 )
 	{
 		TSharedPtr<FWidgetTemplateDescriptor> Template = WidgetsListView->GetSelectedItems()[0];
-//		Reply = FReply::Handled().BeginDragDrop(FWidgetTemplateDragDropOp::New(Template));
+		return FReply::Handled().BeginDragDrop(FWidgetTemplateDragDropOp::New(Template));
 	}
 
-	return Reply;
+	return FReply::Unhandled();
 }
diff --git a/Engine/Source/Editor/UMGEditor/Private/SUMGEditorWidgetTemplates.h b/Engine/Source/Editor/UMGEditor/Private/SUMGEditorWidgetTemplates.h
--- a/Engine/Source/Editor/UMGEditor/Private/SUMGEditorWidgetTemplates.h
+++ b/Engine/Source/Editor/UMGEditor/Private/SUMGEditorWidgetTemplates.h
@@ -6,20 +6,7 @@
 #include "IUMGEditor.h"
 #include "BlueprintEditor.h"
 
-struct FWidgetTemplateDescriptor
-{
-	FText Name;
-	
-	FText ToolTip;
-
-	FSlateIcon Icon;
-
-	TSubclassOf<USlateWrapperComponent> WidgetClass;
-	
-	FWidgetTemplateDescriptor()
-	{
-	}
-};
+#include "WidgetTemplateDescriptor.h"
 
 class SUMGEditorWidgetTemplates : public SCompoundWidget
 {
diff --git a/Engine/Source/Editor/UMGEditor/Private/UMGBlueprintEditorExtensionHook.cpp b/Engine/Source/Editor/UMGEditor/Private/UMGBlueprintEditorExtensionHook.cpp
--- a/Engine/Source/Editor/UMGEditor/Private/UMGBlueprintEditorExtensionHook.cpp
+++ b/Engine/Source/Editor/UMGEditor/Private/UMGBlueprintEditorExtensionHook.cpp
@@ -277,6 +277,8 @@ public:
 //////////////////////////////////////////////////////////////////////////
 // FUMGBlueprintEditorExtensionHook
 
+FWorkflowApplicationModeExtender BlueprintEditorExtenderDelegate;
+
 void FUMGBlueprintEditorExtensionHook::InstallHooks()
 {
 	BlueprintEditorExtenderDelegate = FWorkflowApplicationModeExtender::CreateStatic(&FBlueprintEditorExtenderForSlatePreview::OnModeCreated);
diff --git a/Engine/Source/Editor/UMGEditor/Private/WidgetTemplateDescriptor.h b/Engine/Source/Editor/UMGEditor/Private/WidgetTemplateDescriptor.h
new file mode 100644
--- /dev/null
+++ b/Engine/Source/Editor/UMGEditor/Private/WidgetTemplateDescriptor.h
@@ -0,0 +1,22 @@
+// Copyright 1998-2014 Epic Games, Inc. All Rights Reserved.
+
+#pragma once
+
+#include "SCompoundWidget.h"
+#include "IUMGEditor.h"
+#include "BlueprintEditor.h"
+
+struct FWidgetTemplateDescriptor
+{
+	FText Name;
+
+	FText ToolTip;
+
+	FSlateIcon Icon;
+
+	TSubclassOf<USlateWrapperComponent> WidgetClass;
+
+	FWidgetTemplateDescriptor()
+	{
+	}
+};
diff --git a/Engine/Source/Editor/UMGEditor/Private/WidgetTemplateDragDropOp.cpp b/Engine/Source/Editor/UMGEditor/Private/WidgetTemplateDragDropOp.cpp
new file mode 100644
--- /dev/null
+++ b/Engine/Source/Editor/UMGEditor/Private/WidgetTemplateDragDropOp.cpp
@@ -0,0 +1,39 @@
+// Copyright 1998-2014 Epic Games, Inc. All Rights Reserved.
+
+#include "UMGEditorPrivatePCH.h"
+
+#include "WidgetTemplateDragDropOp.h"
+
+#define LOCTEXT_NAMESPACE "UMG"
+
+TSharedRef<FWidgetTemplateDragDropOp> FWidgetTemplateDragDropOp::New(const TSharedPtr<FWidgetTemplateDescriptor>& InTemplate)
+{
+	TSharedRef<FWidgetTemplateDragDropOp> Operation = MakeShareable(new FWidgetTemplateDragDropOp());
+	Operation->Template = InTemplate;
+	Operation->DefaultHoverText = InTemplate->Name;
+	Operation->CurrentHoverText = InTemplate->Name;
+	Operation->Construct();
+
+	return Operation;
+}
+
+//TSharedPtr<SWidget> FWidgetTemplateDragDropOp::GetDefaultDecorator() const
+//{
+//	return 
+//		SNew(SBorder)
+//		.BorderImage(FEditorStyle::GetBrush("ContentBrowser.AssetDragDropTooltipBackground"))
+//		.Content()
+//		[
+//			SNew(SHorizontalBox)
+//
+//			+ SHorizontalBox::Slot()
+//			.AutoWidth()
+//			.HAlign(HAlign_Left)
+//			[
+//				SNew(STextBlock)
+//				.Text(Template->Name)
+//			]
+//		];
+//}
+
+#undef LOCTEXT_NAMESPACE
diff --git a/Engine/Source/Editor/UMGEditor/Private/WidgetTemplateDragDropOp.h b/Engine/Source/Editor/UMGEditor/Private/WidgetTemplateDragDropOp.h
new file mode 100644
--- /dev/null
+++ b/Engine/Source/Editor/UMGEditor/Private/WidgetTemplateDragDropOp.h
@@ -0,0 +1,18 @@
+// Copyright 1998-2014 Epic Games, Inc. All Rights Reserved.
+
+#pragma once
+
+#include "DecoratedDragDropOp.h"
+#include "WidgetTemplateDescriptor.h"
+
+class FWidgetTemplateDragDropOp : public FDecoratedDragDropOp
+{
+public:
+	DRAG_DROP_OPERATOR_TYPE(FWidgetTemplateDragDropOp, FDecoratedDragDropOp)
+
+	TSharedPtr<FWidgetTemplateDescriptor> Template;
+
+	static TSharedRef<FWidgetTemplateDragDropOp> New(const TSharedPtr<FWidgetTemplateDescriptor>& InTemplate);
+
+	//virtual TSharedPtr<SWidget> GetDefaultDecorator() const OVERRIDE;
+};
