Fix unity build

diff --git a/Engine/Source/Editor/UMGEditor/Private/SUMGDesigner.h b/Engine/Source/Editor/UMGEditor/Private/SUMGDesigner.h
--- a/Engine/Source/Editor/UMGEditor/Private/SUMGDesigner.h
+++ b/Engine/Source/Editor/UMGEditor/Private/SUMGDesigner.h
@@ -11,7 +11,7 @@ public:
 	SLATE_BEGIN_ARGS( SUMGDesigner ) {}
 	SLATE_END_ARGS()
 
-	void Construct(const FArguments& InArgs, TSharedPtr<FBlueprintEditor> InBlueprintEditor);
+	void Construct(const FArguments& InArgs, TSharedPtr<class FBlueprintEditor> InBlueprintEditor);
 
 	virtual FReply OnMouseButtonDown(const FGeometry& MyGeometry, const FPointerEvent& MouseEvent) OVERRIDE;
 
diff --git a/Engine/Source/Editor/UMGEditor/Private/SUMGEditorTree.cpp b/Engine/Source/Editor/UMGEditor/Private/SUMGEditorTree.cpp
--- a/Engine/Source/Editor/UMGEditor/Private/SUMGEditorTree.cpp
+++ b/Engine/Source/Editor/UMGEditor/Private/SUMGEditorTree.cpp
@@ -12,6 +12,7 @@
 
 #include "BlueprintEditor.h"
 #include "SKismetInspector.h"
+#include "BlueprintEditorUtils.h"
 
 void SUMGEditorTree::Construct(const FArguments& InArgs, TSharedPtr<FBlueprintEditor> InBlueprintEditor, USimpleConstructionScript* InSCS)
 {
diff --git a/Engine/Source/Editor/UMGEditor/Private/UMGBlueprintEditorExtensionHook.cpp b/Engine/Source/Editor/UMGEditor/Private/UMGBlueprintEditorExtensionHook.cpp
--- a/Engine/Source/Editor/UMGEditor/Private/UMGBlueprintEditorExtensionHook.cpp
+++ b/Engine/Source/Editor/UMGEditor/Private/UMGBlueprintEditorExtensionHook.cpp
@@ -12,6 +12,8 @@
 
 #include "SUMGEditorTree.h"
 #include "SUMGDesigner.h"
+#include "SUMGEditorWidgetTemplates.h"
+#include "SKismetInspector.h"
 
 #define LOCTEXT_NAMESPACE "UMG_EXTENSION"
 
diff --git a/Engine/Source/Editor/WorldBrowser/Private/Tiles/STiledLandscapeImportDlg.cpp b/Engine/Source/Editor/WorldBrowser/Private/Tiles/STiledLandscapeImportDlg.cpp
--- a/Engine/Source/Editor/WorldBrowser/Private/Tiles/STiledLandscapeImportDlg.cpp
+++ b/Engine/Source/Editor/WorldBrowser/Private/Tiles/STiledLandscapeImportDlg.cpp
@@ -5,6 +5,7 @@
 #include "SVectorInputBox.h"
 #include "IDesktopPlatform.h"
 #include "DesktopPlatformModule.h"
+#include "ContentBrowserModule.h"
 
 #define LOCTEXT_NAMESPACE "WorldBrowser"
 
diff --git a/Engine/Source/Runtime/Engine/Private/GameSession.cpp b/Engine/Source/Runtime/Engine/Private/GameSession.cpp
--- a/Engine/Source/Runtime/Engine/Private/GameSession.cpp
+++ b/Engine/Source/Runtime/Engine/Private/GameSession.cpp
@@ -7,6 +7,7 @@
 #include "EnginePrivate.h"
 #include "Net/UnrealNetwork.h"
 #include "OnlineSubsystemUtils.h"
+#include "Kismet/GameplayStatics.h"
 
 DEFINE_LOG_CATEGORY_STATIC(LogGameSession, Log, All);
 
diff --git a/Engine/Source/Runtime/Online/OnlineSubsystem/Private/OnlineSubsystemImpl.cpp b/Engine/Source/Runtime/Online/OnlineSubsystem/Private/OnlineSubsystemImpl.cpp
--- a/Engine/Source/Runtime/Online/OnlineSubsystem/Private/OnlineSubsystemImpl.cpp
+++ b/Engine/Source/Runtime/Online/OnlineSubsystem/Private/OnlineSubsystemImpl.cpp
@@ -3,6 +3,7 @@
 #include "OnlineSubsystemPrivatePCH.h"
 #include "OnlineSubsystemImpl.h"
 #include "NamedInterfaces.h"
+#include "OnlineIdentityInterface.h"
 
 FOnlineSubsystemImpl::FOnlineSubsystemImpl() :
 	InstanceName(DEFAULT_INSTANCE),
