UMG. Creating a reusable baseclass that could be used for all 2D design surfaces, that need zoom levels and a grid.  Getting borders working.  Working on drag drop.

diff --git a/Engine/Source/Editor/UMGEditor/Private/SDesignSurface.cpp b/Engine/Source/Editor/UMGEditor/Private/SDesignSurface.cpp
new file mode 100644
--- /dev/null
+++ b/Engine/Source/Editor/UMGEditor/Private/SDesignSurface.cpp
@@ -0,0 +1,371 @@
+// Copyright 1998-2014 Epic Games, Inc. All Rights Reserved.
+
+#include "UMGEditorPrivatePCH.h"
+
+#include "SDesignSurface.h"
+
+#define LOCTEXT_NAMESPACE "Slate"
+
+struct FZoomLevelEntry
+{
+public:
+	FZoomLevelEntry(float InZoomAmount, const FText& InDisplayText, EGraphRenderingLOD::Type InLOD)
+		: DisplayText(FText::Format(NSLOCTEXT("GraphEditor", "Zoom", "Zoom {0}"), InDisplayText))
+		, ZoomAmount(InZoomAmount)
+		, LOD(InLOD)
+	{
+	}
+
+public:
+	FText DisplayText;
+	float ZoomAmount;
+	EGraphRenderingLOD::Type LOD;
+};
+
+struct FFixedZoomLevelsContainer : public FZoomLevelsContainer
+{
+	FFixedZoomLevelsContainer()
+	{
+		// Initialize zoom levels if not done already
+		if ( ZoomLevels.Num() == 0 )
+		{
+			ZoomLevels.Reserve(20);
+			ZoomLevels.Add(FZoomLevelEntry(0.100f, NSLOCTEXT("GraphEditor", "ZoomLevel", "-12"), EGraphRenderingLOD::LowestDetail));
+			ZoomLevels.Add(FZoomLevelEntry(0.125f, NSLOCTEXT("GraphEditor", "ZoomLevel", "-11"), EGraphRenderingLOD::LowestDetail));
+			ZoomLevels.Add(FZoomLevelEntry(0.150f, NSLOCTEXT("GraphEditor", "ZoomLevel", "-10"), EGraphRenderingLOD::LowestDetail));
+			ZoomLevels.Add(FZoomLevelEntry(0.175f, NSLOCTEXT("GraphEditor", "ZoomLevel", "-9"), EGraphRenderingLOD::LowestDetail));
+			ZoomLevels.Add(FZoomLevelEntry(0.200f, NSLOCTEXT("GraphEditor", "ZoomLevel", "-8"), EGraphRenderingLOD::LowestDetail));
+			ZoomLevels.Add(FZoomLevelEntry(0.225f, NSLOCTEXT("GraphEditor", "ZoomLevel", "-7"), EGraphRenderingLOD::LowDetail));
+			ZoomLevels.Add(FZoomLevelEntry(0.250f, NSLOCTEXT("GraphEditor", "ZoomLevel", "-6"), EGraphRenderingLOD::LowDetail));
+			ZoomLevels.Add(FZoomLevelEntry(0.375f, NSLOCTEXT("GraphEditor", "ZoomLevel", "-5"), EGraphRenderingLOD::MediumDetail));
+			ZoomLevels.Add(FZoomLevelEntry(0.500f, NSLOCTEXT("GraphEditor", "ZoomLevel", "-4"), EGraphRenderingLOD::MediumDetail));
+			ZoomLevels.Add(FZoomLevelEntry(0.675f, NSLOCTEXT("GraphEditor", "ZoomLevel", "-3"), EGraphRenderingLOD::MediumDetail));
+			ZoomLevels.Add(FZoomLevelEntry(0.750f, NSLOCTEXT("GraphEditor", "ZoomLevel", "-2"), EGraphRenderingLOD::DefaultDetail));
+			ZoomLevels.Add(FZoomLevelEntry(0.875f, NSLOCTEXT("GraphEditor", "ZoomLevel", "-1"), EGraphRenderingLOD::DefaultDetail));
+			ZoomLevels.Add(FZoomLevelEntry(1.000f, NSLOCTEXT("GraphEditor", "ZoomLevel", "1:1"), EGraphRenderingLOD::DefaultDetail));
+			ZoomLevels.Add(FZoomLevelEntry(1.250f, NSLOCTEXT("GraphEditor", "ZoomLevel", "+1"), EGraphRenderingLOD::DefaultDetail));
+			ZoomLevels.Add(FZoomLevelEntry(1.375f, NSLOCTEXT("GraphEditor", "ZoomLevel", "+2"), EGraphRenderingLOD::DefaultDetail));
+			ZoomLevels.Add(FZoomLevelEntry(1.500f, NSLOCTEXT("GraphEditor", "ZoomLevel", "+3"), EGraphRenderingLOD::FullyZoomedIn));
+			ZoomLevels.Add(FZoomLevelEntry(1.675f, NSLOCTEXT("GraphEditor", "ZoomLevel", "+4"), EGraphRenderingLOD::FullyZoomedIn));
+			ZoomLevels.Add(FZoomLevelEntry(1.750f, NSLOCTEXT("GraphEditor", "ZoomLevel", "+5"), EGraphRenderingLOD::FullyZoomedIn));
+			ZoomLevels.Add(FZoomLevelEntry(1.875f, NSLOCTEXT("GraphEditor", "ZoomLevel", "+6"), EGraphRenderingLOD::FullyZoomedIn));
+			ZoomLevels.Add(FZoomLevelEntry(2.000f, NSLOCTEXT("GraphEditor", "ZoomLevel", "+7"), EGraphRenderingLOD::FullyZoomedIn));
+		}
+	}
+
+	float GetZoomAmount(int32 InZoomLevel) const OVERRIDE
+	{
+		checkSlow(ZoomLevels.IsValidIndex(InZoomLevel));
+		return ZoomLevels[InZoomLevel].ZoomAmount;
+	}
+
+	int32 GetNearestZoomLevel(float InZoomAmount) const OVERRIDE
+	{
+		for ( int32 ZoomLevelIndex=0; ZoomLevelIndex < GetNumZoomLevels(); ++ZoomLevelIndex )
+		{
+			if ( InZoomAmount <= GetZoomAmount(ZoomLevelIndex) )
+			{
+				return ZoomLevelIndex;
+			}
+		}
+
+		return GetDefaultZoomLevel();
+	}
+
+	FText GetZoomText(int32 InZoomLevel) const OVERRIDE
+	{
+		checkSlow(ZoomLevels.IsValidIndex(InZoomLevel));
+		return ZoomLevels[InZoomLevel].DisplayText;
+	}
+
+	int32 GetNumZoomLevels() const OVERRIDE
+	{
+		return ZoomLevels.Num();
+	}
+
+	int32 GetDefaultZoomLevel() const OVERRIDE
+	{
+		return 12;
+	}
+
+	EGraphRenderingLOD::Type GetLOD(int32 InZoomLevel) const OVERRIDE
+	{
+		checkSlow(ZoomLevels.IsValidIndex(InZoomLevel));
+		return ZoomLevels[InZoomLevel].LOD;
+	}
+
+	static TArray<FZoomLevelEntry> ZoomLevels;
+};
+
+TArray<FZoomLevelEntry> FFixedZoomLevelsContainer::ZoomLevels;
+
+/////////////////////////////////////////////////////
+// SDesignSurface
+
+void SDesignSurface::Construct(const FArguments& InArgs)
+{
+	if ( !ZoomLevels.IsValid() )
+	{
+		ZoomLevels = new FFixedZoomLevelsContainer();
+	}
+	ZoomLevel = ZoomLevels->GetDefaultZoomLevel();
+	PreviousZoomLevel = ZoomLevels->GetDefaultZoomLevel();
+	PostChangedZoom();
+	SnapGridSize = 16.0f;
+	bAllowContinousZoomInterpolation = false;
+
+	ViewOffset = FVector2D::ZeroVector;
+
+	ZoomLevelFade = FCurveSequence(0.0f, 1.0f);
+	ZoomLevelFade.Play();
+
+	ZoomLevelGraphFade = FCurveSequence(0.0f, 0.5f);
+	ZoomLevelGraphFade.Play();
+
+	ChildSlot
+	[
+		InArgs._Content.Widget
+	];
+}
+
+int32 SDesignSurface::OnPaint(const FGeometry& AllottedGeometry, const FSlateRect& MyClippingRect, FSlateWindowElementList& OutDrawElements, int32 LayerId, const FWidgetStyle& InWidgetStyle, bool bParentEnabled) const
+{
+	const FSlateBrush* BackgroundImage = FEditorStyle::GetBrush(TEXT("Graph.Panel.SolidBackground"));
+	PaintBackgroundAsLines(BackgroundImage, AllottedGeometry, MyClippingRect, OutDrawElements, LayerId);
+
+	SCompoundWidget::OnPaint(AllottedGeometry, MyClippingRect, OutDrawElements, LayerId, InWidgetStyle, bParentEnabled);
+
+	return LayerId;
+}
+
+FReply SDesignSurface::OnMouseWheel(const FGeometry& MyGeometry, const FPointerEvent& MouseEvent)
+{
+	// We want to zoom into this point; i.e. keep it the same fraction offset into the panel
+	const FVector2D WidgetSpaceCursorPos = MyGeometry.AbsoluteToLocal(MouseEvent.GetScreenSpacePosition());
+	const int32 ZoomLevelDelta = FMath::Floor(MouseEvent.GetWheelDelta());
+	ChangeZoomLevel(ZoomLevelDelta, WidgetSpaceCursorPos, MouseEvent.IsControlDown());
+
+	return FReply::Handled();
+}
+
+inline float FancyMod(float Value, float Size)
+{
+	return ( ( Value >= 0 ) ? 0.0f : Size ) + FMath::Fmod(Value, Size);
+}
+
+float SDesignSurface::GetZoomAmount() const
+{
+	if ( bAllowContinousZoomInterpolation )
+	{
+		return FMath::Lerp(ZoomLevels->GetZoomAmount(PreviousZoomLevel), ZoomLevels->GetZoomAmount(ZoomLevel), ZoomLevelGraphFade.GetLerp());
+	}
+	else
+	{
+		return ZoomLevels->GetZoomAmount(ZoomLevel);
+	}
+}
+
+void SDesignSurface::ChangeZoomLevel(int32 ZoomLevelDelta, const FVector2D& WidgetSpaceZoomOrigin, bool bOverrideZoomLimiting)
+{
+	// We want to zoom into this point; i.e. keep it the same fraction offset into the panel
+	const FVector2D PointToMaintainGraphSpace = PanelCoordToGraphCoord(WidgetSpaceZoomOrigin);
+
+	const int32 DefaultZoomLevel = ZoomLevels->GetDefaultZoomLevel();
+	const int32 NumZoomLevels = ZoomLevels->GetNumZoomLevels();
+
+	const bool bAllowFullZoomRange =
+		// To zoom in past 1:1 the user must press control
+		( ZoomLevel == DefaultZoomLevel && ZoomLevelDelta > 0 && bOverrideZoomLimiting ) ||
+		// If they are already zoomed in past 1:1, user may zoom freely
+		( ZoomLevel > DefaultZoomLevel );
+
+	const float OldZoomLevel = ZoomLevel;
+
+	if ( bAllowFullZoomRange )
+	{
+		ZoomLevel = FMath::Clamp(ZoomLevel + ZoomLevelDelta, 0, NumZoomLevels - 1);
+	}
+	else
+	{
+		// Without control, we do not allow zooming in past 1:1.
+		ZoomLevel = FMath::Clamp(ZoomLevel + ZoomLevelDelta, 0, DefaultZoomLevel);
+	}
+
+	if ( OldZoomLevel != ZoomLevel )
+	{
+		PostChangedZoom();
+	}
+
+	// Note: This happens even when maxed out at a stop; so the user sees the animation and knows that they're at max zoom in/out
+	ZoomLevelFade.Play();
+
+	// Re-center the screen so that it feels like zooming around the cursor.
+	{
+		FSlateRect GraphBounds = ComputeSensibleGraphBounds();
+
+		// Make sure we are not zooming into/out into emptiness; otherwise the user will get lost..
+		const FVector2D ClampedPointToMaintainGraphSpace(
+			FMath::Clamp(PointToMaintainGraphSpace.X, GraphBounds.Left, GraphBounds.Right),
+			FMath::Clamp(PointToMaintainGraphSpace.Y, GraphBounds.Top, GraphBounds.Bottom)
+			);
+
+		this->ViewOffset = ClampedPointToMaintainGraphSpace - WidgetSpaceZoomOrigin / GetZoomAmount();
+	}
+}
+
+FSlateRect SDesignSurface::ComputeSensibleGraphBounds() const
+{
+	float Left = 0.0f;
+	float Top = 0.0f;
+	float Right = 0.0f;
+	float Bottom = 0.0f;
+
+	// Find the bounds of the node positions
+	//for ( int32 ChildIndex = 0; ChildIndex < Children.Num(); ++ChildIndex )
+	//{
+	//	const TSharedRef<SNode>& SomeChild = Children[ChildIndex];
+
+	//	FVector2D ChildPos = SomeChild->GetPosition();
+
+	//	Left = FMath::Min(Left, ChildPos.X);
+	//	Right = FMath::Max(Right, ChildPos.X);
+	//	Top = FMath::Min(Top, ChildPos.Y);
+	//	Bottom = FMath::Max(Bottom, ChildPos.Y);
+	//}
+
+	// Pad it out in every direction, to roughly account for nodes being of non-zero extent
+	const float Padding = 100.0f;
+
+	return FSlateRect(Left - Padding, Top - Padding, Right + Padding, Bottom + Padding);
+}
+
+void SDesignSurface::PostChangedZoom()
+{
+}
+
+FVector2D SDesignSurface::GetViewOffset() const
+{
+	return ViewOffset;
+}
+
+float SDesignSurface::GetSnapGridSize() const
+{
+	return SnapGridSize;
+}
+
+FVector2D SDesignSurface::GraphCoordToPanelCoord(const FVector2D& GraphSpaceCoordinate) const
+{
+	return ( GraphSpaceCoordinate - GetViewOffset() ) * GetZoomAmount();
+}
+
+FVector2D SDesignSurface::PanelCoordToGraphCoord(const FVector2D& PanelSpaceCoordinate) const
+{
+	return PanelSpaceCoordinate / GetZoomAmount() + GetViewOffset();
+}
+
+void SDesignSurface::PaintBackgroundAsLines(const FSlateBrush* BackgroundImage, const FGeometry& AllottedGeometry, const FSlateRect& MyClippingRect, FSlateWindowElementList& OutDrawElements, int32& DrawLayerId) const
+{
+	const bool bAntialias = false;
+
+	const int32 RulePeriod = (int32)FEditorStyle::GetFloat("Graph.Panel.GridRulePeriod");
+	check(RulePeriod > 0);
+
+	const FLinearColor RegularColor(FEditorStyle::GetColor("Graph.Panel.GridLineColor"));
+	const FLinearColor RuleColor(FEditorStyle::GetColor("Graph.Panel.GridRuleColor"));
+	const FLinearColor CenterColor(FEditorStyle::GetColor("Graph.Panel.GridCenterColor"));
+	const float GraphSmallestGridSize = 8.0f;
+	const float RawZoomFactor = GetZoomAmount();
+	const float NominalGridSize = GetSnapGridSize();
+
+	float ZoomFactor = RawZoomFactor;
+	float Inflation = 1.0f;
+	while ( ZoomFactor*Inflation*NominalGridSize <= GraphSmallestGridSize )
+	{
+		Inflation *= 2.0f;
+	}
+
+	const float GridCellSize = NominalGridSize * ZoomFactor * Inflation;
+
+	const float GraphSpaceGridX0 = FancyMod(ViewOffset.X, Inflation * NominalGridSize * RulePeriod);
+	const float GraphSpaceGridY0 = FancyMod(ViewOffset.Y, Inflation * NominalGridSize * RulePeriod);
+
+	float ImageOffsetX = GraphSpaceGridX0 * -ZoomFactor;
+	float ImageOffsetY = GraphSpaceGridY0 * -ZoomFactor;
+
+	const FVector2D ZeroSpace = GraphCoordToPanelCoord(FVector2D::ZeroVector);
+
+	// Fill the background
+	FSlateDrawElement::MakeBox(
+		OutDrawElements,
+		DrawLayerId,
+		AllottedGeometry.ToPaintGeometry(),
+		BackgroundImage,
+		MyClippingRect
+		);
+
+	TArray<FVector2D> LinePoints;
+	new (LinePoints)FVector2D(0.0f, 0.0f);
+	new (LinePoints)FVector2D(0.0f, 0.0f);
+
+	// Horizontal bars
+	for ( int32 GridIndex = 0; ImageOffsetY < AllottedGeometry.Size.Y; ImageOffsetY += GridCellSize, ++GridIndex )
+	{
+		if ( ImageOffsetY >= 0.0f )
+		{
+			const bool bIsRuleLine = ( GridIndex % RulePeriod ) == 0;
+			const int32 Layer = bIsRuleLine ? ( DrawLayerId + 1 ) : DrawLayerId;
+
+			const FLinearColor* Color = bIsRuleLine ? &RuleColor : &RegularColor;
+			if ( FMath::IsNearlyEqual(ZeroSpace.Y, ImageOffsetY, 1.0f) )
+			{
+				Color = &CenterColor;
+			}
+
+			LinePoints[0] = FVector2D(0.0f, ImageOffsetY);
+			LinePoints[1] = FVector2D(AllottedGeometry.Size.X, ImageOffsetY);
+
+			FSlateDrawElement::MakeLines(
+				OutDrawElements,
+				Layer,
+				AllottedGeometry.ToPaintGeometry(),
+				LinePoints,
+				MyClippingRect,
+				ESlateDrawEffect::None,
+				*Color,
+				bAntialias);
+		}
+	}
+
+	// Vertical bars
+	for ( int32 GridIndex = 0; ImageOffsetX < AllottedGeometry.Size.X; ImageOffsetX += GridCellSize, ++GridIndex )
+	{
+		if ( ImageOffsetX >= 0.0f )
+		{
+			const bool bIsRuleLine = ( GridIndex % RulePeriod ) == 0;
+			const int32 Layer = bIsRuleLine ? ( DrawLayerId + 1 ) : DrawLayerId;
+
+			const FLinearColor* Color = bIsRuleLine ? &RuleColor : &RegularColor;
+			if ( FMath::IsNearlyEqual(ZeroSpace.X, ImageOffsetX, 1.0f) )
+			{
+				Color = &CenterColor;
+			}
+
+			LinePoints[0] = FVector2D(ImageOffsetX, 0.0f);
+			LinePoints[1] = FVector2D(ImageOffsetX, AllottedGeometry.Size.Y);
+
+			FSlateDrawElement::MakeLines(
+				OutDrawElements,
+				Layer,
+				AllottedGeometry.ToPaintGeometry(),
+				LinePoints,
+				MyClippingRect,
+				ESlateDrawEffect::None,
+				*Color,
+				bAntialias);
+		}
+	}
+
+	DrawLayerId += 2;
+}
+
+#undef LOCTEXT_NAMESPACE
diff --git a/Engine/Source/Editor/UMGEditor/Private/SDesignSurface.h b/Engine/Source/Editor/UMGEditor/Private/SDesignSurface.h
new file mode 100644
--- /dev/null
+++ b/Engine/Source/Editor/UMGEditor/Private/SDesignSurface.h
@@ -0,0 +1,66 @@
+// Copyright 1998-2014 Epic Games, Inc. All Rights Reserved.
+
+#pragma once
+
+#include "SCompoundWidget.h"
+#include "SNodePanel.h"
+
+class SDesignSurface : public SCompoundWidget
+{
+public:
+
+	SLATE_BEGIN_ARGS( SDesignSurface ) {}
+		/** Slot for this designers content (optional) */
+		SLATE_DEFAULT_SLOT(FArguments, Content)
+	SLATE_END_ARGS()
+
+	void Construct(const FArguments& InArgs);
+
+	// SWidget interface
+	virtual FReply OnMouseWheel(const FGeometry& MyGeometry, const FPointerEvent& MouseEvent) OVERRIDE;
+	virtual int32 OnPaint(const FGeometry& AllottedGeometry, const FSlateRect& MyClippingRect, FSlateWindowElementList& OutDrawElements, int32 LayerId, const FWidgetStyle& InWidgetStyle, bool bParentEnabled) const OVERRIDE;
+	// End of Swidget interface
+
+protected:
+	void PaintBackgroundAsLines(const FSlateBrush* BackgroundImage, const FGeometry& AllottedGeometry, const FSlateRect& MyClippingRect, FSlateWindowElementList& OutDrawElements, int32& DrawLayerId) const;
+
+	/** Get the grid snap size */
+	float GetSnapGridSize() const;
+
+	float GetZoomAmount() const;
+
+	void ChangeZoomLevel(int32 ZoomLevelDelta, const FVector2D& WidgetSpaceZoomOrigin, bool bOverrideZoomLimiting);
+
+	void PostChangedZoom();
+
+	FVector2D GetViewOffset() const;
+
+	FSlateRect ComputeSensibleGraphBounds() const;
+
+	FVector2D GraphCoordToPanelCoord(const FVector2D& GraphSpaceCoordinate) const;
+	FVector2D PanelCoordToGraphCoord(const FVector2D& PanelSpaceCoordinate) const;
+
+protected:
+	/** The position within the graph at which the user is looking */
+	FVector2D ViewOffset;
+
+	float SnapGridSize;
+
+	/** Previous Zoom Level */
+	int32 PreviousZoomLevel;
+
+	/** How zoomed in/out we are. e.g. 0.25f results in quarter-sized nodes. */
+	int32 ZoomLevel;
+
+	/** Allow continous zoom interpolation? */
+	bool bAllowContinousZoomInterpolation;
+
+	/** Fade on zoom for graph */
+	FCurveSequence ZoomLevelGraphFade;
+
+	/** Curve that handles fading the 'Zoom +X' text */
+	FCurveSequence ZoomLevelFade;
+
+	// The interface for mapping ZoomLevel values to actual node scaling values
+	TScopedPointer<FZoomLevelsContainer> ZoomLevels;
+};
diff --git a/Engine/Source/Editor/UMGEditor/Private/SUMGDesigner.cpp b/Engine/Source/Editor/UMGEditor/Private/SUMGDesigner.cpp
--- a/Engine/Source/Editor/UMGEditor/Private/SUMGDesigner.cpp
+++ b/Engine/Source/Editor/UMGEditor/Private/SUMGDesigner.cpp
@@ -98,25 +98,28 @@ void SUMGDesigner::Construct(const FArguments& InArgs, TSharedPtr<FBlueprintEdit
 {
 	PreviewWidgetActor = NULL;
 	BlueprintEditor = InBlueprintEditor;
-
-	ChildSlot
-	[
-		SNew(SOverlay)
-		+ SOverlay::Slot()
-		.HAlign(HAlign_Fill)
-		.VAlign(VAlign_Fill)
+	
+	SDesignSurface::Construct(SDesignSurface::FArguments()
+		.Content()
 		[
-			SAssignNew(PreviewSurface, SBorder)
-			.Visibility(EVisibility::HitTestInvisible)
-		]
+			SNew(SOverlay)
 
-		+ SOverlay::Slot()
-		.HAlign(HAlign_Fill)
-		.VAlign(VAlign_Fill)
-		[
-			SNew(STextBlock)
+			+ SOverlay::Slot()
+			.HAlign(HAlign_Fill)
+			.VAlign(VAlign_Fill)
+			[
+				SAssignNew(PreviewSurface, SBorder)
+				.Visibility(EVisibility::HitTestInvisible)
+			]
+
+			+ SOverlay::Slot()
+			.HAlign(HAlign_Fill)
+			.VAlign(VAlign_Fill)
+			[
+				SNew(STextBlock)
+			]
 		]
-	];
+	);
 }
 
 UWidgetBlueprint* SUMGDesigner::GetBlueprint() const
@@ -229,7 +232,7 @@ bool SUMGDesigner::GetLocallyArrangedWidget(TSharedRef<SWidget> Widget, FArrange
 
 int32 SUMGDesigner::OnPaint(const FGeometry& AllottedGeometry, const FSlateRect& MyClippingRect, FSlateWindowElementList& OutDrawElements, int32 LayerId, const FWidgetStyle& InWidgetStyle, bool bParentEnabled) const
 {
-	SCompoundWidget::OnPaint(AllottedGeometry, MyClippingRect, OutDrawElements, LayerId, InWidgetStyle, bParentEnabled);
+	SDesignSurface::OnPaint(AllottedGeometry, MyClippingRect, OutDrawElements, LayerId, InWidgetStyle, bParentEnabled);
 
 	if ( SelectedWidget.IsValid() )
 	{
diff --git a/Engine/Source/Editor/UMGEditor/Private/SUMGDesigner.h b/Engine/Source/Editor/UMGEditor/Private/SUMGDesigner.h
--- a/Engine/Source/Editor/UMGEditor/Private/SUMGDesigner.h
+++ b/Engine/Source/Editor/UMGEditor/Private/SUMGDesigner.h
@@ -3,8 +3,10 @@
 #pragma once
 
 #include "SCompoundWidget.h"
+#include "SNodePanel.h"
+#include "SDesignSurface.h"
 
-class SUMGDesigner : public SCompoundWidget
+class SUMGDesigner : public SDesignSurface
 {
 public:
 
diff --git a/Engine/Source/Editor/UMGEditor/Private/SUMGEditorTree.cpp b/Engine/Source/Editor/UMGEditor/Private/SUMGEditorTree.cpp
--- a/Engine/Source/Editor/UMGEditor/Private/SUMGEditorTree.cpp
+++ b/Engine/Source/Editor/UMGEditor/Private/SUMGEditorTree.cpp
@@ -6,6 +6,7 @@
 #include "UMGEditor.h"
 #include "UMGEditorViewportClient.h"
 #include "UMGEditorActions.h"
+#include "SUMGEditorTreeItem.h"
 
 #include "PreviewScene.h"
 #include "SceneViewport.h"
@@ -26,6 +27,7 @@ void SUMGEditorTree::Construct(const FArguments& InArgs, TSharedPtr<FBlueprintEd
 	ChildSlot
 	[
 		SNew(SVerticalBox)
+
 		+ SVerticalBox::Slot()
 		.AutoHeight()
 		[
@@ -132,8 +134,7 @@ TSharedRef< ITableRow > SUMGEditorTree::WidgetHierarchy_OnGenerateRow(USlateWrap
 		.Padding(2.0f)
 //		.OnDragDetected(this, &SUMGEditorWidgetTemplates::OnDraggingWidgetTemplateItem)
 		[
-			SNew(STextBlock)
-			.Text(InItem->GetFName().ToString())
+			SNew(SUMGEditorTreeItem, BlueprintEditor.Pin(), InItem)
 		];
 }
 
@@ -156,6 +157,7 @@ FReply SUMGEditorTree::CreateTestUI()
 	UWidgetBlueprint* BP = CastChecked<UWidgetBlueprint>(BlueprintEditor.Pin()->GetBlueprintObj());
 
 	UCanvasPanelComponent* Canvas = ConstructObject<UCanvasPanelComponent>(UCanvasPanelComponent::StaticClass(), BP);
+	UBorderComponent* Border = ConstructObject<UBorderComponent>(UBorderComponent::StaticClass(), BP);
 	UVerticalBoxComponent* Vertical = ConstructObject<UVerticalBoxComponent>(UVerticalBoxComponent::StaticClass(), BP);
 	UButtonComponent* Button1 = ConstructObject<UButtonComponent>(UButtonComponent::StaticClass(), BP);
 	Button1->ButtonText = FText::FromString("Button 1");
@@ -170,6 +172,7 @@ FReply SUMGEditorTree::CreateTestUI()
 	Button4->SetContent(Text1);
 
 	BP->WidgetTemplates.Add(Canvas);
+	BP->WidgetTemplates.Add(Border);
 	BP->WidgetTemplates.Add(Vertical);
 	BP->WidgetTemplates.Add(Button1);
 	BP->WidgetTemplates.Add(Button2);
@@ -177,12 +180,14 @@ FReply SUMGEditorTree::CreateTestUI()
 	BP->WidgetTemplates.Add(Button4);
 	BP->WidgetTemplates.Add(Text1);
 
-	UCanvasPanelSlot* Slot = Canvas->AddSlot(Vertical);
-	Slot->Size.X = 100;
-	Slot->Size.Y = 100;
+	UCanvasPanelSlot* Slot = Canvas->AddSlot(Border);
+	Slot->Size.X = 200;
+	Slot->Size.Y = 300;
 	Slot->Position.X = 20;
 	Slot->Position.Y = 50;
 
+	Border->SetContent(Vertical);
+
 	Vertical->AddSlot(Button1);
 	Vertical->AddSlot(Button2);
 	Vertical->AddSlot(Button3);
@@ -203,16 +208,3 @@ void SUMGEditorTree::RefreshTree()
 		RootWidgets.Add(Blueprint->WidgetTemplates[0]);
 	}
 }
-
-//void SUMGEditorTree::AddReferencedObjects( FReferenceCollector& Collector )
-//{
-//	Collector.AddReferencedObject(Page);
-//}
-//
-//void SUMGEditorTree::OnObjectPropertyChanged(UObject* ObjectBeingModified)
-//{
-//	if ( !ensure(ObjectBeingModified) )
-//	{
-//		return;
-//	}
-//}
diff --git a/Engine/Source/Editor/UMGEditor/Private/SUMGEditorTreeItem.cpp b/Engine/Source/Editor/UMGEditor/Private/SUMGEditorTreeItem.cpp
new file mode 100644
--- /dev/null
+++ b/Engine/Source/Editor/UMGEditor/Private/SUMGEditorTreeItem.cpp
@@ -0,0 +1,79 @@
+// Copyright 1998-2014 Epic Games, Inc. All Rights Reserved.
+
+#include "UMGEditorPrivatePCH.h"
+
+#include "SUMGEditorTreeItem.h"
+#include "UMGEditor.h"
+#include "UMGEditorViewportClient.h"
+#include "UMGEditorActions.h"
+#include "WidgetTemplateDragDropOp.h"
+
+#include "PreviewScene.h"
+#include "SceneViewport.h"
+
+#include "BlueprintEditor.h"
+#include "SKismetInspector.h"
+#include "BlueprintEditorUtils.h"
+
+void SUMGEditorTreeItem::Construct(const FArguments& InArgs, TSharedPtr<FBlueprintEditor> InBlueprintEditor, USlateWrapperComponent* InItem)
+{
+	BlueprintEditor = InBlueprintEditor;
+	Item = InItem;
+
+	ChildSlot
+	[
+		SNew(STextBlock)
+		.Text(this, &SUMGEditorTreeItem::GetItemText)
+	];
+}
+
+FText SUMGEditorTreeItem::GetItemText() const
+{
+	return FText::FromString(Item->GetName());
+}
+
+void SUMGEditorTreeItem::OnDragEnter(const FGeometry& MyGeometry, const FDragDropEvent& DragDropEvent)
+{
+	//TSharedPtr<FDragDropOperation> Operation = DragDropEvent.GetOperation();
+	//if ( !Operation.IsValid() )
+	//{
+	//	return;
+	//}
+
+	//if ( Operation->IsOfType<FWidgetTemplateDragDropOp>() )
+	//{
+	//	TSharedPtr<FWidgetTemplateDragDropOp> DragDropOp = StaticCastSharedPtr<FWidgetTemplateDragDropOp>(Operation);
+	//	bDraggedOver = true;
+	//}
+}
+
+void SUMGEditorTreeItem::OnDragLeave(const FDragDropEvent& DragDropEvent)
+{
+	TSharedPtr<FWidgetTemplateDragDropOp> DragDropOp = DragDropEvent.GetOperationAs<FWidgetTemplateDragDropOp>();
+	if ( DragDropOp.IsValid() )
+	{
+		DragDropOp->ResetToDefaultToolTip();
+	}
+}
+
+FReply SUMGEditorTreeItem::OnDragOver(const FGeometry& MyGeometry, const FDragDropEvent& DragDropEvent)
+{
+	TSharedPtr<FWidgetTemplateDragDropOp> DragDropOp = DragDropEvent.GetOperationAs<FWidgetTemplateDragDropOp>();
+	if ( DragDropOp.IsValid() )
+	{
+		return FReply::Handled();
+	}
+
+	return FReply::Unhandled();
+}
+
+FReply SUMGEditorTreeItem::OnDrop(const FGeometry& MyGeometry, const FDragDropEvent& DragDropEvent)
+{
+	TSharedPtr<FWidgetTemplateDragDropOp> DragDropOp = DragDropEvent.GetOperationAs<FWidgetTemplateDragDropOp>();
+	if ( DragDropOp.IsValid() )
+	{
+		return FReply::Handled();
+	}
+
+	return FReply::Unhandled();
+}
diff --git a/Engine/Source/Editor/UMGEditor/Private/SUMGEditorTreeItem.h b/Engine/Source/Editor/UMGEditor/Private/SUMGEditorTreeItem.h
new file mode 100644
--- /dev/null
+++ b/Engine/Source/Editor/UMGEditor/Private/SUMGEditorTreeItem.h
@@ -0,0 +1,28 @@
+// Copyright 1998-2014 Epic Games, Inc. All Rights Reserved.
+
+#pragma once
+
+#include "SCompoundWidget.h"
+#include "IUMGEditor.h"
+#include "BlueprintEditor.h"
+
+//TODO rename SUMGEditorHierarchy
+class SUMGEditorTreeItem : public SCompoundWidget
+{
+public:
+	SLATE_BEGIN_ARGS( SUMGEditorTreeItem ){}
+	SLATE_END_ARGS()
+
+	void Construct(const FArguments& InArgs, TSharedPtr<FBlueprintEditor> InBlueprintEditor, USlateWrapperComponent* InItem);
+
+	virtual void OnDragEnter(const FGeometry& MyGeometry, const FDragDropEvent& DragDropEvent) OVERRIDE;
+	virtual void OnDragLeave(const FDragDropEvent& DragDropEvent) OVERRIDE;
+	virtual FReply OnDragOver(const FGeometry& MyGeometry, const FDragDropEvent& DragDropEvent) OVERRIDE;
+	virtual FReply OnDrop(const FGeometry& MyGeometry, const FDragDropEvent& DragDropEvent) OVERRIDE;
+	
+private:
+	FText GetItemText() const;
+
+	TWeakPtr<class FBlueprintEditor> BlueprintEditor;
+	USlateWrapperComponent* Item;
+};
diff --git a/Engine/Source/Runtime/UMG/Classes/BorderComponent.h b/Engine/Source/Runtime/UMG/Classes/BorderComponent.h
--- a/Engine/Source/Runtime/UMG/Classes/BorderComponent.h
+++ b/Engine/Source/Runtime/UMG/Classes/BorderComponent.h
@@ -30,8 +30,8 @@ class UMG_API UBorderComponent : public UContentWidget
 	// SLATE_EVENT( FPointerEventHandler, OnMouseDoubleClick )
 
 	/** Image to use for the border */
-	UPROPERTY(EditAnywhere, BlueprintReadWrite, Category=Appearance)
-	FSlateBrush BorderImage;
+	UPROPERTY(EditAnywhere, BlueprintReadWrite, Category=Appearance, meta=( DisplayThumbnail = "true" ))
+	USlateBrushAsset* BorderBrush;
 
 	/** The scaling factor for the border */
 	UPROPERTY(EditAnywhere, BlueprintReadOnly, Category=Appearance)
@@ -71,6 +71,7 @@ protected:
 
 	FMargin GetContentPadding() const;
 	FLinearColor GetContentColor() const;
+	const FSlateBrush* GetBorderBrush() const;
 	FSlateColor GetBorderColor() const;
 	FSlateColor GetForegroundColor() const;
 };
diff --git a/Engine/Source/Runtime/UMG/Private/BorderComponent.cpp b/Engine/Source/Runtime/UMG/Private/BorderComponent.cpp
--- a/Engine/Source/Runtime/UMG/Private/BorderComponent.cpp
+++ b/Engine/Source/Runtime/UMG/Private/BorderComponent.cpp
@@ -34,7 +34,7 @@ TSharedRef<SWidget> UBorderComponent::RebuildWidget()
 		// SLATE_EVENT( FPointerEventHandler, OnMouseButtonUp )
 		// SLATE_EVENT( FPointerEventHandler, OnMouseMove )
 		// SLATE_EVENT( FPointerEventHandler, OnMouseDoubleClick )
-		.BorderImage(&BorderImage)
+		.BorderImage(BIND_UOBJECT_ATTRIBUTE(const FSlateBrush*, GetBorderBrush))
 		.ContentScale(ContentScale)
 		.DesiredSizeScale(DesiredSizeScale)
 		.ColorAndOpacity(BIND_UOBJECT_ATTRIBUTE(FLinearColor, GetContentColor))
@@ -42,6 +42,8 @@ TSharedRef<SWidget> UBorderComponent::RebuildWidget()
 		.ForegroundColor(BIND_UOBJECT_ATTRIBUTE(FSlateColor, GetForegroundColor))
 		.ShowEffectWhenDisabled(bShowEffectWhenDisabled);
 
+	NewBorder->SetContent(Content ? Content->GetWidget() : SNullWidget::NullWidget);
+
 	MyBorder = NewBorder;
 	return NewBorder;
 }
@@ -53,7 +55,7 @@ void UBorderComponent::SetContent(USlateWrapperComponent* Content)
 	TSharedPtr<SBorder> Border = MyBorder.Pin();
 	if ( Border.IsValid() )
 	{
-		Border->SetContent(Content->GetWidget());
+		Border->SetContent(Content ? Content->GetWidget() : SNullWidget::NullWidget);
 	}
 }
 
@@ -67,6 +69,17 @@ FLinearColor UBorderComponent::GetContentColor() const
 	return ContentColorAndOpacity;
 }
 
+const FSlateBrush* UBorderComponent::GetBorderBrush() const
+{
+	if ( BorderBrush == NULL )
+	{
+		SBorder::FArguments BorderDefaults;
+		return BorderDefaults._BorderImage.Get();
+	}
+
+	return &BorderBrush->Brush;
+}
+
 FSlateColor UBorderComponent::GetBorderColor() const
 {
 	return BorderColorAndOpacity;
