UMG - Basic support added for dragging and dropping new widgets into the tree, only works for vertical box as the drop target for now.

diff --git a/Engine/Source/Editor/UMGEditor/Private/SDesignSurface.h b/Engine/Source/Editor/UMGEditor/Private/SDesignSurface.h
--- a/Engine/Source/Editor/UMGEditor/Private/SDesignSurface.h
+++ b/Engine/Source/Editor/UMGEditor/Private/SDesignSurface.h
@@ -52,7 +52,7 @@ protected:
 	/** How zoomed in/out we are. e.g. 0.25f results in quarter-sized nodes. */
 	int32 ZoomLevel;
 
-	/** Allow continous zoom interpolation? */
+	/** Allow continuous zoom interpolation? */
 	bool bAllowContinousZoomInterpolation;
 
 	/** Fade on zoom for graph */
diff --git a/Engine/Source/Editor/UMGEditor/Private/SUMGDesigner.cpp b/Engine/Source/Editor/UMGEditor/Private/SUMGDesigner.cpp
--- a/Engine/Source/Editor/UMGEditor/Private/SUMGDesigner.cpp
+++ b/Engine/Source/Editor/UMGEditor/Private/SUMGDesigner.cpp
@@ -243,10 +243,11 @@ int32 SUMGDesigner::OnPaint(const FGeometry& AllottedGeometry, const FSlateRect&
 
 		const FLinearColor Tint(0, 1, 0);
 
+		LayerId += 100;
+
 		FSlateDrawElement::MakeBox(
 			OutDrawElements,
-			++LayerId,
-			FPaintGeometry(
+			LayerId,			FPaintGeometry(
 			ArrangedWidget.Geometry.AbsolutePosition,
 			ArrangedWidget.Geometry.Size * ArrangedWidget.Geometry.Scale,
 			ArrangedWidget.Geometry.Scale),
diff --git a/Engine/Source/Editor/UMGEditor/Private/SUMGEditorTreeItem.cpp b/Engine/Source/Editor/UMGEditor/Private/SUMGEditorTreeItem.cpp
--- a/Engine/Source/Editor/UMGEditor/Private/SUMGEditorTreeItem.cpp
+++ b/Engine/Source/Editor/UMGEditor/Private/SUMGEditorTreeItem.cpp
@@ -72,6 +72,20 @@ FReply SUMGEditorTreeItem::OnDrop(const FGeometry& MyGeometry, const FDragDropEv
 	TSharedPtr<FWidgetTemplateDragDropOp> DragDropOp = DragDropEvent.GetOperationAs<FWidgetTemplateDragDropOp>();
 	if ( DragDropOp.IsValid() )
 	{
+		if ( Item->IsA(USlateNonLeafWidgetComponent::StaticClass()) )
+		{
+			UWidgetBlueprint* BP = CastChecked<UWidgetBlueprint>(BlueprintEditor.Pin()->GetBlueprintObj());
+
+			USlateWrapperComponent* Widget = ConstructObject<USlateWrapperComponent>(DragDropOp->Template->WidgetClass, BP);
+			
+			USlateNonLeafWidgetComponent* Parent = Cast<USlateNonLeafWidgetComponent>(Item);
+			Parent->AddChild(Widget);
+
+			BP->WidgetTemplates.Add(Widget);
+
+			FBlueprintEditorUtils::MarkBlueprintAsStructurallyModified(BP);
+		}
+
 		return FReply::Handled();
 	}
 
diff --git a/Engine/Source/Runtime/UMG/Classes/SlateNonLeafWidgetComponent.h b/Engine/Source/Runtime/UMG/Classes/SlateNonLeafWidgetComponent.h
--- a/Engine/Source/Runtime/UMG/Classes/SlateNonLeafWidgetComponent.h
+++ b/Engine/Source/Runtime/UMG/Classes/SlateNonLeafWidgetComponent.h
@@ -11,4 +11,6 @@ class UMG_API USlateNonLeafWidgetComponent : public USlateWrapperComponent
 
 	virtual int32 GetChildrenCount() const { return 0; }
 	virtual USlateWrapperComponent* GetChildAt(int32 Index) const { return NULL; }
+
+	virtual bool AddChild(USlateWrapperComponent* Child) { return false; }
 };
diff --git a/Engine/Source/Runtime/UMG/Classes/VerticalBoxComponent.h b/Engine/Source/Runtime/UMG/Classes/VerticalBoxComponent.h
--- a/Engine/Source/Runtime/UMG/Classes/VerticalBoxComponent.h
+++ b/Engine/Source/Runtime/UMG/Classes/VerticalBoxComponent.h
@@ -18,6 +18,11 @@ class UMG_API UVerticalBoxComponent : public USlateNonLeafWidgetComponent
 	// USlateNonLeafWidgetComponent
 	virtual int32 GetChildrenCount() const OVERRIDE;
 	virtual USlateWrapperComponent* GetChildAt(int32 Index) const OVERRIDE;
+	virtual bool AddChild(USlateWrapperComponent* Child) OVERRIDE
+	{
+		AddSlot(Child);
+		return true;
+	}
 	// End USlateNonLeafWidgetComponent
 
 #if WITH_EDITOR
