UMG working on drag drop, and safe duplication of the widget blueprint.  Still buggy, GLEO issues.

diff --git a/Engine/Source/Editor/UMGEditor/Private/SUMGEditorTree.cpp b/Engine/Source/Editor/UMGEditor/Private/SUMGEditorTree.cpp
--- a/Engine/Source/Editor/UMGEditor/Private/SUMGEditorTree.cpp
+++ b/Engine/Source/Editor/UMGEditor/Private/SUMGEditorTree.cpp
@@ -71,7 +71,7 @@ UWidgetBlueprint* SUMGEditorTree::GetBlueprint() const
 	if ( BlueprintEditor.IsValid() )
 	{
 		UBlueprint* BP = BlueprintEditor.Pin()->GetBlueprintObj();
-		return Cast<UWidgetBlueprint>(BP);
+		return CastChecked<UWidgetBlueprint>(BP);
 	}
 
 	return NULL;
@@ -154,44 +154,58 @@ FReply SUMGEditorTree::CreateTestUI()
 	//UCanvasPanelComponent* Canvas = Cast<UCanvasPanelComponent>(SCSEditor->AddNewComponent(UCanvasPanelComponent::StaticClass(), NULL));
 	//UButtonComponent* Button = Cast<UButtonComponent>(SCSEditor->AddNewComponent(UButtonComponent::StaticClass(), NULL));
 
-	UWidgetBlueprint* BP = CastChecked<UWidgetBlueprint>(BlueprintEditor.Pin()->GetBlueprintObj());
-
-	UCanvasPanelComponent* Canvas = ConstructObject<UCanvasPanelComponent>(UCanvasPanelComponent::StaticClass(), BP);
-	UBorderComponent* Border = ConstructObject<UBorderComponent>(UBorderComponent::StaticClass(), BP);
-	UVerticalBoxComponent* Vertical = ConstructObject<UVerticalBoxComponent>(UVerticalBoxComponent::StaticClass(), BP);
-	UButtonComponent* Button1 = ConstructObject<UButtonComponent>(UButtonComponent::StaticClass(), BP);
-	Button1->ButtonText = FText::FromString("Button 1");
-	UButtonComponent* Button2 = ConstructObject<UButtonComponent>(UButtonComponent::StaticClass(), BP);
-	Button2->ButtonText = FText::FromString("Button 2");
-	UButtonComponent* Button3 = ConstructObject<UButtonComponent>(UButtonComponent::StaticClass(), BP);
-	Button3->ButtonText = FText::FromString("Button 3");
-	UButtonComponent* Button4 = ConstructObject<UButtonComponent>(UButtonComponent::StaticClass(), BP);
-
-	UTextBlockComponent* Text1 = ConstructObject<UTextBlockComponent>(UTextBlockComponent::StaticClass(), BP);
-	Text1->Text = FText::FromString(TEXT("Button!"));
-	Button4->SetContent(Text1);
-
-	BP->WidgetTemplates.Add(Canvas);
-	BP->WidgetTemplates.Add(Border);
-	BP->WidgetTemplates.Add(Vertical);
-	BP->WidgetTemplates.Add(Button1);
-	BP->WidgetTemplates.Add(Button2);
-	BP->WidgetTemplates.Add(Button3);
-	BP->WidgetTemplates.Add(Button4);
-	BP->WidgetTemplates.Add(Text1);
-
-	UCanvasPanelSlot* Slot = Canvas->AddSlot(Border);
-	Slot->Size.X = 200;
-	Slot->Size.Y = 300;
-	Slot->Position.X = 20;
-	Slot->Position.Y = 50;
-
-	Border->SetContent(Vertical);
-
-	Vertical->AddSlot(Button1);
-	Vertical->AddSlot(Button2);
-	Vertical->AddSlot(Button3);
-	Vertical->AddSlot(Button4);
+	UWidgetBlueprint* BP = GetBlueprint();
+
+	if ( BP->WidgetTemplates.Num() > 0 )
+	{
+		UVerticalBoxComponent* Vertical = CastChecked<UVerticalBoxComponent>(BP->WidgetTemplates[2]);
+
+		UButtonComponent* NewButton = ConstructObject<UButtonComponent>(UButtonComponent::StaticClass(), BP);
+		NewButton->ButtonText = FText::FromString("New Button");
+
+		Vertical->AddChild(NewButton);
+
+		BP->WidgetTemplates.Add(NewButton);
+	}
+	else
+	{
+		UCanvasPanelComponent* Canvas = ConstructObject<UCanvasPanelComponent>(UCanvasPanelComponent::StaticClass(), BP);
+		UBorderComponent* Border = ConstructObject<UBorderComponent>(UBorderComponent::StaticClass(), BP);
+		UVerticalBoxComponent* Vertical = ConstructObject<UVerticalBoxComponent>(UVerticalBoxComponent::StaticClass(), BP);
+		UButtonComponent* Button1 = ConstructObject<UButtonComponent>(UButtonComponent::StaticClass(), BP);
+		Button1->ButtonText = FText::FromString("Button 1");
+		UButtonComponent* Button2 = ConstructObject<UButtonComponent>(UButtonComponent::StaticClass(), BP);
+		Button2->ButtonText = FText::FromString("Button 2");
+		UButtonComponent* Button3 = ConstructObject<UButtonComponent>(UButtonComponent::StaticClass(), BP);
+		Button3->ButtonText = FText::FromString("Button 3");
+		UButtonComponent* Button4 = ConstructObject<UButtonComponent>(UButtonComponent::StaticClass(), BP);
+
+		UTextBlockComponent* Text1 = ConstructObject<UTextBlockComponent>(UTextBlockComponent::StaticClass(), BP);
+		Text1->Text = FText::FromString(TEXT("Button!"));
+		Button4->SetContent(Text1);
+
+		BP->WidgetTemplates.Add(Canvas);
+		BP->WidgetTemplates.Add(Border);
+		BP->WidgetTemplates.Add(Vertical);
+		BP->WidgetTemplates.Add(Button1);
+		BP->WidgetTemplates.Add(Button2);
+		BP->WidgetTemplates.Add(Button3);
+		BP->WidgetTemplates.Add(Button4);
+		BP->WidgetTemplates.Add(Text1);
+
+		UCanvasPanelSlot* Slot = Canvas->AddSlot(Border);
+		Slot->Size.X = 200;
+		Slot->Size.Y = 300;
+		Slot->Position.X = 20;
+		Slot->Position.Y = 50;
+
+		Border->SetContent(Vertical);
+
+		Vertical->AddSlot(Button1);
+		Vertical->AddSlot(Button2);
+		Vertical->AddSlot(Button3);
+		Vertical->AddSlot(Button4);
+	}
 
 	FBlueprintEditorUtils::MarkBlueprintAsStructurallyModified(BP);
 
@@ -207,4 +221,6 @@ void SUMGEditorTree::RefreshTree()
 	{
 		RootWidgets.Add(Blueprint->WidgetTemplates[0]);
 	}
+
+	WidgetTreeView->RequestTreeRefresh();
 }
diff --git a/Engine/Source/Editor/UMGEditor/Private/SUMGEditorTreeItem.cpp b/Engine/Source/Editor/UMGEditor/Private/SUMGEditorTreeItem.cpp
--- a/Engine/Source/Editor/UMGEditor/Private/SUMGEditorTreeItem.cpp
+++ b/Engine/Source/Editor/UMGEditor/Private/SUMGEditorTreeItem.cpp
@@ -23,6 +23,7 @@ void SUMGEditorTreeItem::Construct(const FArguments& InArgs, TSharedPtr<FBluepri
 	ChildSlot
 	[
 		SNew(STextBlock)
+		.ToolTipText(this, &SUMGEditorTreeItem::GetItemTooltipText)
 		.Text(this, &SUMGEditorTreeItem::GetItemText)
 	];
 }
@@ -32,6 +33,11 @@ FText SUMGEditorTreeItem::GetItemText() const
 	return FText::FromString(Item->GetName());
 }
 
+FString SUMGEditorTreeItem::GetItemTooltipText() const
+{
+	return Item->GetDetailedInfo();
+}
+
 void SUMGEditorTreeItem::OnDragEnter(const FGeometry& MyGeometry, const FDragDropEvent& DragDropEvent)
 {
 	//TSharedPtr<FDragDropOperation> Operation = DragDropEvent.GetOperation();
@@ -76,7 +82,8 @@ FReply SUMGEditorTreeItem::OnDrop(const FGeometry& MyGeometry, const FDragDropEv
 		{
 			UWidgetBlueprint* BP = CastChecked<UWidgetBlueprint>(BlueprintEditor.Pin()->GetBlueprintObj());
 
-			USlateWrapperComponent* Widget = ConstructObject<USlateWrapperComponent>(DragDropOp->Template->WidgetClass, BP);
+			UClass* WidgetClass = DragDropOp->Template->WidgetClass;
+			USlateWrapperComponent* Widget = ConstructObject<USlateWrapperComponent>(WidgetClass, BP);
 			
 			USlateNonLeafWidgetComponent* Parent = Cast<USlateNonLeafWidgetComponent>(Item);
 			Parent->AddChild(Widget);
diff --git a/Engine/Source/Editor/UMGEditor/Private/SUMGEditorTreeItem.h b/Engine/Source/Editor/UMGEditor/Private/SUMGEditorTreeItem.h
--- a/Engine/Source/Editor/UMGEditor/Private/SUMGEditorTreeItem.h
+++ b/Engine/Source/Editor/UMGEditor/Private/SUMGEditorTreeItem.h
@@ -22,6 +22,7 @@ public:
 	
 private:
 	FText GetItemText() const;
+	FString GetItemTooltipText() const;
 
 	TWeakPtr<class FBlueprintEditor> BlueprintEditor;
 	USlateWrapperComponent* Item;
diff --git a/Engine/Source/Runtime/UMG/Classes/ImageComponent.h b/Engine/Source/Runtime/UMG/Classes/ImageComponent.h
--- a/Engine/Source/Runtime/UMG/Classes/ImageComponent.h
+++ b/Engine/Source/Runtime/UMG/Classes/ImageComponent.h
@@ -11,17 +11,18 @@ class UMG_API UImageComponent : public USlateLeafWidgetComponent
 	GENERATED_UCLASS_BODY()
 
 	/** Image to draw */
-	UPROPERTY(EditAnywhere, BlueprintReadWrite, Category=Appearance)
-	FSlateBrush Image;
+	UPROPERTY(EditAnywhere, BlueprintReadWrite, Category=Appearance, meta=( DisplayThumbnail = "true" ))
+	USlateBrushAsset* Image;
 
 	/** Color and opacity */
 	UPROPERTY(EditAnywhere, BlueprintReadWrite, Category=Appearance)
- 	FSlateColor ColorAndOpacity;
+	FSlateColor ColorAndOpacity;
 
 protected:
 	// USlateWrapperComponent interface
 	virtual TSharedRef<SWidget> RebuildWidget() OVERRIDE;
 	// End of USlateWrapperComponent interface
 
- 	FSlateColor GetColorAndOpacity() const;
+	const FSlateBrush* GetImageBrush() const;
+	FSlateColor GetColorAndOpacity() const;
 };
diff --git a/Engine/Source/Runtime/UMG/Classes/WidgetBlueprint.h b/Engine/Source/Runtime/UMG/Classes/WidgetBlueprint.h
--- a/Engine/Source/Runtime/UMG/Classes/WidgetBlueprint.h
+++ b/Engine/Source/Runtime/UMG/Classes/WidgetBlueprint.h
@@ -15,6 +15,7 @@ public:
 	TArray<class USlateWrapperComponent*> WidgetTemplates;
 
 	virtual void PostLoad() OVERRIDE;
+	virtual void PostLoadSubobjects(FObjectInstancingGraph* OuterInstanceGraph) OVERRIDE;
 	
 #if WITH_EDITOR
 	// UBlueprint interface
diff --git a/Engine/Source/Runtime/UMG/Private/ImageComponent.cpp b/Engine/Source/Runtime/UMG/Private/ImageComponent.cpp
--- a/Engine/Source/Runtime/UMG/Private/ImageComponent.cpp
+++ b/Engine/Source/Runtime/UMG/Private/ImageComponent.cpp
@@ -19,10 +19,21 @@ UImageComponent::UImageComponent(const FPostConstructInitializeProperties& PCIP)
 TSharedRef<SWidget> UImageComponent::RebuildWidget()
 {
 	return SNew(SImage)
-		.Image( &Image )
+		.Image( BIND_UOBJECT_ATTRIBUTE(const FSlateBrush*, GetImageBrush) )
 		.ColorAndOpacity( BIND_UOBJECT_ATTRIBUTE(FSlateColor, GetColorAndOpacity) );
 }
 
+const FSlateBrush* UImageComponent::GetImageBrush() const
+{
+	if ( Image == NULL )
+	{
+		SImage::FArguments ImageDefaults;
+		return ImageDefaults._Image.Get();
+	}
+
+	return &Image->Brush;
+}
+
 FSlateColor UImageComponent::GetColorAndOpacity() const
 {
 	return ColorAndOpacity;
diff --git a/Engine/Source/Runtime/UMG/Private/SpinningImageComponent.cpp b/Engine/Source/Runtime/UMG/Private/SpinningImageComponent.cpp
--- a/Engine/Source/Runtime/UMG/Private/SpinningImageComponent.cpp
+++ b/Engine/Source/Runtime/UMG/Private/SpinningImageComponent.cpp
@@ -17,7 +17,7 @@ USpinningImageComponent::USpinningImageComponent(const FPostConstructInitializeP
 TSharedRef<SWidget> USpinningImageComponent::RebuildWidget()
 {
 	return SNew(SSpinningImage)
-		.Image( &Image )
+		.Image( BIND_UOBJECT_ATTRIBUTE(const FSlateBrush*, GetImageBrush) )
 		.ColorAndOpacity( BIND_UOBJECT_ATTRIBUTE(FSlateColor, GetColorAndOpacity) )
 		.Period( Period );
 }
diff --git a/Engine/Source/Runtime/UMG/Private/WidgetBlueprint.cpp b/Engine/Source/Runtime/UMG/Private/WidgetBlueprint.cpp
--- a/Engine/Source/Runtime/UMG/Private/WidgetBlueprint.cpp
+++ b/Engine/Source/Runtime/UMG/Private/WidgetBlueprint.cpp
@@ -26,25 +26,45 @@ bool UWidgetBlueprint::ValidateGeneratedClass(const UClass* InClass)
 {
 	bool Result = Super::ValidateGeneratedClass(InClass);
 
-	//for ( auto CompIt = Blueprint->Timelines.CreateConstIterator(); CompIt; ++CompIt )
-	//{
-	//	const UTimelineTemplate* Template = ( *CompIt );
-	//	if ( !ensure(Template && ( Template->GetOuter() == GeneratedClass )) )
-	//	{
-	//		return false;
-	//	}
-	//}
-
-	//for ( auto CompIt = GeneratedClass->Timelines.CreateConstIterator(); CompIt; ++CompIt )
-	//{
-	//	const UTimelineTemplate* Template = ( *CompIt );
-	//	if ( !ensure(Template && ( Template->GetOuter() == GeneratedClass )) )
-	//	{
-	//		return false;
-	//	}
-	//}
+	const UWidgetBlueprintGeneratedClass* GeneratedClass = Cast<const UWidgetBlueprintGeneratedClass>(InClass);
+	if ( !ensure(GeneratedClass) )
+	{
+		return false;
+	}
+	const UWidgetBlueprint* Blueprint = Cast<UWidgetBlueprint>(GetBlueprintFromClass(GeneratedClass));
+	if ( !ensure(Blueprint) )
+	{
+		return false;
+	}
+
+	for ( const USlateWrapperComponent* Template : Blueprint->WidgetTemplates )
+	{
+		if ( !ensure(Template && ( Template->GetOuter() == GeneratedClass )) )
+		{
+			return false;
+		}
+	}
+
+	for ( const USlateWrapperComponent* Template : GeneratedClass->WidgetTemplates )
+	{
+		if ( !ensure(Template && ( Template->GetOuter() == GeneratedClass )) )
+		{
+			return false;
+		}
+	}
 
 	return Result;
 }
 
+void UWidgetBlueprint::PostLoadSubobjects(FObjectInstancingGraph* OuterInstanceGraph)
+{
+	Super::PostLoadSubobjects(OuterInstanceGraph);
+
+	UWidgetBlueprintGeneratedClass* BPGClass = Cast<UWidgetBlueprintGeneratedClass>(*GeneratedClass);
+	if ( BPGClass )
+	{
+		BPGClass->WidgetTemplates = WidgetTemplates;
+	}
+}
+
 #undef LOCTEXT_NAMESPACE 
\ No newline at end of file
