UMG - Continued work on the designer, introducing WidgetTree to make cloning and reparenting a cleaner process without so many GLEO issues.

diff --git a/Engine/Source/Editor/KismetCompiler/Private/KismetCompiler.cpp b/Engine/Source/Editor/KismetCompiler/Private/KismetCompiler.cpp
--- a/Engine/Source/Editor/KismetCompiler/Private/KismetCompiler.cpp
+++ b/Engine/Source/Editor/KismetCompiler/Private/KismetCompiler.cpp
@@ -3269,10 +3269,16 @@ void FKismetCompilerContext::Compile()
 
 	if (bIsFullCompile && !Blueprint->bIsRegeneratingOnLoad)
 	{
-		UBlueprint::ValidateGeneratedClass(NewClass);
+		bool Result = ValidateGeneratedClass(NewClass);
+		// TODO What do we do if validation fails?
 	}
 }
 
+bool FKismetCompilerContext::ValidateGeneratedClass(UBlueprintGeneratedClass* Class)
+{
+	return UBlueprint::ValidateGeneratedClass(Class);
+}
+
 const UK2Node_FunctionEntry* FKismetCompilerContext::FindLocalEntryPoint(const UFunction* Function) const
 {
 	for (int32 i = 0; i < FunctionList.Num(); ++i)
diff --git a/Engine/Source/Editor/KismetCompiler/Public/KismetCompiler.h b/Engine/Source/Editor/KismetCompiler/Public/KismetCompiler.h
--- a/Engine/Source/Editor/KismetCompiler/Public/KismetCompiler.h
+++ b/Engine/Source/Editor/KismetCompiler/Public/KismetCompiler.h
@@ -345,6 +345,9 @@ protected:
 	/** Ensures that all function graphs have valid names for compilation/replication */
 	void ValidateFunctionGraphNames();
 
+	/** Validates the generated class */
+	virtual bool ValidateGeneratedClass(UBlueprintGeneratedClass* Class);
+
 private:
 	/**
 	 * Handles creating a new event node for a given output on a timeline node utilizing the named function
diff --git a/Engine/Source/Editor/UMGEditor/Private/SUMGEditorTree.cpp b/Engine/Source/Editor/UMGEditor/Private/SUMGEditorTree.cpp
--- a/Engine/Source/Editor/UMGEditor/Private/SUMGEditorTree.cpp
+++ b/Engine/Source/Editor/UMGEditor/Private/SUMGEditorTree.cpp
@@ -155,43 +155,44 @@ FReply SUMGEditorTree::CreateTestUI()
 	//UButtonComponent* Button = Cast<UButtonComponent>(SCSEditor->AddNewComponent(UButtonComponent::StaticClass(), NULL));
 
 	UWidgetBlueprint* BP = GetBlueprint();
+	TArray<USlateWrapperComponent*>& WidgetTemplates = BP->WidgetTree->WidgetTemplates;
 
-	if ( BP->WidgetTemplates.Num() > 0 )
+	if ( WidgetTemplates.Num() > 0 )
 	{
-		UVerticalBoxComponent* Vertical = CastChecked<UVerticalBoxComponent>(BP->WidgetTemplates[2]);
+		UVerticalBoxComponent* Vertical = CastChecked<UVerticalBoxComponent>(WidgetTemplates[2]);
 
-		UButtonComponent* NewButton = ConstructObject<UButtonComponent>(UButtonComponent::StaticClass(), BP);
+		UButtonComponent* NewButton = ConstructObject<UButtonComponent>(UButtonComponent::StaticClass(), BP->WidgetTree);
 		NewButton->ButtonText = FText::FromString("New Button");
 
 		Vertical->AddChild(NewButton);
 
-		BP->WidgetTemplates.Add(NewButton);
+		WidgetTemplates.Add(NewButton);
 	}
 	else
 	{
-		UCanvasPanelComponent* Canvas = ConstructObject<UCanvasPanelComponent>(UCanvasPanelComponent::StaticClass(), BP);
-		UBorderComponent* Border = ConstructObject<UBorderComponent>(UBorderComponent::StaticClass(), BP);
-		UVerticalBoxComponent* Vertical = ConstructObject<UVerticalBoxComponent>(UVerticalBoxComponent::StaticClass(), BP);
-		UButtonComponent* Button1 = ConstructObject<UButtonComponent>(UButtonComponent::StaticClass(), BP);
+		UCanvasPanelComponent* Canvas = ConstructObject<UCanvasPanelComponent>(UCanvasPanelComponent::StaticClass(), BP->WidgetTree);
+		UBorderComponent* Border = ConstructObject<UBorderComponent>(UBorderComponent::StaticClass(), BP->WidgetTree);
+		UVerticalBoxComponent* Vertical = ConstructObject<UVerticalBoxComponent>(UVerticalBoxComponent::StaticClass(), BP->WidgetTree);
+		UButtonComponent* Button1 = ConstructObject<UButtonComponent>(UButtonComponent::StaticClass(), BP->WidgetTree);
 		Button1->ButtonText = FText::FromString("Button 1");
-		UButtonComponent* Button2 = ConstructObject<UButtonComponent>(UButtonComponent::StaticClass(), BP);
+		UButtonComponent* Button2 = ConstructObject<UButtonComponent>(UButtonComponent::StaticClass(), BP->WidgetTree);
 		Button2->ButtonText = FText::FromString("Button 2");
-		UButtonComponent* Button3 = ConstructObject<UButtonComponent>(UButtonComponent::StaticClass(), BP);
+		UButtonComponent* Button3 = ConstructObject<UButtonComponent>(UButtonComponent::StaticClass(), BP->WidgetTree);
 		Button3->ButtonText = FText::FromString("Button 3");
-		UButtonComponent* Button4 = ConstructObject<UButtonComponent>(UButtonComponent::StaticClass(), BP);
+		UButtonComponent* Button4 = ConstructObject<UButtonComponent>(UButtonComponent::StaticClass(), BP->WidgetTree);
 
-		UTextBlockComponent* Text1 = ConstructObject<UTextBlockComponent>(UTextBlockComponent::StaticClass(), BP);
+		UTextBlockComponent* Text1 = ConstructObject<UTextBlockComponent>(UTextBlockComponent::StaticClass(), BP->WidgetTree);
 		Text1->Text = FText::FromString(TEXT("Button!"));
 		Button4->SetContent(Text1);
 
-		BP->WidgetTemplates.Add(Canvas);
-		BP->WidgetTemplates.Add(Border);
-		BP->WidgetTemplates.Add(Vertical);
-		BP->WidgetTemplates.Add(Button1);
-		BP->WidgetTemplates.Add(Button2);
-		BP->WidgetTemplates.Add(Button3);
-		BP->WidgetTemplates.Add(Button4);
-		BP->WidgetTemplates.Add(Text1);
+		WidgetTemplates.Add(Canvas);
+		WidgetTemplates.Add(Border);
+		WidgetTemplates.Add(Vertical);
+		WidgetTemplates.Add(Button1);
+		WidgetTemplates.Add(Button2);
+		WidgetTemplates.Add(Button3);
+		WidgetTemplates.Add(Button4);
+		WidgetTemplates.Add(Text1);
 
 		UCanvasPanelSlot* Slot = Canvas->AddSlot(Border);
 		Slot->Size.X = 200;
@@ -217,9 +218,11 @@ void SUMGEditorTree::RefreshTree()
 	RootWidgets.Reset();
 
 	UWidgetBlueprint* Blueprint = GetBlueprint();
-	if ( Blueprint->WidgetTemplates.Num() > 0 )
+	TArray<USlateWrapperComponent*>& WidgetTemplates = Blueprint->WidgetTree->WidgetTemplates;
+
+	if ( WidgetTemplates.Num() > 0 )
 	{
-		RootWidgets.Add(Blueprint->WidgetTemplates[0]);
+		RootWidgets.Add(WidgetTemplates[0]);
 	}
 
 	WidgetTreeView->RequestTreeRefresh();
diff --git a/Engine/Source/Editor/UMGEditor/Private/SUMGEditorTreeItem.cpp b/Engine/Source/Editor/UMGEditor/Private/SUMGEditorTreeItem.cpp
--- a/Engine/Source/Editor/UMGEditor/Private/SUMGEditorTreeItem.cpp
+++ b/Engine/Source/Editor/UMGEditor/Private/SUMGEditorTreeItem.cpp
@@ -83,12 +83,12 @@ FReply SUMGEditorTreeItem::OnDrop(const FGeometry& MyGeometry, const FDragDropEv
 			UWidgetBlueprint* BP = CastChecked<UWidgetBlueprint>(BlueprintEditor.Pin()->GetBlueprintObj());
 
 			UClass* WidgetClass = DragDropOp->Template->WidgetClass;
-			USlateWrapperComponent* Widget = ConstructObject<USlateWrapperComponent>(WidgetClass, BP);
+			USlateWrapperComponent* Widget = ConstructObject<USlateWrapperComponent>(WidgetClass, BP->WidgetTree);
 			
 			USlateNonLeafWidgetComponent* Parent = Cast<USlateNonLeafWidgetComponent>(Item);
 			Parent->AddChild(Widget);
 
-			BP->WidgetTemplates.Add(Widget);
+			BP->WidgetTree->WidgetTemplates.Add(Widget);
 
 			FBlueprintEditorUtils::MarkBlueprintAsStructurallyModified(BP);
 		}
diff --git a/Engine/Source/Editor/UMGEditor/Private/WidgetBlueprintCompiler.cpp b/Engine/Source/Editor/UMGEditor/Private/WidgetBlueprintCompiler.cpp
--- a/Engine/Source/Editor/UMGEditor/Private/WidgetBlueprintCompiler.cpp
+++ b/Engine/Source/Editor/UMGEditor/Private/WidgetBlueprintCompiler.cpp
@@ -29,7 +29,7 @@ void FWidgetBlueprintCompiler::ValidateWidgetNames()
 		}
 	}
 
-	for ( USlateWrapperComponent* Widget : Blueprint->WidgetTemplates )
+	for ( USlateWrapperComponent* Widget : Blueprint->WidgetTree->WidgetTemplates )
 	{
 		if ( ParentBPNameValidator.IsValid() && ParentBPNameValidator->IsValid(Widget->GetName()) != EValidatorResult::Ok )
 		{
@@ -70,9 +70,9 @@ void FWidgetBlueprintCompiler::CleanAndSanitizeClass(UBlueprintGeneratedClass* C
 
 	UWidgetBlueprint* Blueprint = WidgetBlueprint();
 
-	if ( 0 != Blueprint->WidgetTemplates.Num() )
+	if ( 0 != Blueprint->WidgetTree->WidgetTemplates.Num() )
 	{
-		ClassSubObjects.RemoveAllSwap(FCullTemplateObjectsHelper<USlateWrapperComponent>(Blueprint->WidgetTemplates));
+		ClassSubObjects.RemoveAllSwap(FCullTemplateObjectsHelper<USlateWrapperComponent>(Blueprint->WidgetTree->WidgetTemplates));
 	}
 }
 
@@ -85,7 +85,7 @@ void FWidgetBlueprintCompiler::CreateClassVariablesFromBlueprint()
 	ValidateWidgetNames();
 
 	int32 i = 0;
-	for ( USlateWrapperComponent* Widget : Blueprint->WidgetTemplates )
+	for ( USlateWrapperComponent* Widget : Blueprint->WidgetTree->WidgetTemplates )
 	{
 		FString VariableName = Widget->GetClass()->GetName() + TEXT("_") + FString::FromInt(i);
 
@@ -101,23 +101,17 @@ void FWidgetBlueprintCompiler::CreateClassVariablesFromBlueprint()
 
 		i++;
 	}
-
-	printf("Compile");
 }
 
 void FWidgetBlueprintCompiler::FinishCompilingClass(UClass* Class)
 {
-	Super::FinishCompilingClass(Class);
-
-	UClass* ParentClass = Class->GetSuperClass();
 	UWidgetBlueprint* Blueprint = WidgetBlueprint();
 
-	{
-		UWidgetBlueprintGeneratedClass* BPGClass = CastChecked<UWidgetBlueprintGeneratedClass>(Class);
+	UWidgetBlueprintGeneratedClass* BPGClass = CastChecked<UWidgetBlueprintGeneratedClass>(Class);
+	BPGClass->WidgetTree = NULL;
+	BPGClass->WidgetTree = Blueprint->WidgetTree;
 
-		BPGClass->WidgetTemplates.Empty();
-		BPGClass->WidgetTemplates = Blueprint->WidgetTemplates;
-	}
+	Super::FinishCompilingClass(Class);
 }
 
 void FWidgetBlueprintCompiler::Compile()
@@ -151,3 +145,11 @@ void FWidgetBlueprintCompiler::SpawnNewClass(const FString& NewClassName)
 	}
 	NewClass = NewWidgetBlueprintClass;
 }
+
+bool FWidgetBlueprintCompiler::ValidateGeneratedClass(UBlueprintGeneratedClass* Class)
+{
+	bool SuperResult = Super::ValidateGeneratedClass(Class);
+	bool Result = UWidgetBlueprint::ValidateGeneratedClass(Class);
+
+	return SuperResult && Result;
+}
diff --git a/Engine/Source/Editor/UMGEditor/Private/WidgetBlueprintCompiler.h b/Engine/Source/Editor/UMGEditor/Private/WidgetBlueprintCompiler.h
--- a/Engine/Source/Editor/UMGEditor/Private/WidgetBlueprintCompiler.h
+++ b/Engine/Source/Editor/UMGEditor/Private/WidgetBlueprintCompiler.h
@@ -31,6 +31,7 @@ protected:
 	virtual void CreateClassVariablesFromBlueprint() OVERRIDE;
 	virtual void FinishCompilingClass(UClass* Class) OVERRIDE;
 	virtual void CleanAndSanitizeClass(UBlueprintGeneratedClass* ClassToClean, UObject*& OldCDO) OVERRIDE;
+	virtual bool ValidateGeneratedClass(UBlueprintGeneratedClass* Class) OVERRIDE;
 	// End FKismetCompilerContext
 
 protected:
diff --git a/Engine/Source/Editor/UMGEditor/Private/WidgetBlueprintEditorUtils.cpp b/Engine/Source/Editor/UMGEditor/Private/WidgetBlueprintEditorUtils.cpp
new file mode 100644
--- /dev/null
+++ b/Engine/Source/Editor/UMGEditor/Private/WidgetBlueprintEditorUtils.cpp
@@ -0,0 +1,78 @@
+// Copyright 1998-2014 Epic Games, Inc. All Rights Reserved.
+
+#include "UMGEditorPrivatePCH.h"
+
+#include "WidgetBlueprintEditorUtils.h"
+#include "Kismet2NameValidators.h"
+#include "BlueprintEditorUtils.h"
+
+bool FWidgetBlueprintEditorUtils::RenameWidget(UBlueprint* Blueprint, const FName& OldName, const FName& NewName)
+{
+	check(Blueprint);
+
+	bool bRenamed = false;
+
+	TSharedPtr<INameValidatorInterface> NameValidator = MakeShareable(new FKismetNameValidator(Blueprint));
+	const FString NewTemplateName = UTimelineTemplate::TimelineVariableNameToTemplateName(NewName);
+	// NewName should be already validated. But one must make sure that NewTemplateName is also unique.
+	const bool bUniqueNameForTemplate = (EValidatorResult::Ok == NameValidator->IsValid(NewTemplateName));
+
+	UTimelineTemplate* Template = Blueprint->FindTimelineTemplateByVariableName(NewName);
+	if ((Template == NULL) && bUniqueNameForTemplate)
+	{
+		Template = Blueprint->FindTimelineTemplateByVariableName(OldName);
+		if (Template)
+		{
+			Blueprint->Modify();
+			Template->Modify();
+
+			const FString NewNameStr = NewName.ToString();
+			const FString OldNameStr = OldName.ToString();
+
+			TArray<UK2Node_Variable*> TimelineVarNodes;
+			FBlueprintEditorUtils::GetAllNodesOfClass<UK2Node_Variable>(Blueprint, TimelineVarNodes);
+			for(int32 It = 0; It < TimelineVarNodes.Num(); It++)
+			{
+				UK2Node_Variable* TestNode = TimelineVarNodes[It];
+				if(TestNode && (OldName == TestNode->GetVarName()))
+				{
+					UEdGraphPin* TestPin = TestNode->FindPin(OldNameStr);
+					if(TestPin && (UTimelineComponent::StaticClass() == TestPin->PinType.PinSubCategoryObject.Get()))
+					{
+						TestNode->Modify();
+						if(TestNode->VariableReference.IsSelfContext())
+						{
+							TestNode->VariableReference.SetSelfMember(NewName);
+						}
+						else
+						{
+							//TODO:
+							UClass* ParentClass = TestNode->VariableReference.GetMemberParentClass((UClass*)NULL);
+							TestNode->VariableReference.SetExternalMember(NewName, ParentClass);
+						}
+						TestPin->Modify();
+						TestPin->PinName = NewNameStr;
+					}
+				}
+			}
+
+			Blueprint->Timelines.Remove(Template);
+			
+			UObject* ExistingObject = StaticFindObject(NULL, Template->GetOuter(), *NewTemplateName, true);
+			if (ExistingObject != Template && ExistingObject != NULL)
+			{
+				ExistingObject->Rename(*MakeUniqueObjectName(ExistingObject->GetOuter(), ExistingObject->GetClass(), ExistingObject->GetFName()).ToString());
+			}
+			Template->Rename(*NewTemplateName, Template->GetOuter(), (Blueprint->bIsRegeneratingOnLoad ? REN_ForceNoResetLoaders : REN_None));
+			Blueprint->Timelines.Add(Template);
+
+			// Validate child blueprints and adjust variable names to avoid a potential name collision
+			FBlueprintEditorUtils::ValidateBlueprintChildVariables(Blueprint, NewName);
+
+			// Refresh references and flush editors
+			FBlueprintEditorUtils::MarkBlueprintAsStructurallyModified(Blueprint);
+			bRenamed = true;
+		}
+	}
+	return bRenamed;
+}
\ No newline at end of file
diff --git a/Engine/Source/Editor/UMGEditor/Private/WidgetBlueprintEditorUtils.h b/Engine/Source/Editor/UMGEditor/Private/WidgetBlueprintEditorUtils.h
new file mode 100644
--- /dev/null
+++ b/Engine/Source/Editor/UMGEditor/Private/WidgetBlueprintEditorUtils.h
@@ -0,0 +1,12 @@
+// Copyright 1998-2014 Epic Games, Inc. All Rights Reserved.
+
+#pragma once
+
+//////////////////////////////////////////////////////////////////////////
+// FWidgetBlueprintEditorUtils
+
+class FWidgetBlueprintEditorUtils
+{
+public:
+	static bool RenameWidget(UBlueprint* Blueprint, const FName& OldName, const FName& NewName);
+};
diff --git a/Engine/Source/Runtime/UMG/Classes/WidgetBlueprint.h b/Engine/Source/Runtime/UMG/Classes/WidgetBlueprint.h
--- a/Engine/Source/Runtime/UMG/Classes/WidgetBlueprint.h
+++ b/Engine/Source/Runtime/UMG/Classes/WidgetBlueprint.h
@@ -10,9 +10,9 @@ class UMG_API UWidgetBlueprint : public UBlueprint
 	GENERATED_UCLASS_BODY()
 
 public:
-	/** Array of templates for timelines that should be created */
+	/** A tree of the widget templates to be created */
 	UPROPERTY()
-	TArray<class USlateWrapperComponent*> WidgetTemplates;
+	class UWidgetTree* WidgetTree;
 
 	virtual void PostLoad() OVERRIDE;
 	virtual void PostLoadSubobjects(FObjectInstancingGraph* OuterInstanceGraph) OVERRIDE;
diff --git a/Engine/Source/Runtime/UMG/Classes/WidgetBlueprintGeneratedClass.h b/Engine/Source/Runtime/UMG/Classes/WidgetBlueprintGeneratedClass.h
--- a/Engine/Source/Runtime/UMG/Classes/WidgetBlueprintGeneratedClass.h
+++ b/Engine/Source/Runtime/UMG/Classes/WidgetBlueprintGeneratedClass.h
@@ -9,8 +9,9 @@ class UMG_API UWidgetBlueprintGeneratedClass : public UBlueprintGeneratedClass
 {
 	GENERATED_UCLASS_BODY()
 
+	/** A tree of the widget templates to be created */
 	UPROPERTY()
-	TArray<class USlateWrapperComponent*> WidgetTemplates;
+	class UWidgetTree* WidgetTree;
 
 	virtual void CreateComponentsForActor(AActor* Actor) const OVERRIDE;
 };
\ No newline at end of file
diff --git a/Engine/Source/Runtime/UMG/Classes/WidgetTree.h b/Engine/Source/Runtime/UMG/Classes/WidgetTree.h
new file mode 100644
--- /dev/null
+++ b/Engine/Source/Runtime/UMG/Classes/WidgetTree.h
@@ -0,0 +1,15 @@
+// Copyright 1998-2014 Epic Games, Inc. All Rights Reserved.
+
+#pragma once
+
+#include "WidgetTree.generated.h"
+
+UCLASS()
+class UMG_API UWidgetTree : public UObject
+{
+	GENERATED_UCLASS_BODY()
+
+	/** Array of templates for widgets */
+	UPROPERTY()
+	TArray<class USlateWrapperComponent*> WidgetTemplates;
+};
diff --git a/Engine/Source/Runtime/UMG/Private/WidgetBlueprint.cpp b/Engine/Source/Runtime/UMG/Private/WidgetBlueprint.cpp
--- a/Engine/Source/Runtime/UMG/Private/WidgetBlueprint.cpp
+++ b/Engine/Source/Runtime/UMG/Private/WidgetBlueprint.cpp
@@ -10,16 +10,19 @@
 UWidgetBlueprint::UWidgetBlueprint(const FPostConstructInitializeProperties& PCIP)
 	: Super(PCIP)
 {
+	WidgetTree = ConstructObject<UWidgetTree>(UWidgetTree::StaticClass(), this);
 }
 
 void UWidgetBlueprint::PostLoad()
 {
 	Super::PostLoad();
 
-	for ( USlateWrapperComponent* Widget : WidgetTemplates )
+#if WITH_EDITOR
+	for ( USlateWrapperComponent* Widget : WidgetTree->WidgetTemplates )
 	{
 		Widget->ConnectEditorData();
 	}
+#endif
 }
 
 bool UWidgetBlueprint::ValidateGeneratedClass(const UClass* InClass)
@@ -37,21 +40,15 @@ bool UWidgetBlueprint::ValidateGeneratedClass(const UClass* InClass)
 		return false;
 	}
 
-	for ( const USlateWrapperComponent* Template : Blueprint->WidgetTemplates )
-	{
-		if ( !ensure(Template && ( Template->GetOuter() == GeneratedClass )) )
-		{
-			return false;
-		}
-	}
+	//if ( !ensure(Blueprint->WidgetTree && ( Blueprint->WidgetTree->GetOuter() == GeneratedClass )) )
+	//{
+	//	return false;
+	//}
 
-	for ( const USlateWrapperComponent* Template : GeneratedClass->WidgetTemplates )
-	{
-		if ( !ensure(Template && ( Template->GetOuter() == GeneratedClass )) )
-		{
-			return false;
-		}
-	}
+	//if ( !ensure(GeneratedClass->WidgetTree && ( GeneratedClass->WidgetTree->GetOuter() == GeneratedClass )) )
+	//{
+	//	return false;
+	//}
 
 	return Result;
 }
@@ -63,7 +60,7 @@ void UWidgetBlueprint::PostLoadSubobjects(FObjectInstancingGraph* OuterInstanceG
 	UWidgetBlueprintGeneratedClass* BPGClass = Cast<UWidgetBlueprintGeneratedClass>(*GeneratedClass);
 	if ( BPGClass )
 	{
-		BPGClass->WidgetTemplates = WidgetTemplates;
+		//BPGClass->WidgetTree = WidgetTree;
 	}
 }
 
diff --git a/Engine/Source/Runtime/UMG/Private/WidgetBlueprintGeneratedClass.cpp b/Engine/Source/Runtime/UMG/Private/WidgetBlueprintGeneratedClass.cpp
--- a/Engine/Source/Runtime/UMG/Private/WidgetBlueprintGeneratedClass.cpp
+++ b/Engine/Source/Runtime/UMG/Private/WidgetBlueprintGeneratedClass.cpp
@@ -10,6 +10,7 @@
 UWidgetBlueprintGeneratedClass::UWidgetBlueprintGeneratedClass(const FPostConstructInitializeProperties& PCIP)
 	: Super(PCIP)
 {
+	WidgetTree = ConstructObject<UWidgetTree>(UWidgetTree::StaticClass(), this);
 }
 
 void UWidgetBlueprintGeneratedClass::CreateComponentsForActor(AActor* Actor) const
@@ -19,28 +20,15 @@ void UWidgetBlueprintGeneratedClass::CreateComponentsForActor(AActor* Actor) con
 	// Duplicate the graph, keeping track of what was duplicated
 	TMap<UObject*, UObject*> DuplicatedObjectList;
 
-	TArray<USlateWrapperComponent*> ClonedWidgets;
+	FObjectDuplicationParameters Parameters(const_cast<UWidgetTree*>( WidgetTree ), Actor);
+	Parameters.CreatedObjects = &DuplicatedObjectList;
 
-	// Iterate over each timeline template
-	for ( const USlateWrapperComponent* Template : WidgetTemplates )
-	{
-		FObjectDuplicationParameters Parameters(const_cast<USlateWrapperComponent*>( Template ), Actor);
-		Parameters.CreatedObjects = &DuplicatedObjectList;
-
-		USlateWrapperComponent* ClonedWidget = CastChecked<USlateWrapperComponent>(StaticDuplicateObjectEx(Parameters));
-		ClonedWidgets.Add(ClonedWidget);
-	}
-
-	// Replace references to old classes and instances on this object with the corresponding new ones
-	for ( USlateWrapperComponent* Widget : ClonedWidgets )
-	{
-		FArchiveReplaceObjectRef<UObject> ReplaceInCDOAr(Widget, DuplicatedObjectList, /*bNullPrivateRefs=*/ false, /*bIgnoreOuterRef=*/ false, /*bIgnoreArchetypeRef=*/ false);
-	}
+	UWidgetTree* ClonedTree = CastChecked<UWidgetTree>(StaticDuplicateObjectEx(Parameters));
 
 	AUserWidget* WidgetActor = CastChecked<AUserWidget>(Actor);
 
 	int32 i = 0;
-	for ( USlateWrapperComponent* Widget : ClonedWidgets )
+	for ( USlateWrapperComponent* Widget : ClonedTree->WidgetTemplates )
 	{
 		// Not fatal if NULL, but shouldn't happen
 		if ( !ensure(Widget != NULL) )
@@ -78,8 +66,6 @@ void UWidgetBlueprintGeneratedClass::CreateComponentsForActor(AActor* Actor) con
 #if WITH_EDITOR
 		Widget->ConnectEditorData();
 #endif
-
-		i++;
 	}
 }
 
diff --git a/Engine/Source/Runtime/UMG/Private/WidgetTree.cpp b/Engine/Source/Runtime/UMG/Private/WidgetTree.cpp
new file mode 100644
--- /dev/null
+++ b/Engine/Source/Runtime/UMG/Private/WidgetTree.cpp
@@ -0,0 +1,11 @@
+// Copyright 1998-2014 Epic Games, Inc. All Rights Reserved.
+
+#include "UMGPrivatePCH.h"
+
+/////////////////////////////////////////////////////
+// UWidgetTree
+
+UWidgetTree::UWidgetTree(const FPostConstructInitializeProperties& PCIP)
+	: Super(PCIP)
+{
+}
