Fix for build break caused by CL#2055369.

diff --git a/Engine/Source/Editor/BlueprintGraph/Classes/K2Node_Variable.h b/Engine/Source/Editor/BlueprintGraph/Classes/K2Node_Variable.h
--- a/Engine/Source/Editor/BlueprintGraph/Classes/K2Node_Variable.h
+++ b/Engine/Source/Editor/BlueprintGraph/Classes/K2Node_Variable.h
@@ -2,6 +2,7 @@
 
 
 #pragma once
+#include "K2Node.h"
 #include "K2Node_Variable.generated.h"
 
 UENUM()
diff --git a/Engine/Source/Editor/UMGEditor/Private/WidgetBlueprintEditorUtils.cpp b/Engine/Source/Editor/UMGEditor/Private/WidgetBlueprintEditorUtils.cpp
--- a/Engine/Source/Editor/UMGEditor/Private/WidgetBlueprintEditorUtils.cpp
+++ b/Engine/Source/Editor/UMGEditor/Private/WidgetBlueprintEditorUtils.cpp
@@ -5,6 +5,7 @@
 #include "WidgetBlueprintEditorUtils.h"
 #include "Kismet2NameValidators.h"
 #include "BlueprintEditorUtils.h"
+#include "K2Node_Variable.h"
 
 bool FWidgetBlueprintEditorUtils::RenameWidget(UBlueprint* Blueprint, const FName& OldName, const FName& NewName)
 {
