UMG - The widget tree now factories new widgets to ensure they are constructed and added to the tree correctly.  Custumizable variable flags, fixing binding bugs.  Text input events.

diff --git a/Engine/Source/Editor/UMGEditor/Private/SUMGEditorTree.cpp b/Engine/Source/Editor/UMGEditor/Private/SUMGEditorTree.cpp
--- a/Engine/Source/Editor/UMGEditor/Private/SUMGEditorTree.cpp
+++ b/Engine/Source/Editor/UMGEditor/Private/SUMGEditorTree.cpp
@@ -150,10 +150,6 @@ void SUMGEditorTree::WidgetHierarchy_OnSelectionChanged(USlateWrapperComponent*
 
 FReply SUMGEditorTree::CreateTestUI()
 {
-	//TSharedRef<ISCSEditor> SCSEditor = BlueprintEditor.Pin()->GetSCSEditorModel();
-	//UCanvasPanelComponent* Canvas = Cast<UCanvasPanelComponent>(SCSEditor->AddNewComponent(UCanvasPanelComponent::StaticClass(), NULL));
-	//UButtonComponent* Button = Cast<UButtonComponent>(SCSEditor->AddNewComponent(UButtonComponent::StaticClass(), NULL));
-
 	UWidgetBlueprint* BP = GetBlueprint();
 	TArray<USlateWrapperComponent*>& WidgetTemplates = BP->WidgetTree->WidgetTemplates;
 
@@ -161,39 +157,28 @@ FReply SUMGEditorTree::CreateTestUI()
 	{
 		UVerticalBoxComponent* Vertical = CastChecked<UVerticalBoxComponent>(WidgetTemplates[2]);
 
-		UButtonComponent* NewButton = ConstructObject<UButtonComponent>(UButtonComponent::StaticClass(), BP->WidgetTree);
+		UButtonComponent* NewButton = BP->WidgetTree->ConstructWidget<UButtonComponent>(UButtonComponent::StaticClass());
 		NewButton->ButtonText = FText::FromString("New Button");
 
 		Vertical->AddChild(NewButton);
-
-		WidgetTemplates.Add(NewButton);
 	}
 	else
 	{
-		UCanvasPanelComponent* Canvas = ConstructObject<UCanvasPanelComponent>(UCanvasPanelComponent::StaticClass(), BP->WidgetTree);
-		UBorderComponent* Border = ConstructObject<UBorderComponent>(UBorderComponent::StaticClass(), BP->WidgetTree);
-		UVerticalBoxComponent* Vertical = ConstructObject<UVerticalBoxComponent>(UVerticalBoxComponent::StaticClass(), BP->WidgetTree);
-		UButtonComponent* Button1 = ConstructObject<UButtonComponent>(UButtonComponent::StaticClass(), BP->WidgetTree);
+		UCanvasPanelComponent* Canvas = BP->WidgetTree->ConstructWidget<UCanvasPanelComponent>(UCanvasPanelComponent::StaticClass());
+		UBorderComponent* Border = BP->WidgetTree->ConstructWidget<UBorderComponent>(UBorderComponent::StaticClass());
+		UVerticalBoxComponent* Vertical = BP->WidgetTree->ConstructWidget<UVerticalBoxComponent>(UVerticalBoxComponent::StaticClass());
+		UButtonComponent* Button1 = BP->WidgetTree->ConstructWidget<UButtonComponent>(UButtonComponent::StaticClass());
 		Button1->ButtonText = FText::FromString("Button 1");
-		UButtonComponent* Button2 = ConstructObject<UButtonComponent>(UButtonComponent::StaticClass(), BP->WidgetTree);
+		UButtonComponent* Button2 = BP->WidgetTree->ConstructWidget<UButtonComponent>(UButtonComponent::StaticClass());
 		Button2->ButtonText = FText::FromString("Button 2");
-		UButtonComponent* Button3 = ConstructObject<UButtonComponent>(UButtonComponent::StaticClass(), BP->WidgetTree);
+		UButtonComponent* Button3 = BP->WidgetTree->ConstructWidget<UButtonComponent>(UButtonComponent::StaticClass());
 		Button3->ButtonText = FText::FromString("Button 3");
-		UButtonComponent* Button4 = ConstructObject<UButtonComponent>(UButtonComponent::StaticClass(), BP->WidgetTree);
+		UButtonComponent* Button4 = BP->WidgetTree->ConstructWidget<UButtonComponent>(UButtonComponent::StaticClass());
 
-		UTextBlockComponent* Text1 = ConstructObject<UTextBlockComponent>(UTextBlockComponent::StaticClass(), BP->WidgetTree);
+		UTextBlockComponent* Text1 = BP->WidgetTree->ConstructWidget<UTextBlockComponent>(UTextBlockComponent::StaticClass());
 		Text1->Text = FText::FromString(TEXT("Button!"));
 		Button4->SetContent(Text1);
 
-		WidgetTemplates.Add(Canvas);
-		WidgetTemplates.Add(Border);
-		WidgetTemplates.Add(Vertical);
-		WidgetTemplates.Add(Button1);
-		WidgetTemplates.Add(Button2);
-		WidgetTemplates.Add(Button3);
-		WidgetTemplates.Add(Button4);
-		WidgetTemplates.Add(Text1);
-
 		UCanvasPanelSlot* Slot = Canvas->AddSlot(Border);
 		Slot->Size.X = 200;
 		Slot->Size.Y = 300;
diff --git a/Engine/Source/Editor/UMGEditor/Private/WidgetBlueprintCompiler.cpp b/Engine/Source/Editor/UMGEditor/Private/WidgetBlueprintCompiler.cpp
--- a/Engine/Source/Editor/UMGEditor/Private/WidgetBlueprintCompiler.cpp
+++ b/Engine/Source/Editor/UMGEditor/Private/WidgetBlueprintCompiler.cpp
@@ -84,10 +84,15 @@ void FWidgetBlueprintCompiler::CreateClassVariablesFromBlueprint()
 
 	ValidateWidgetNames();
 
-	int32 i = 0;
 	for ( USlateWrapperComponent* Widget : Blueprint->WidgetTree->WidgetTemplates )
 	{
-		FString VariableName = Widget->GetClass()->GetName() + TEXT("_") + FString::FromInt(i);
+		// Skip non-variable widgets
+		if ( !Widget->bIsVariable )
+		{
+			continue;
+		}
+
+		FString VariableName = Widget->GetName();
 
 		FEdGraphPinType WidgetPinType(Schema->PC_Object, TEXT(""), Widget->GetClass(), false, false);
 		UProperty* WidgetProperty = CreateVariable(FName(*VariableName), WidgetPinType);
@@ -98,8 +103,6 @@ void FWidgetBlueprintCompiler::CreateClassVariablesFromBlueprint()
 
 			WidgetToMemberVariableMap.Add(Widget, WidgetProperty);
 		}
-
-		i++;
 	}
 }
 
diff --git a/Engine/Source/Runtime/UMG/Private/Components/BorderComponent.cpp b/Engine/Source/Runtime/UMG/Private/Components/BorderComponent.cpp
--- a/Engine/Source/Runtime/UMG/Private/Components/BorderComponent.cpp
+++ b/Engine/Source/Runtime/UMG/Private/Components/BorderComponent.cpp
@@ -10,6 +10,8 @@
 UBorderComponent::UBorderComponent(const FPostConstructInitializeProperties& PCIP)
 	: Super(PCIP)
 {
+	bIsVariable = false;
+
 	SBorder::FArguments BorderDefaults;
 
 	HorizontalAlignment = BorderDefaults._HAlign;
diff --git a/Engine/Source/Runtime/UMG/Private/Components/BoxPanelComponent.cpp b/Engine/Source/Runtime/UMG/Private/Components/BoxPanelComponent.cpp
--- a/Engine/Source/Runtime/UMG/Private/Components/BoxPanelComponent.cpp
+++ b/Engine/Source/Runtime/UMG/Private/Components/BoxPanelComponent.cpp
@@ -105,6 +105,7 @@ protected:
 UBoxPanelComponent::UBoxPanelComponent(const FPostConstructInitializeProperties& PCIP)
 	: Super(PCIP)
 {
+	bIsVariable = false;
 }
 
 TSharedRef<SWidget> UBoxPanelComponent::RebuildWidget()
diff --git a/Engine/Source/Runtime/UMG/Private/Components/CanvasPanelComponent.cpp b/Engine/Source/Runtime/UMG/Private/Components/CanvasPanelComponent.cpp
--- a/Engine/Source/Runtime/UMG/Private/Components/CanvasPanelComponent.cpp
+++ b/Engine/Source/Runtime/UMG/Private/Components/CanvasPanelComponent.cpp
@@ -37,6 +37,7 @@ class SFixedSizeCanvas : public SCanvas
 UCanvasPanelComponent::UCanvasPanelComponent(const FPostConstructInitializeProperties& PCIP)
 	: Super(PCIP)
 {
+	bIsVariable = false;
 	DesiredCanvasSize = FVector2D(128.0f, 128.0f);
 }
 
diff --git a/Engine/Source/Runtime/UMG/Private/Components/EditableText.cpp b/Engine/Source/Runtime/UMG/Private/Components/EditableText.cpp
--- a/Engine/Source/Runtime/UMG/Private/Components/EditableText.cpp
+++ b/Engine/Source/Runtime/UMG/Private/Components/EditableText.cpp
@@ -46,7 +46,14 @@ TSharedRef<SWidget> UEditableText::RebuildWidget()
 		.SelectAllTextWhenFocused(SelectAllTextWhenFocused)
 		.RevertTextOnEscape(RevertTextOnEscape)
 		.ClearKeyboardFocusOnCommit(ClearKeyboardFocusOnCommit)
-		.SelectAllTextOnCommit(SelectAllTextOnCommit);
+		.SelectAllTextOnCommit(SelectAllTextOnCommit)
+		.OnTextChanged(BIND_UOBJECT_DELEGATE(FOnTextChanged, SlateOnTextChanged))
+		;
+}
+
+void UEditableText::SlateOnTextChanged(const FText& Text)
+{
+	OnTextChanged.Broadcast(Text);
 }
 
 /////////////////////////////////////////////////////
diff --git a/Engine/Source/Runtime/UMG/Private/Components/EditableTextBlockComponent.cpp b/Engine/Source/Runtime/UMG/Private/Components/EditableTextBlockComponent.cpp
--- a/Engine/Source/Runtime/UMG/Private/Components/EditableTextBlockComponent.cpp
+++ b/Engine/Source/Runtime/UMG/Private/Components/EditableTextBlockComponent.cpp
@@ -49,7 +49,14 @@ TSharedRef<SWidget> UEditableTextBlockComponent::RebuildWidget()
 		.SelectAllTextWhenFocused(SelectAllTextWhenFocused)
 		.RevertTextOnEscape(RevertTextOnEscape)
 		.ClearKeyboardFocusOnCommit(ClearKeyboardFocusOnCommit)
-		.SelectAllTextOnCommit(SelectAllTextOnCommit);
+		.SelectAllTextOnCommit(SelectAllTextOnCommit)
+		.OnTextChanged(BIND_UOBJECT_DELEGATE(FOnTextChanged, SlateOnTextChanged))
+		;
+}
+
+void UEditableTextBlockComponent::SlateOnTextChanged(const FText& Text)
+{
+	OnTextChanged.Broadcast(Text);
 }
 
 /////////////////////////////////////////////////////
diff --git a/Engine/Source/Runtime/UMG/Private/Components/GridPanelComponent.cpp b/Engine/Source/Runtime/UMG/Private/Components/GridPanelComponent.cpp
--- a/Engine/Source/Runtime/UMG/Private/Components/GridPanelComponent.cpp
+++ b/Engine/Source/Runtime/UMG/Private/Components/GridPanelComponent.cpp
@@ -8,6 +8,7 @@
 UGridPanelComponent::UGridPanelComponent(const FPostConstructInitializeProperties& PCIP)
 	: Super(PCIP)
 {
+	bIsVariable = false;
 }
 
 TSharedRef<SWidget> UGridPanelComponent::RebuildWidget()
diff --git a/Engine/Source/Runtime/UMG/Private/Components/HorizontalBoxComponent.cpp b/Engine/Source/Runtime/UMG/Private/Components/HorizontalBoxComponent.cpp
--- a/Engine/Source/Runtime/UMG/Private/Components/HorizontalBoxComponent.cpp
+++ b/Engine/Source/Runtime/UMG/Private/Components/HorizontalBoxComponent.cpp
@@ -8,6 +8,7 @@
 UHorizontalBoxComponent::UHorizontalBoxComponent(const FPostConstructInitializeProperties& PCIP)
 	: Super(PCIP)
 {
+	bIsVariable = false;
 }
 
 int32 UHorizontalBoxComponent::GetChildrenCount() const
diff --git a/Engine/Source/Runtime/UMG/Private/Components/ScrollBoxComponent.cpp b/Engine/Source/Runtime/UMG/Private/Components/ScrollBoxComponent.cpp
--- a/Engine/Source/Runtime/UMG/Private/Components/ScrollBoxComponent.cpp
+++ b/Engine/Source/Runtime/UMG/Private/Components/ScrollBoxComponent.cpp
@@ -10,6 +10,7 @@
 UScrollBoxComponent::UScrollBoxComponent(const FPostConstructInitializeProperties& PCIP)
 	: Super(PCIP)
 {
+	bIsVariable = false;
 }
 
 TSharedRef<SWidget> UScrollBoxComponent::RebuildWidget()
diff --git a/Engine/Source/Runtime/UMG/Private/Components/SlateWrapperComponent.cpp b/Engine/Source/Runtime/UMG/Private/Components/SlateWrapperComponent.cpp
--- a/Engine/Source/Runtime/UMG/Private/Components/SlateWrapperComponent.cpp
+++ b/Engine/Source/Runtime/UMG/Private/Components/SlateWrapperComponent.cpp
@@ -9,6 +9,7 @@ USlateWrapperComponent::USlateWrapperComponent(const FPostConstructInitializePro
 	: Super(PCIP)
 {
 	bIsEnabled = true;
+	bIsVariable = true;
 }
 
 bool USlateWrapperComponent::GetIsEnabled() const
@@ -36,6 +37,11 @@ void USlateWrapperComponent::OnUnregister()
 	Super::OnUnregister();
 }
 
+void USlateWrapperComponent::PostEditChangeProperty(struct FPropertyChangedEvent& PropertyChangedEvent)
+{
+	// TODO Rename?
+}
+
 TSharedRef<SWidget> USlateWrapperComponent::GetWidget()
 {
 	if ( !MyWidget.IsValid() )
diff --git a/Engine/Source/Runtime/UMG/Private/Components/SpacerComponent.cpp b/Engine/Source/Runtime/UMG/Private/Components/SpacerComponent.cpp
--- a/Engine/Source/Runtime/UMG/Private/Components/SpacerComponent.cpp
+++ b/Engine/Source/Runtime/UMG/Private/Components/SpacerComponent.cpp
@@ -11,6 +11,7 @@ USpacerComponent::USpacerComponent(const FPostConstructInitializeProperties& PCI
 	: Super(PCIP)
 	, Size(1.0f, 1.0f)
 {
+	bIsVariable = false;
 }
 
 TSharedRef<SWidget> USpacerComponent::RebuildWidget()
diff --git a/Engine/Source/Runtime/UMG/Private/Components/TextBlockComponent.cpp b/Engine/Source/Runtime/UMG/Private/Components/TextBlockComponent.cpp
--- a/Engine/Source/Runtime/UMG/Private/Components/TextBlockComponent.cpp
+++ b/Engine/Source/Runtime/UMG/Private/Components/TextBlockComponent.cpp
@@ -10,6 +10,8 @@
 UTextBlockComponent::UTextBlockComponent(const FPostConstructInitializeProperties& PCIP)
 	: Super(PCIP)
 {
+	bIsVariable = false;
+
 	Text = LOCTEXT("TextBlockDefaultValue", "Text Block");
 	ShadowOffset = FVector2D(1.0f, 1.0f);
 	ColorAndOpacity = FLinearColor::White;
diff --git a/Engine/Source/Runtime/UMG/Private/Components/VerticalBoxComponent.cpp b/Engine/Source/Runtime/UMG/Private/Components/VerticalBoxComponent.cpp
--- a/Engine/Source/Runtime/UMG/Private/Components/VerticalBoxComponent.cpp
+++ b/Engine/Source/Runtime/UMG/Private/Components/VerticalBoxComponent.cpp
@@ -8,6 +8,7 @@
 UVerticalBoxComponent::UVerticalBoxComponent(const FPostConstructInitializeProperties& PCIP)
 	: Super(PCIP)
 {
+	bIsVariable = false;
 }
 
 int32 UVerticalBoxComponent::GetChildrenCount() const
diff --git a/Engine/Source/Runtime/UMG/Private/WidgetBlueprintGeneratedClass.cpp b/Engine/Source/Runtime/UMG/Private/WidgetBlueprintGeneratedClass.cpp
--- a/Engine/Source/Runtime/UMG/Private/WidgetBlueprintGeneratedClass.cpp
+++ b/Engine/Source/Runtime/UMG/Private/WidgetBlueprintGeneratedClass.cpp
@@ -26,8 +26,8 @@ void UWidgetBlueprintGeneratedClass::CreateComponentsForActor(AActor* Actor) con
 	UWidgetTree* ClonedTree = CastChecked<UWidgetTree>(StaticDuplicateObjectEx(Parameters));
 
 	AUserWidget* WidgetActor = CastChecked<AUserWidget>(Actor);
+	UClass* ActorClass = Actor->GetClass();
 
-	int32 i = 0;
 	for ( USlateWrapperComponent* Widget : ClonedTree->WidgetTemplates )
 	{
 		// Not fatal if NULL, but shouldn't happen
@@ -38,23 +38,13 @@ void UWidgetBlueprintGeneratedClass::CreateComponentsForActor(AActor* Actor) con
 
 		FName NewName = *( FString::Printf(TEXT("WidgetComponent__%d"), Actor->SerializedComponents.Num()) );
 
-		////Serialize object properties using write/read operations.
-		//TArray<uint8> SavedProperties;
-		//FObjectWriter Writer(ComponentTemplate, SavedProperties);
-		//FObjectReader(ClonedComponent, SavedProperties);
+		FString VariableName = Widget->GetName();
 
-
-		FString VariableName = Widget->GetClass()->GetName() + TEXT("_") + FString::FromInt(i);
-
-
-		//USlateWrapperComponent* Widget = ConstructObject<USlateWrapperComponent>(Template->GetClass(), Actor, NewName);
 		Widget->bCreatedByConstructionScript = true; // Indicate it comes from a blueprint so it gets cleared when we rerun construction scripts
 		Actor->SerializedComponents.Add(Widget); // Add to array so it gets saved
 		Widget->SetNetAddressable();	// This component has a stable name that can be referenced for replication
 
 		// Find property with the same name as the template and assign the new Timeline to it
-		UClass* ActorClass = Actor->GetClass();
-		//UTimelineTemplate::TimelineTemplateNameToVariableName(Template->GetFName())
 		UObjectPropertyBase* Prop = FindField<UObjectPropertyBase>(ActorClass, *VariableName);
 		if ( Prop )
 		{
diff --git a/Engine/Source/Runtime/UMG/Public/Blueprint/WidgetTree.h b/Engine/Source/Runtime/UMG/Public/Blueprint/WidgetTree.h
--- a/Engine/Source/Runtime/UMG/Public/Blueprint/WidgetTree.h
+++ b/Engine/Source/Runtime/UMG/Public/Blueprint/WidgetTree.h
@@ -12,4 +12,20 @@ class UMG_API UWidgetTree : public UObject
 	/** Array of templates for widgets */
 	UPROPERTY()
 	TArray<class USlateWrapperComponent*> WidgetTemplates;
+
+	void RenameWidget(USlateWrapperComponent* Widget, FString& NewName)
+	{
+		Widget->Rename(*NewName);
+
+		// TODO Update nodes in the blueprint!
+	}
+
+	template< class T >
+	T* ConstructWidget(TSubclassOf<class USlateWrapperComponent> WidgetType)
+	{
+		USlateWrapperComponent* Widget = (USlateWrapperComponent*)ConstructObject<class USlateWrapperComponent>(WidgetType, this);
+		WidgetTemplates.Add(Widget);
+
+		return (T*)Widget;
+	}
 };
diff --git a/Engine/Source/Runtime/UMG/Public/Components/ButtonComponent.h b/Engine/Source/Runtime/UMG/Public/Components/ButtonComponent.h
--- a/Engine/Source/Runtime/UMG/Public/Components/ButtonComponent.h
+++ b/Engine/Source/Runtime/UMG/Public/Components/ButtonComponent.h
@@ -4,7 +4,7 @@
 
 #include "ButtonComponent.generated.h"
 
-DECLARE_DYNAMIC_MULTICAST_DELEGATE( FOnButtonComponentClicked );
+DECLARE_DYNAMIC_MULTICAST_DELEGATE(FOnButtonClickedEvent);
 
 DECLARE_DYNAMIC_DELEGATE_RetVal( FText, FOnGetText );
 
@@ -61,7 +61,7 @@ public:
 
 	/** Called when the button is clicked */
 	UPROPERTY(BlueprintAssignable)
-	FOnButtonComponentClicked OnClicked;
+	FOnButtonClickedEvent OnClicked;
 
 	/** Called when the controller has instigated damage in any way */
 	UPROPERTY(BlueprintReadWrite, Category=Appearance)
diff --git a/Engine/Source/Runtime/UMG/Public/Components/EditableText.h b/Engine/Source/Runtime/UMG/Public/Components/EditableText.h
--- a/Engine/Source/Runtime/UMG/Public/Components/EditableText.h
+++ b/Engine/Source/Runtime/UMG/Public/Components/EditableText.h
@@ -4,6 +4,9 @@
 
 #include "EditableText.generated.h"
 
+DECLARE_DYNAMIC_MULTICAST_DELEGATE_OneParam(FOnEditableTextChangedEvent, const FText&, Text);
+//DECLARE_DYNAMIC_MULTICAST_DELEGATE_TwoParams(FOnEditableTextCommittedEvent, const FText&, Text, ETextCommit::Type, Type);
+
 /** Editable text box widget */
 UCLASS(meta=(BlueprintSpawnableComponent), ClassGroup=UserInterface)
 class UMG_API UEditableText : public USlateWrapperComponent
@@ -76,8 +79,18 @@ protected:
 	///** Image brush used for the caret (overrides Style) */
 	//SLATE_ATTRIBUTE(const FSlateBrush*, CaretImage)
 
+	/** Called whenever the text is changed interactively by the user */
+	UPROPERTY(BlueprintAssignable)
+	FOnEditableTextChangedEvent OnTextChanged;
+
+	/** Called whenever the text is committed.  This happens when the user presses enter or the text box loses focus. */
+	//UPROPERTY(BlueprintAssignable)
+	//FOnTextCommittedEvent OnTextCommitted;
+
 protected:
 	// USlateWrapperComponent interface
 	virtual TSharedRef<SWidget> RebuildWidget() OVERRIDE;
 	// End of USlateWrapperComponent
+
+	void SlateOnTextChanged(const FText& Text);
 };
diff --git a/Engine/Source/Runtime/UMG/Public/Components/EditableTextBlockComponent.h b/Engine/Source/Runtime/UMG/Public/Components/EditableTextBlockComponent.h
--- a/Engine/Source/Runtime/UMG/Public/Components/EditableTextBlockComponent.h
+++ b/Engine/Source/Runtime/UMG/Public/Components/EditableTextBlockComponent.h
@@ -4,6 +4,9 @@
 
 #include "EditableTextBlockComponent.generated.h"
 
+DECLARE_DYNAMIC_MULTICAST_DELEGATE_OneParam(FOnEditableTextBlock_TextChangedEvent, const FText&, Text);
+//DECLARE_DYNAMIC_MULTICAST_DELEGATE_TwoParams(FOnEditableTextBlock_TextCommittedEvent, const FText&, Text, ETextCommit::Type, Type);
+
 /** Editable text box widget */
 UCLASS(meta=(BlueprintSpawnableComponent), ClassGroup=UserInterface)
 class UMG_API UEditableTextBlockComponent : public USlateWrapperComponent
@@ -73,10 +76,12 @@ protected:
 	bool SelectAllTextOnCommit;
 
 	/** Called whenever the text is changed interactively by the user */
-	//SLATE_EVENT(FOnTextChanged, OnTextChanged)
+	UPROPERTY(BlueprintAssignable)
+	FOnEditableTextBlock_TextChangedEvent OnTextChanged;
 
-	/** Called whenever the text is committed.  This happens when the user presses enter or the text box loses focus. */
-	//SLATE_EVENT(FOnTextCommitted, OnTextCommitted)
+	/** Called whenever the text is committed.  This happens when the user presses enter or the text box loses focus. *///
+	//UPROPERTY(BlueprintAssignable)
+	//FOnTextCommittedEvent OnTextCommitted;
 
 	/** Provide a alternative mechanism for error reporting. */
 	//SLATE_ARGUMENT(TSharedPtr<class IErrorReportingWidget>, ErrorReporting)
@@ -85,4 +90,6 @@ protected:
 	// USlateWrapperComponent interface
 	virtual TSharedRef<SWidget> RebuildWidget() OVERRIDE;
 	// End of USlateWrapperComponent
+
+	void SlateOnTextChanged(const FText& Text);
 };
diff --git a/Engine/Source/Runtime/UMG/Public/Components/SlateWrapperComponent.h b/Engine/Source/Runtime/UMG/Public/Components/SlateWrapperComponent.h
--- a/Engine/Source/Runtime/UMG/Public/Components/SlateWrapperComponent.h
+++ b/Engine/Source/Runtime/UMG/Public/Components/SlateWrapperComponent.h
@@ -11,6 +11,12 @@ class UMG_API USlateWrapperComponent : public UActorComponent
 {
 	GENERATED_UCLASS_BODY()
 
+	UPROPERTY(EditAnywhere, Category=Variable)
+	FString Name;
+
+	UPROPERTY(EditAnywhere, Category=Variable)
+	bool bIsVariable;
+
 #if WITH_EDITOR
 	UPROPERTY(Transient, EditInline, EditAnywhere, BlueprintReadWrite, Category=Layout, meta=( ShowOnlyInnerProperties ))
 	UPanelSlot* Slot;
@@ -37,6 +43,12 @@ class UMG_API USlateWrapperComponent : public UActorComponent
 	virtual void ConnectEditorData() { }
 #endif
 
+#if WITH_EDITOR
+	// UObject interface
+	virtual void PostEditChangeProperty(struct FPropertyChangedEvent& PropertyChangedEvent) OVERRIDE;
+	// End of UObject interface
+#endif
+
 	// Utility methods
 	//@TODO: Should move elsewhere
 	static EVisibility ConvertSerializedVisibilityToRuntime(TEnumAsByte<ESlateVisibility::Type> Input);
