UMG - Working on a more complex templating system to allow for templates that create more than a single widget, like the button that also creates inner text.  Removing text hacks from Button.

diff --git a/Engine/Source/Editor/UMGEditor/Private/SUMGEditorTree.cpp b/Engine/Source/Editor/UMGEditor/Private/SUMGEditorTree.cpp
--- a/Engine/Source/Editor/UMGEditor/Private/SUMGEditorTree.cpp
+++ b/Engine/Source/Editor/UMGEditor/Private/SUMGEditorTree.cpp
@@ -158,7 +158,6 @@ FReply SUMGEditorTree::CreateTestUI()
 		UVerticalBoxComponent* Vertical = CastChecked<UVerticalBoxComponent>(WidgetTemplates[2]);
 
 		UButtonComponent* NewButton = BP->WidgetTree->ConstructWidget<UButtonComponent>(UButtonComponent::StaticClass());
-		NewButton->ButtonText = FText::FromString("New Button");
 
 		Vertical->AddChild(NewButton);
 	}
@@ -168,11 +167,8 @@ FReply SUMGEditorTree::CreateTestUI()
 		UBorderComponent* Border = BP->WidgetTree->ConstructWidget<UBorderComponent>(UBorderComponent::StaticClass());
 		UVerticalBoxComponent* Vertical = BP->WidgetTree->ConstructWidget<UVerticalBoxComponent>(UVerticalBoxComponent::StaticClass());
 		UButtonComponent* Button1 = BP->WidgetTree->ConstructWidget<UButtonComponent>(UButtonComponent::StaticClass());
-		Button1->ButtonText = FText::FromString("Button 1");
 		UButtonComponent* Button2 = BP->WidgetTree->ConstructWidget<UButtonComponent>(UButtonComponent::StaticClass());
-		Button2->ButtonText = FText::FromString("Button 2");
 		UButtonComponent* Button3 = BP->WidgetTree->ConstructWidget<UButtonComponent>(UButtonComponent::StaticClass());
-		Button3->ButtonText = FText::FromString("Button 3");
 		UButtonComponent* Button4 = BP->WidgetTree->ConstructWidget<UButtonComponent>(UButtonComponent::StaticClass());
 
 		UTextBlockComponent* Text1 = BP->WidgetTree->ConstructWidget<UTextBlockComponent>(UTextBlockComponent::StaticClass());
diff --git a/Engine/Source/Editor/UMGEditor/Private/SUMGEditorTreeItem.cpp b/Engine/Source/Editor/UMGEditor/Private/SUMGEditorTreeItem.cpp
--- a/Engine/Source/Editor/UMGEditor/Private/SUMGEditorTreeItem.cpp
+++ b/Engine/Source/Editor/UMGEditor/Private/SUMGEditorTreeItem.cpp
@@ -80,15 +80,12 @@ FReply SUMGEditorTreeItem::OnDrop(const FGeometry& MyGeometry, const FDragDropEv
 	{
 		if ( Item->IsA(USlateNonLeafWidgetComponent::StaticClass()) )
 		{
-			UWidgetBlueprint* BP = CastChecked<UWidgetBlueprint>(BlueprintEditor.Pin()->GetBlueprintObj());
-
-			UClass* WidgetClass = DragDropOp->Template->WidgetClass;
-			USlateWrapperComponent* Widget = ConstructObject<USlateWrapperComponent>(WidgetClass, BP->WidgetTree);
-			
+			UWidgetBlueprint* BP = CastChecked<UWidgetBlueprint>(BlueprintEditor.Pin()->GetBlueprintObj());			
 			USlateNonLeafWidgetComponent* Parent = Cast<USlateNonLeafWidgetComponent>(Item);
-			Parent->AddChild(Widget);
 
-			BP->WidgetTree->WidgetTemplates.Add(Widget);
+			USlateWrapperComponent* Widget = DragDropOp->Template->Create(BP->WidgetTree);
+
+			Parent->AddChild(Widget);
 
 			FBlueprintEditorUtils::MarkBlueprintAsStructurallyModified(BP);
 		}
diff --git a/Engine/Source/Editor/UMGEditor/Private/SUMGEditorWidgetTemplates.cpp b/Engine/Source/Editor/UMGEditor/Private/SUMGEditorWidgetTemplates.cpp
--- a/Engine/Source/Editor/UMGEditor/Private/SUMGEditorWidgetTemplates.cpp
+++ b/Engine/Source/Editor/UMGEditor/Private/SUMGEditorWidgetTemplates.cpp
@@ -15,6 +15,10 @@
 #include "DecoratedDragDropOp.h"
 #include "WidgetTemplateDragDropOp.h"
 
+#include "WidgetTemplate.h"
+#include "WidgetTemplateClass.h"
+#include "WidgetTemplateButton.h"
+
 class SWidgetTemplateItem : public SCompoundWidget
 {
 public:
@@ -27,7 +31,7 @@ public:
 	*
 	* @param InArgs    Declaration from which to construct the widget
 	*/
-	void Construct(const FArguments& InArgs, TSharedPtr<FWidgetTemplateDescriptor> InTemplate)
+	void Construct(const FArguments& InArgs, TSharedPtr<FWidgetTemplate> InTemplate)
 	{
 		Template = InTemplate;
 
@@ -54,7 +58,7 @@ public:
 	}
 
 private:
-	TSharedPtr<FWidgetTemplateDescriptor> Template;
+	TSharedPtr<FWidgetTemplate> Template;
 };
 
 void SUMGEditorWidgetTemplates::Construct(const FArguments& InArgs, TSharedPtr<FBlueprintEditor> InBlueprintEditor, USimpleConstructionScript* InSCS)
@@ -74,7 +78,7 @@ void SUMGEditorWidgetTemplates::Construct(const FArguments& InArgs, TSharedPtr<F
 		+ SVerticalBox::Slot()
 		.FillHeight(1.0f)
 		[
-			SAssignNew(WidgetsListView, SListView< TSharedPtr<FWidgetTemplateDescriptor> >)
+			SAssignNew(WidgetsListView, SListView< TSharedPtr<FWidgetTemplate> >)
 			.OnGenerateRow(this, &SUMGEditorWidgetTemplates::OnGenerateWidgetTemplateItem)
 			.ListItemsSource(&WidgetTemplates)
 			.SelectionMode(ESelectionMode::Single)
@@ -101,26 +105,54 @@ UWidgetBlueprint* SUMGEditorWidgetTemplates::GetBlueprint() const
 void SUMGEditorWidgetTemplates::BuildWidgetList()
 {
 	WidgetTemplates.Reset();
+	WidgetTemplateCategories.Reset();
+
+	BuildClassWidgetList();
+	BuildSpecialWidgetList();
+
+	for ( auto& Entry : WidgetTemplateCategories )
+	{
+		for ( auto& Template : Entry.Value )
+		{
+			WidgetTemplates.Add(Template);
+		}
+	}
+}
 
+void SUMGEditorWidgetTemplates::BuildClassWidgetList()
+{
 	for ( TObjectIterator<UClass> ClassIt; ClassIt; ++ClassIt )
 	{
 		UClass* WidgetClass = *ClassIt;
 		if ( WidgetClass->IsChildOf(USlateWrapperComponent::StaticClass()) )
 		{
-			TSharedPtr<FWidgetTemplateDescriptor> Template = MakeShareable(new FWidgetTemplateDescriptor());
+			TSharedPtr<FWidgetTemplateClass> Template = MakeShareable(new FWidgetTemplateClass(WidgetClass));
 			Template->Name = WidgetClass->GetDisplayNameText();
 			Template->ToolTip = WidgetClass->GetDisplayNameText();
-			Template->WidgetClass = WidgetClass;
 
-			WidgetTemplates.Add(Template);
+			AddWidgetTemplate(Template);
 		}
 	}
 }
 
-TSharedRef<ITableRow> SUMGEditorWidgetTemplates::OnGenerateWidgetTemplateItem(TSharedPtr<FWidgetTemplateDescriptor> Item, const TSharedRef<STableViewBase>& OwnerTable)
+void SUMGEditorWidgetTemplates::BuildSpecialWidgetList()
+{
+	AddWidgetTemplate(MakeShareable(new FWidgetTemplateButton()));
+
+	//TODO Make this pluggable.
+}
+
+void SUMGEditorWidgetTemplates::AddWidgetTemplate(TSharedPtr<FWidgetTemplate> Template)
+{
+	FString Category = Template->GetCategory().ToString();
+	WidgetTemplateArray& Group = WidgetTemplateCategories.FindOrAdd(Category);
+	Group.Add(Template);
+}
+
+TSharedRef<ITableRow> SUMGEditorWidgetTemplates::OnGenerateWidgetTemplateItem(TSharedPtr<FWidgetTemplate> Item, const TSharedRef<STableViewBase>& OwnerTable)
 {
 	return
-		SNew(STableRow< TSharedPtr<FWidgetTemplateDescriptor> >, OwnerTable)
+		SNew(STableRow< TSharedPtr<FWidgetTemplate> >, OwnerTable)
 		.Padding(2.0f)
 		.OnDragDetected(this, &SUMGEditorWidgetTemplates::OnDraggingWidgetTemplateItem)
 		[
@@ -132,7 +164,7 @@ FReply SUMGEditorWidgetTemplates::OnDraggingWidgetTemplateItem(const FGeometry&
 {
 	if ( WidgetsListView->GetNumItemsSelected() > 0 )
 	{
-		TSharedPtr<FWidgetTemplateDescriptor> Template = WidgetsListView->GetSelectedItems()[0];
+		TSharedPtr<FWidgetTemplate> Template = WidgetsListView->GetSelectedItems()[0];
 		return FReply::Handled().BeginDragDrop(FWidgetTemplateDragDropOp::New(Template));
 	}
 
diff --git a/Engine/Source/Editor/UMGEditor/Private/SUMGEditorWidgetTemplates.h b/Engine/Source/Editor/UMGEditor/Private/SUMGEditorWidgetTemplates.h
--- a/Engine/Source/Editor/UMGEditor/Private/SUMGEditorWidgetTemplates.h
+++ b/Engine/Source/Editor/UMGEditor/Private/SUMGEditorWidgetTemplates.h
@@ -6,7 +6,7 @@
 #include "IUMGEditor.h"
 #include "BlueprintEditor.h"
 
-#include "WidgetTemplateDescriptor.h"
+#include "WidgetTemplate.h"
 
 class SUMGEditorWidgetTemplates : public SCompoundWidget
 {
@@ -21,15 +21,21 @@ private:
 	UWidgetBlueprint* GetBlueprint() const;
 
 	void BuildWidgetList();
+	void BuildClassWidgetList();
+	void BuildSpecialWidgetList();
 
-	TSharedRef<ITableRow> OnGenerateWidgetTemplateItem(TSharedPtr<FWidgetTemplateDescriptor> Item, const TSharedRef<STableViewBase>& OwnerTable);
+	TSharedRef<ITableRow> OnGenerateWidgetTemplateItem(TSharedPtr<FWidgetTemplate> Item, const TSharedRef<STableViewBase>& OwnerTable);
 
 	FReply OnDraggingWidgetTemplateItem(const FGeometry& MyGeometry, const FPointerEvent& MouseEvent);
 
 private:
+	void AddWidgetTemplate(TSharedPtr<FWidgetTemplate> Template);
 
 	TWeakPtr<class FBlueprintEditor> BlueprintEditor;
 
-	TArray< TSharedPtr<FWidgetTemplateDescriptor> > WidgetTemplates;
-	TSharedPtr< SListView< TSharedPtr<FWidgetTemplateDescriptor> > > WidgetsListView;
+	typedef TArray< TSharedPtr<FWidgetTemplate> > WidgetTemplateArray;
+	TMap< FString, WidgetTemplateArray > WidgetTemplateCategories;
+	WidgetTemplateArray WidgetTemplates;
+
+	TSharedPtr< SListView< TSharedPtr<FWidgetTemplate> > > WidgetsListView;
 };
diff --git a/Engine/Source/Editor/UMGEditor/Private/WidgetTemplate.cpp b/Engine/Source/Editor/UMGEditor/Private/WidgetTemplate.cpp
new file mode 100644
--- /dev/null
+++ b/Engine/Source/Editor/UMGEditor/Private/WidgetTemplate.cpp
@@ -0,0 +1,9 @@
+// Copyright 1998-2014 Epic Games, Inc. All Rights Reserved.
+
+#include "UMGEditorPrivatePCH.h"
+
+#include "WidgetTemplate.h"
+
+FWidgetTemplate::FWidgetTemplate()
+{
+}
diff --git a/Engine/Source/Editor/UMGEditor/Private/WidgetTemplate.h b/Engine/Source/Editor/UMGEditor/Private/WidgetTemplate.h
new file mode 100644
--- /dev/null
+++ b/Engine/Source/Editor/UMGEditor/Private/WidgetTemplate.h
@@ -0,0 +1,24 @@
+// Copyright 1998-2014 Epic Games, Inc. All Rights Reserved.
+
+#pragma once
+
+#include "SCompoundWidget.h"
+#include "IUMGEditor.h"
+#include "BlueprintEditor.h"
+
+class FWidgetTemplate : public TSharedFromThis<FWidgetTemplate>
+{
+public:
+	FWidgetTemplate();
+
+	virtual FText GetCategory() = 0;
+
+	virtual USlateWrapperComponent* Create(UWidgetTree* Tree) = 0;
+
+public:
+	FText Name;
+
+	FText ToolTip;
+
+	FSlateIcon Icon;
+};
diff --git a/Engine/Source/Editor/UMGEditor/Private/WidgetTemplateButton.cpp b/Engine/Source/Editor/UMGEditor/Private/WidgetTemplateButton.cpp
new file mode 100644
--- /dev/null
+++ b/Engine/Source/Editor/UMGEditor/Private/WidgetTemplateButton.cpp
@@ -0,0 +1,33 @@
+// Copyright 1998-2014 Epic Games, Inc. All Rights Reserved.
+
+#include "UMGEditorPrivatePCH.h"
+
+#include "WidgetTemplateButton.h"
+
+#define LOCTEXT_NAMESPACE "UMGEditor"
+
+FWidgetTemplateButton::FWidgetTemplateButton()
+{
+	Name = LOCTEXT("Button", "Button");
+	ToolTip = LOCTEXT("Button_Tooltip", "Button With Text");
+	//Icon = 
+}
+
+FText FWidgetTemplateButton::GetCategory()
+{
+	return LOCTEXT("Common", "Common");
+}
+
+USlateWrapperComponent* FWidgetTemplateButton::Create(UWidgetTree* Tree)
+{
+	UButtonComponent* Button = Tree->ConstructWidget<UButtonComponent>(UButtonComponent::StaticClass());
+	UTextBlockComponent* Text = Tree->ConstructWidget<UTextBlockComponent>(UTextBlockComponent::StaticClass());
+
+	Text->Text = FText::FromString(TEXT("Button Text"));
+	Text->Font.Size = 10;
+	Button->SetContent(Text);
+
+	return Button;
+}
+
+#undef LOCTEXT_NAMESPACE
diff --git a/Engine/Source/Editor/UMGEditor/Private/WidgetTemplateButton.h b/Engine/Source/Editor/UMGEditor/Private/WidgetTemplateButton.h
new file mode 100644
--- /dev/null
+++ b/Engine/Source/Editor/UMGEditor/Private/WidgetTemplateButton.h
@@ -0,0 +1,15 @@
+// Copyright 1998-2014 Epic Games, Inc. All Rights Reserved.
+
+#pragma once
+
+#include "WidgetTemplate.h"
+
+class FWidgetTemplateButton : public FWidgetTemplate
+{
+public:
+	FWidgetTemplateButton();
+
+	virtual FText GetCategory() OVERRIDE;
+
+	virtual USlateWrapperComponent* Create(UWidgetTree* Tree) OVERRIDE;
+};
diff --git a/Engine/Source/Editor/UMGEditor/Private/WidgetTemplateClass.cpp b/Engine/Source/Editor/UMGEditor/Private/WidgetTemplateClass.cpp
new file mode 100644
--- /dev/null
+++ b/Engine/Source/Editor/UMGEditor/Private/WidgetTemplateClass.cpp
@@ -0,0 +1,25 @@
+// Copyright 1998-2014 Epic Games, Inc. All Rights Reserved.
+
+#include "UMGEditorPrivatePCH.h"
+
+#include "WidgetTemplateClass.h"
+
+#define LOCTEXT_NAMESPACE "UMGEditor"
+
+FWidgetTemplateClass::FWidgetTemplateClass(TSubclassOf<USlateWrapperComponent> InWidgetClass)
+	: WidgetClass(InWidgetClass)
+{
+}
+
+FText FWidgetTemplateClass::GetCategory()
+{
+	// TODO this should be pulled from metadata from the class, that describes its categorization, is it a panel, a common control...etc?
+	return LOCTEXT("Class", "Class");
+}
+
+USlateWrapperComponent* FWidgetTemplateClass::Create(UWidgetTree* Tree)
+{
+	return Tree->ConstructWidget<USlateWrapperComponent>(WidgetClass);
+}
+
+#undef LOCTEXT_NAMESPACE
diff --git a/Engine/Source/Editor/UMGEditor/Private/WidgetTemplateClass.h b/Engine/Source/Editor/UMGEditor/Private/WidgetTemplateClass.h
new file mode 100644
--- /dev/null
+++ b/Engine/Source/Editor/UMGEditor/Private/WidgetTemplateClass.h
@@ -0,0 +1,18 @@
+// Copyright 1998-2014 Epic Games, Inc. All Rights Reserved.
+
+#pragma once
+
+#include "WidgetTemplate.h"
+
+class FWidgetTemplateClass : public FWidgetTemplate
+{
+public:
+	FWidgetTemplateClass(TSubclassOf<USlateWrapperComponent> InWidgetClass);
+
+	virtual FText GetCategory() OVERRIDE;
+
+	virtual USlateWrapperComponent* Create(UWidgetTree* Tree) OVERRIDE;
+
+protected:
+	TSubclassOf<USlateWrapperComponent> WidgetClass;
+};
diff --git a/Engine/Source/Editor/UMGEditor/Private/WidgetTemplateDescriptor.h b/Engine/Source/Editor/UMGEditor/Private/WidgetTemplateDescriptor.h
deleted file mode 100644
diff --git a/Engine/Source/Editor/UMGEditor/Private/WidgetTemplateDragDropOp.cpp b/Engine/Source/Editor/UMGEditor/Private/WidgetTemplateDragDropOp.cpp
--- a/Engine/Source/Editor/UMGEditor/Private/WidgetTemplateDragDropOp.cpp
+++ b/Engine/Source/Editor/UMGEditor/Private/WidgetTemplateDragDropOp.cpp
@@ -6,7 +6,7 @@
 
 #define LOCTEXT_NAMESPACE "UMG"
 
-TSharedRef<FWidgetTemplateDragDropOp> FWidgetTemplateDragDropOp::New(const TSharedPtr<FWidgetTemplateDescriptor>& InTemplate)
+TSharedRef<FWidgetTemplateDragDropOp> FWidgetTemplateDragDropOp::New(const TSharedPtr<FWidgetTemplate>& InTemplate)
 {
 	TSharedRef<FWidgetTemplateDragDropOp> Operation = MakeShareable(new FWidgetTemplateDragDropOp());
 	Operation->Template = InTemplate;
diff --git a/Engine/Source/Editor/UMGEditor/Private/WidgetTemplateDragDropOp.h b/Engine/Source/Editor/UMGEditor/Private/WidgetTemplateDragDropOp.h
--- a/Engine/Source/Editor/UMGEditor/Private/WidgetTemplateDragDropOp.h
+++ b/Engine/Source/Editor/UMGEditor/Private/WidgetTemplateDragDropOp.h
@@ -3,16 +3,16 @@
 #pragma once
 
 #include "DecoratedDragDropOp.h"
-#include "WidgetTemplateDescriptor.h"
+#include "WidgetTemplate.h"
 
 class FWidgetTemplateDragDropOp : public FDecoratedDragDropOp
 {
 public:
 	DRAG_DROP_OPERATOR_TYPE(FWidgetTemplateDragDropOp, FDecoratedDragDropOp)
 
-	TSharedPtr<FWidgetTemplateDescriptor> Template;
+	TSharedPtr<FWidgetTemplate> Template;
 
-	static TSharedRef<FWidgetTemplateDragDropOp> New(const TSharedPtr<FWidgetTemplateDescriptor>& InTemplate);
+	static TSharedRef<FWidgetTemplateDragDropOp> New(const TSharedPtr<FWidgetTemplate>& InTemplate);
 
 	//virtual TSharedPtr<SWidget> GetDefaultDecorator() const OVERRIDE;
 };
diff --git a/Engine/Source/Runtime/UMG/Private/Components/ButtonComponent.cpp b/Engine/Source/Runtime/UMG/Private/Components/ButtonComponent.cpp
--- a/Engine/Source/Runtime/UMG/Private/Components/ButtonComponent.cpp
+++ b/Engine/Source/Runtime/UMG/Private/Components/ButtonComponent.cpp
@@ -12,15 +12,11 @@ UButtonComponent::UButtonComponent(const FPostConstructInitializeProperties& PCI
 {
 	SButton::FArguments ButtonDefaults;
 
-	ButtonStyle = NULL;
+	Style = NULL;
 	HorizontalAlignment = ButtonDefaults._HAlign;
 	VerticalAlignment = ButtonDefaults._VAlign;
 	ContentPadding = ButtonDefaults._ContentPadding.Get();
 
-	bDisplayText = true;
-	ButtonText = LOCTEXT("ButtonDefaultValue", "This is a button");
-	ButtonTextStyle = NULL;
-
 	DesiredSizeScale = ButtonDefaults._DesiredSizeScale.Get();
 	ContentScale = ButtonDefaults._ContentScale.Get();
 
@@ -32,45 +28,24 @@ TSharedRef<SWidget> UButtonComponent::RebuildWidget()
 {
 	SButton::FArguments ButtonDefaults;
 
-	const FButtonStyle* StylePtr = (ButtonStyle != NULL) ? ButtonStyle->GetStyle<FButtonStyle>() : NULL;
+	const FButtonStyle* StylePtr = ( Style != NULL ) ? Style->GetStyle<FButtonStyle>() : NULL;
 	if (StylePtr == NULL)
 	{
 		StylePtr = ButtonDefaults._ButtonStyle;
 	}
 
-	const FTextBlockStyle* TextStylePtr = ( ButtonTextStyle != NULL ) ? ButtonTextStyle->GetStyle<FTextBlockStyle>() : NULL;
-	if ( TextStylePtr == NULL )
-	{
-		TextStylePtr = ButtonDefaults._TextStyle;
-	}
-
-	if (bDisplayText)
-	{
-		TextAttribute = BIND_UOBJECT_ATTRIBUTE(FText, GetButtonText);
-	}
-	else
-	{
-		TAttribute<FText> NotSet;
-		TextAttribute = NotSet;
-	}
-
 	TSharedRef<SButton> NewButton = SNew(SButton)
 		.ButtonStyle(StylePtr)
-		.TextStyle(TextStylePtr)
 		.HAlign(HorizontalAlignment)
 		.VAlign(VerticalAlignment)
 		.ContentPadding(BIND_UOBJECT_ATTRIBUTE(FMargin, GetContentPadding))
-		.Text(TextAttribute)
 		.OnClicked(BIND_UOBJECT_DELEGATE(FOnClicked, SlateOnClickedCallback))
 		.DesiredSizeScale(DesiredSizeScale)
 		.ContentScale(ContentScale)
 		.ButtonColorAndOpacity(BIND_UOBJECT_ATTRIBUTE(FSlateColor, GetButtonColor))
 		.ForegroundColor(BIND_UOBJECT_ATTRIBUTE(FSlateColor, GetForegroundColor));
 
-	if ( Content )
-	{
-		NewButton->SetContent(Content->GetWidget());
-	}
+	NewButton->SetContent(Content ? Content->GetWidget() : SNullWidget::NullWidget);
 	
 	MyButton = NewButton;
 	return NewButton;
@@ -83,24 +58,7 @@ void UButtonComponent::SetContent(USlateWrapperComponent* InContent)
 	TSharedPtr<SButton> Button = MyButton.Pin();
 	if (Button.IsValid())
 	{
-		if ( InContent )
-		{
-			Button->SetContent(InContent->GetWidget());
-		}
-		else
-		{
-			// Removed a child; replace the text if desired
-			if (bDisplayText)
-			{
-				Button->SetContent(SNew(STextBlock)
-					.Text(TextAttribute)
-					.TextStyle( ButtonTextStyle ) );
-			}
-			else
-			{
-				Button->SetContent(SNullWidget::NullWidget);
-			}
-		}
+		Button->SetContent(InContent ? InContent->GetWidget() : SNullWidget::NullWidget);
 	}
 }
 
@@ -109,16 +67,6 @@ FMargin UButtonComponent::GetContentPadding() const
 	return ContentPadding;
 }
 
-FText UButtonComponent::GetButtonText() const
-{
-	return ButtonText;
-}
-
-FText UButtonComponent::GetText_Implementation() const
-{
-	return ButtonText;
-}
-
 FReply UButtonComponent::SlateOnClickedCallback()
 {
 	OnClicked.Broadcast();
diff --git a/Engine/Source/Runtime/UMG/Private/Components/TextBlockComponent.cpp b/Engine/Source/Runtime/UMG/Private/Components/TextBlockComponent.cpp
--- a/Engine/Source/Runtime/UMG/Private/Components/TextBlockComponent.cpp
+++ b/Engine/Source/Runtime/UMG/Private/Components/TextBlockComponent.cpp
@@ -17,6 +17,8 @@ UTextBlockComponent::UTextBlockComponent(const FPostConstructInitializePropertie
 	ColorAndOpacity = FLinearColor::White;
 	ShadowColorAndOpacity = FLinearColor::Transparent;
 
+	Style = NULL;
+
 	// HACK Special font initialization hack since there are no font assets yet for slate.
 	Font = FSlateFontInfo(TEXT("Slate/Fonts/Roboto-Bold.ttf"), 24);
 }
@@ -25,8 +27,16 @@ TSharedRef<SWidget> UTextBlockComponent::RebuildWidget()
 {
 	FString FontPath = FPaths::EngineContentDir() / Font.FontName.ToString();
 
+	const FTextBlockStyle* StylePtr = ( Style != NULL ) ? Style->GetStyle<FTextBlockStyle>() : NULL;
+	if ( StylePtr == NULL )
+	{
+		STextBlock::FArguments Defaults;
+		StylePtr = Defaults._TextStyle;
+	}
+
 	return SNew(STextBlock)
-		.Text( BIND_UOBJECT_ATTRIBUTE(FString, GetText) )
+		.Text(BIND_UOBJECT_ATTRIBUTE(FString, GetText))
+		.TextStyle(StylePtr)
 		.Font(FSlateFontInfo(FontPath, Font.Size))
 		.ColorAndOpacity( BIND_UOBJECT_ATTRIBUTE(FSlateColor, GetColorAndOpacity) )
 		.ShadowOffset(ShadowOffset)
diff --git a/Engine/Source/Runtime/UMG/Public/Components/ButtonComponent.h b/Engine/Source/Runtime/UMG/Public/Components/ButtonComponent.h
--- a/Engine/Source/Runtime/UMG/Public/Components/ButtonComponent.h
+++ b/Engine/Source/Runtime/UMG/Public/Components/ButtonComponent.h
@@ -6,8 +6,6 @@
 
 DECLARE_DYNAMIC_MULTICAST_DELEGATE(FOnButtonClickedEvent);
 
-DECLARE_DYNAMIC_DELEGATE_RetVal( FText, FOnGetText );
-
 /** Buttons are clickable widgets */
 UCLASS(meta=(BlueprintSpawnableComponent), ClassGroup=UserInterface)
 class UMG_API UButtonComponent : public UContentWidget
@@ -16,8 +14,8 @@ class UMG_API UButtonComponent : public UContentWidget
 
 public:
 	/** Style of the button */
-	UPROPERTY(EditAnywhere, BlueprintReadOnly, Category=Style)
-	USlateWidgetStyleAsset* ButtonStyle;
+	UPROPERTY(EditAnywhere, BlueprintReadOnly, Category=Style, meta=( DisplayThumbnail = "true" ))
+	USlateWidgetStyleAsset* Style;
 
 	/** Horizontal positioning of the content within the button */
 	UPROPERTY(EditAnywhere, BlueprintReadOnly, Category=Appearance)
@@ -31,18 +29,6 @@ public:
 	UPROPERTY(EditAnywhere, BlueprintReadWrite, Category=Appearance)
 	FMargin ContentPadding;
 
-	/** Should the button have text content inside? */
-	UPROPERTY(EditAnywhere, BlueprintReadOnly, Category=Appearance)
-	bool bDisplayText;
-
-	/** Text to display as the content of the button */
-	UPROPERTY(EditAnywhere, BlueprintReadWrite, Category=Appearance, meta=(EditCondition="bDisplayText"))
-	FText ButtonText;
-
-	/** Text to display as the content of the button */
-	UPROPERTY(EditAnywhere, BlueprintReadOnly, Category=Style, meta=( EditCondition="bDisplayText" ))
-	USlateWidgetStyleAsset* ButtonTextStyle;
-
 	/** The scaling factor for the button border */
 	UPROPERTY(EditAnywhere, BlueprintReadOnly, Category=Appearance)
 	FVector2D DesiredSizeScale;
@@ -63,20 +49,11 @@ public:
 	UPROPERTY(BlueprintAssignable)
 	FOnButtonClickedEvent OnClicked;
 
-	/** Called when the controller has instigated damage in any way */
-	UPROPERTY(BlueprintReadWrite, Category=Appearance)
-	FOnGetText OnGetText;
-
-	/** Called when the controller has instigated damage in any way */
-	UFUNCTION(BlueprintNativeEvent, Category=Appearance)
-	FText GetText() const;
-
 	// UContentWidget interface
 	virtual void SetContent(USlateWrapperComponent* Content) OVERRIDE;
 	// End UContentWidget interface
 
 protected:
-	TAttribute<FText> TextAttribute;
 	TWeakPtr<class SButton> MyButton;
 
 protected:
@@ -85,7 +62,6 @@ protected:
 	// End of USlateWrapperComponent interface
 
 	FMargin GetContentPadding() const;
-	FText GetButtonText() const;
 	FSlateColor GetButtonColor() const;
 	FSlateColor GetForegroundColor() const;
 
diff --git a/Engine/Source/Runtime/UMG/Public/Components/TextBlockComponent.h b/Engine/Source/Runtime/UMG/Public/Components/TextBlockComponent.h
--- a/Engine/Source/Runtime/UMG/Public/Components/TextBlockComponent.h
+++ b/Engine/Source/Runtime/UMG/Public/Components/TextBlockComponent.h
@@ -25,6 +25,10 @@ public:
 	UPROPERTY(EditAnywhere, Category=Appearance)
 	FLinearColor ShadowColorAndOpacity;
 
+	/** Text to display as the content of the button */
+	UPROPERTY(EditAnywhere, BlueprintReadOnly, Category=Style, meta=( DisplayThumbnail = "true" ))
+	USlateWidgetStyleAsset* Style;
+
 protected:
 	// USlateWrapperComponent interface
 	virtual TSharedRef<SWidget> RebuildWidget() OVERRIDE;
