UMG - Selection from designer. Adding drag / drop to the designer.  Drag handles, no effect yet.  Fixed several issues with picking on the designer.  Little cleanup on the throbber items, lots of TODOs.

diff --git a/Engine/Source/Editor/UMGEditor/Private/SUMGDesigner.cpp b/Engine/Source/Editor/UMGEditor/Private/SUMGDesigner.cpp
--- a/Engine/Source/Editor/UMGEditor/Private/SUMGDesigner.cpp
+++ b/Engine/Source/Editor/UMGEditor/Private/SUMGDesigner.cpp
@@ -6,6 +6,8 @@
 #include "BlueprintEditor.h"
 #include "SKismetInspector.h"
 
+#include "WidgetTemplateDragDropOp.h"
+
 #define LOCTEXT_NAMESPACE "UMG"
 
 //class SDesignerWidget : public SBorder
@@ -61,7 +63,7 @@ static bool LocateWidgetsUnderCursor_Helper(FArrangedWidget& Candidate, FVector2
 
 		// Check to see if we were asked to still allow children to be hit test visible
 		bool bHitChildWidget = false;
-		if ( ( Candidate.Widget->GetVisibility().AreChildrenHitTestVisible() ) != 0 )
+		//if ( ( Candidate.Widget->GetVisibility().AreChildrenHitTestVisible() ) != 0 || OutWidgetsUnderCursor. )
 		{
 			FArrangedChildren ArrangedChildren(OutWidgetsUnderCursor.GetFilter());
 			Candidate.Widget->ArrangeChildren(Candidate.Geometry, ArrangedChildren);
@@ -75,7 +77,7 @@ static bool LocateWidgetsUnderCursor_Helper(FArrangedWidget& Candidate, FVector2
 		}
 
 		// If we hit a child widget or we hit our candidate widget then we'll append our widgets
-		const bool bHitCandidateWidget = Candidate.Widget->GetVisibility().IsHitTestVisible();
+		const bool bHitCandidateWidget = OutWidgetsUnderCursor.Accepts(Candidate.Widget->GetVisibility());
 		bHitAnyWidget = bHitChildWidget || bHitCandidateWidget;
 		if ( !bHitAnyWidget )
 		{
@@ -171,9 +173,9 @@ void SUMGDesigner::ShowDetailsForObjects(TArray<USlateWrapperComponent*> Widgets
 	BlueprintEditor.Pin()->GetInspector()->ShowDetailsForObjects(InspectorObjects, Options);
 }
 
-FReply SUMGDesigner::OnMouseButtonDown(const FGeometry& MyGeometry, const FPointerEvent& MouseEvent)
+USlateWrapperComponent* SUMGDesigner::GetTemplateAtCursor(const FGeometry& MyGeometry, const FPointerEvent& MouseEvent, FArrangedWidget& ArrangedWidget)
 {
-	FVector2D LocalMouseCoordinates = MyGeometry.AbsoluteToLocal(MouseEvent.GetScreenSpacePosition());
+	//@TODO UMG Make it so you can request dropable widgets only, to find the first parentable.
 
 	FArrangedChildren Children(EVisibility::All);
 
@@ -193,14 +195,40 @@ FReply SUMGDesigner::OnMouseButtonDown(const FGeometry& MyGeometry, const FPoint
 			PreviewHandle = WidgetActor->GetWidgetHandle(Child.Widget);
 			if ( PreviewHandle )
 			{
-				SelectedWidget = Child.Widget;
+				ArrangedWidget = Child;
 				break;
 			}
 		}
 
 		UWidgetBlueprint* Blueprint = GetBlueprint();
-		// TODO we have the widgets runtime handle, we need to look at the handle's name and find the corresponding
-		// design time handle and select it!
+
+		if ( PreviewHandle )
+		{
+			FString Name = PreviewHandle->GetName();
+			USlateWrapperComponent* Template = Blueprint->WidgetTree->FindWidget(Name);
+			return Template;
+		}
+	}
+
+	return NULL;
+}
+
+FReply SUMGDesigner::OnMouseButtonDown(const FGeometry& MyGeometry, const FPointerEvent& MouseEvent)
+{
+	FArrangedWidget ArrangedWidget(SNullWidget::NullWidget, FGeometry());
+	USlateWrapperComponent* Template = GetTemplateAtCursor(MyGeometry, MouseEvent, ArrangedWidget);
+
+	if ( Template )
+	{
+		// Set the selected widget so that we can draw the highlight
+		//@TODO UMG Don't store the transient Widget as the selected object, store the template component and lookup the widget from that.
+		SelectedWidget = ArrangedWidget.Widget;
+
+		//@TODO UMG primary FBlueprintEditor needs to be inherited and selection control needs to be centralized.
+		// Set the template as selected in the details panel
+		TArray<USlateWrapperComponent*> SelectedTemplates;
+		SelectedTemplates.Add(Template);
+		ShowDetailsForObjects(SelectedTemplates);
 	}
 
 	return FReply::Handled();
@@ -243,17 +271,56 @@ int32 SUMGDesigner::OnPaint(const FGeometry& AllottedGeometry, const FSlateRect&
 
 		LayerId += 100;
 
-		FSlateDrawElement::MakeBox(
-			OutDrawElements,
-			LayerId,			FPaintGeometry(
+		FPaintGeometry SelectionGeometry(
 			ArrangedWidget.Geometry.AbsolutePosition,
 			ArrangedWidget.Geometry.Size * ArrangedWidget.Geometry.Scale,
-			ArrangedWidget.Geometry.Scale),
+			ArrangedWidget.Geometry.Scale);
+
+		FSlateDrawElement::MakeBox(
+			OutDrawElements,
+			LayerId,
+			SelectionGeometry,
 			FCoreStyle::Get().GetBrush(TEXT("Debug.Border")),
 			MyClippingRect,
 			ESlateDrawEffect::None,
 			Tint
 		);
+
+		float X = SelectionGeometry.DrawPosition.X;
+		float Y = SelectionGeometry.DrawPosition.Y;
+		float Width = SelectionGeometry.DrawSize.X;
+		float Height = SelectionGeometry.DrawSize.Y;
+		
+		TArray<FVector2D> Handles;
+		Handles.Add(FVector2D(X, Y));					// Top - Left
+		Handles.Add(FVector2D(X + Width * 0.5f, Y));	// Top - Middle
+		Handles.Add(FVector2D(X + Width, Y));			// Top - Right
+
+		Handles.Add(FVector2D(X, Y + Height));					// Bottom - Left
+		Handles.Add(FVector2D(X + Width * 0.5f, Y + Height));	// Bottom - Middle
+		Handles.Add(FVector2D(X + Width, Y + Height));			// Bottom - Right
+
+		Handles.Add(FVector2D(X, Y + Height * 0.5f));			// Middle - Left
+		Handles.Add(FVector2D(X + Width, Y + Height * 0.5f));	// Middle - Right
+
+		const FVector2D HandleSize = FVector2D(10, 10);
+
+		// @TODO UMG - Don't use the curve editors brushes
+		const FSlateBrush* KeyBrush = FEditorStyle::GetBrush("CurveEd.CurveKey");
+		FLinearColor KeyColor = InWidgetStyle.GetColorAndOpacityTint();// IsEditingEnabled() ? InWidgetStyle.GetColorAndOpacityTint() : FLinearColor(0.1f, 0.1f, 0.1f, 1.f);
+
+		for ( FVector2D Handle : Handles )
+		{
+			FSlateDrawElement::MakeBox(
+				OutDrawElements,
+				++LayerId,
+				FPaintGeometry(FVector2D(Handle.X - HandleSize.X * 0.5f, Handle.Y - HandleSize.Y * 0.5f), HandleSize, 1.0f),
+				KeyBrush,
+				MyClippingRect,
+				ESlateDrawEffect::None,
+				KeyBrush->GetTint(InWidgetStyle) * InWidgetStyle.GetColorAndOpacityTint() * KeyColor
+				);
+		}
 	}
 
 	return LayerId;
@@ -293,4 +360,60 @@ void SUMGDesigner::Tick(const FGeometry& AllottedGeometry, const double InCurren
 	SCompoundWidget::Tick(AllottedGeometry, InCurrentTime, InDeltaTime);
 }
 
+void SUMGDesigner::OnDragEnter(const FGeometry& MyGeometry, const FDragDropEvent& DragDropEvent)
+{
+	//@TODO UMG Drop Feedback
+}
+
+void SUMGDesigner::OnDragLeave(const FDragDropEvent& DragDropEvent)
+{
+	TSharedPtr<FWidgetTemplateDragDropOp> DragDropOp = DragDropEvent.GetOperationAs<FWidgetTemplateDragDropOp>();
+	if ( DragDropOp.IsValid() )
+	{
+		DragDropOp->ResetToDefaultToolTip();
+	}
+}
+
+FReply SUMGDesigner::OnDragOver(const FGeometry& MyGeometry, const FDragDropEvent& DragDropEvent)
+{
+	TSharedPtr<FWidgetTemplateDragDropOp> DragDropOp = DragDropEvent.GetOperationAs<FWidgetTemplateDragDropOp>();
+	if ( DragDropOp.IsValid() )
+	{
+		//@TODO UMG Drop Feedback
+		return FReply::Handled();
+	}
+
+	return FReply::Unhandled();
+}
+
+FReply SUMGDesigner::OnDrop(const FGeometry& MyGeometry, const FDragDropEvent& DragDropEvent)
+{
+	TSharedPtr<FWidgetTemplateDragDropOp> DragDropOp = DragDropEvent.GetOperationAs<FWidgetTemplateDragDropOp>();
+	if ( DragDropOp.IsValid() )
+	{
+		FArrangedWidget ArrangedWidget(SNullWidget::NullWidget, FGeometry());
+		USlateWrapperComponent* Template = GetTemplateAtCursor(MyGeometry, DragDropEvent, ArrangedWidget);
+
+		if ( Template && Template->IsA(USlateNonLeafWidgetComponent::StaticClass()) )
+		{
+			UWidgetBlueprint* BP = CastChecked<UWidgetBlueprint>(BlueprintEditor.Pin()->GetBlueprintObj());
+			USlateNonLeafWidgetComponent* Parent = Cast<USlateNonLeafWidgetComponent>(Template);
+
+			USlateWrapperComponent* Widget = DragDropOp->Template->Create(BP->WidgetTree);
+
+			FVector2D LocalPosition = ArrangedWidget.Geometry.AbsoluteToLocal(DragDropEvent.GetScreenSpacePosition());
+			Parent->AddChild(Widget, LocalPosition);
+			//@TODO UMG When we add a child blindly we need to default the slot size to the preferred size of the widget if the container supports such things.
+			//@TODO UMG We may need a desired size canvas, where the slots have no size, they only give you position, alternatively, maybe slots that don't clip, so center is still easy.
+
+			FBlueprintEditorUtils::MarkBlueprintAsStructurallyModified(BP);
+
+			return FReply::Handled();
+		}
+	}
+
+	return FReply::Unhandled();
+}
+
+
 #undef LOCTEXT_NAMESPACE
diff --git a/Engine/Source/Editor/UMGEditor/Private/SUMGDesigner.h b/Engine/Source/Editor/UMGEditor/Private/SUMGDesigner.h
--- a/Engine/Source/Editor/UMGEditor/Private/SUMGDesigner.h
+++ b/Engine/Source/Editor/UMGEditor/Private/SUMGDesigner.h
@@ -20,9 +20,16 @@ public:
 	// SWidget interface
 	virtual int32 OnPaint(const FGeometry& AllottedGeometry, const FSlateRect& MyClippingRect, FSlateWindowElementList& OutDrawElements, int32 LayerId, const FWidgetStyle& InWidgetStyle, bool bParentEnabled) const OVERRIDE;
 	virtual void Tick(const FGeometry& AllottedGeometry, const double InCurrentTime, const float InDeltaTime) OVERRIDE;
-	// End of Swidget interface
+
+	virtual void OnDragEnter(const FGeometry& MyGeometry, const FDragDropEvent& DragDropEvent) OVERRIDE;
+	virtual void OnDragLeave(const FDragDropEvent& DragDropEvent) OVERRIDE;
+	virtual FReply OnDragOver(const FGeometry& MyGeometry, const FDragDropEvent& DragDropEvent) OVERRIDE;
+	virtual FReply OnDrop(const FGeometry& MyGeometry, const FDragDropEvent& DragDropEvent) OVERRIDE;
+	// End of SWidget interface
 
 private:
+	USlateWrapperComponent* GetTemplateAtCursor(const FGeometry& MyGeometry, const FPointerEvent& MouseEvent, FArrangedWidget& ArrangedWidget);
+
 	UWidgetBlueprint* GetBlueprint() const;
 	void OnBlueprintChanged(UBlueprint* InBlueprint);
 	void OnObjectPropertyChanged(UObject* ObjectBeingModified);
diff --git a/Engine/Source/Editor/UMGEditor/Private/SUMGEditorTree.cpp b/Engine/Source/Editor/UMGEditor/Private/SUMGEditorTree.cpp
--- a/Engine/Source/Editor/UMGEditor/Private/SUMGEditorTree.cpp
+++ b/Engine/Source/Editor/UMGEditor/Private/SUMGEditorTree.cpp
@@ -159,7 +159,7 @@ FReply SUMGEditorTree::CreateTestUI()
 
 		UButtonComponent* NewButton = BP->WidgetTree->ConstructWidget<UButtonComponent>(UButtonComponent::StaticClass());
 
-		Vertical->AddChild(NewButton);
+		Vertical->AddChild(NewButton, FVector2D(0,0));
 	}
 	else
 	{
@@ -208,3 +208,6 @@ void SUMGEditorTree::RefreshTree()
 
 	WidgetTreeView->RequestTreeRefresh();
 }
+
+
+//@TODO UMG Drop widgets onto the tree, when nothing is present, if there is a root node present, what happens then, let the root node attempt to place it?
\ No newline at end of file
diff --git a/Engine/Source/Editor/UMGEditor/Private/SUMGEditorTreeItem.cpp b/Engine/Source/Editor/UMGEditor/Private/SUMGEditorTreeItem.cpp
--- a/Engine/Source/Editor/UMGEditor/Private/SUMGEditorTreeItem.cpp
+++ b/Engine/Source/Editor/UMGEditor/Private/SUMGEditorTreeItem.cpp
@@ -85,7 +85,7 @@ FReply SUMGEditorTreeItem::OnDrop(const FGeometry& MyGeometry, const FDragDropEv
 
 			USlateWrapperComponent* Widget = DragDropOp->Template->Create(BP->WidgetTree);
 
-			Parent->AddChild(Widget);
+			Parent->AddChild(Widget, FVector2D(0,0));
 
 			FBlueprintEditorUtils::MarkBlueprintAsStructurallyModified(BP);
 		}
diff --git a/Engine/Source/Runtime/UMG/Private/Components/CanvasPanelComponent.cpp b/Engine/Source/Runtime/UMG/Private/Components/CanvasPanelComponent.cpp
--- a/Engine/Source/Runtime/UMG/Private/Components/CanvasPanelComponent.cpp
+++ b/Engine/Source/Runtime/UMG/Private/Components/CanvasPanelComponent.cpp
@@ -97,6 +97,15 @@ UCanvasPanelSlot* UCanvasPanelComponent::AddSlot(USlateWrapperComponent* Content
 	return Slot;
 }
 
+bool UCanvasPanelComponent::AddChild(USlateWrapperComponent* Child, FVector2D Position)
+{
+	UCanvasPanelSlot* Slot = AddSlot(Child);
+	Slot->Position = Position;
+	Slot->Size = FVector2D(100, 25);
+
+	return true;
+}
+
 #if WITH_EDITOR
 
 void UCanvasPanelComponent::ConnectEditorData()
diff --git a/Engine/Source/Runtime/UMG/Private/Components/CircularThrobberComponent.cpp b/Engine/Source/Runtime/UMG/Private/Components/CircularThrobberComponent.cpp
--- a/Engine/Source/Runtime/UMG/Private/Components/CircularThrobberComponent.cpp
+++ b/Engine/Source/Runtime/UMG/Private/Components/CircularThrobberComponent.cpp
@@ -15,13 +15,16 @@ UCircularThrobberComponent::UCircularThrobberComponent(const FPostConstructIniti
 	NumberOfPieces = DefaultArgs._NumPieces;
 	Period = DefaultArgs._Period;
 	Radius = DefaultArgs._Radius;
-	PieceImage = *DefaultArgs._PieceImage;
 }
 
 TSharedRef<SWidget> UCircularThrobberComponent::RebuildWidget()
 {
+	SCircularThrobber::FArguments DefaultArgs;
+
+	const FSlateBrush* Image = PieceImage ? &PieceImage->Brush : DefaultArgs._PieceImage;
+	
 	return SNew(SCircularThrobber)
-		.PieceImage(&PieceImage)
+		.PieceImage(Image)
 		.NumPieces(NumberOfPieces)
 		.Period(Period)
 		.Radius(Radius);
diff --git a/Engine/Source/Runtime/UMG/Private/Components/ContentWidget.cpp b/Engine/Source/Runtime/UMG/Private/Components/ContentWidget.cpp
--- a/Engine/Source/Runtime/UMG/Private/Components/ContentWidget.cpp
+++ b/Engine/Source/Runtime/UMG/Private/Components/ContentWidget.cpp
@@ -41,7 +41,7 @@ USlateWrapperComponent* UContentWidget::GetChildAt(int32 Index) const
 	return Content;
 }
 
-bool UContentWidget::AddChild(USlateWrapperComponent* InContent)
+bool UContentWidget::AddChild(USlateWrapperComponent* InContent, FVector2D Position)
 {
 	Content = InContent;
 	return true;
diff --git a/Engine/Source/Runtime/UMG/Private/Components/HorizontalBoxComponent.cpp b/Engine/Source/Runtime/UMG/Private/Components/HorizontalBoxComponent.cpp
--- a/Engine/Source/Runtime/UMG/Private/Components/HorizontalBoxComponent.cpp
+++ b/Engine/Source/Runtime/UMG/Private/Components/HorizontalBoxComponent.cpp
@@ -21,7 +21,7 @@ USlateWrapperComponent* UHorizontalBoxComponent::GetChildAt(int32 Index) const
 	return Slots[Index]->Content;
 }
 
-bool UHorizontalBoxComponent::AddChild(USlateWrapperComponent* Child)
+bool UHorizontalBoxComponent::AddChild(USlateWrapperComponent* Child, FVector2D Position)
 {
 	AddSlot(Child);
 	return true;
diff --git a/Engine/Source/Runtime/UMG/Private/Components/ThrobberComponent.cpp b/Engine/Source/Runtime/UMG/Private/Components/ThrobberComponent.cpp
--- a/Engine/Source/Runtime/UMG/Private/Components/ThrobberComponent.cpp
+++ b/Engine/Source/Runtime/UMG/Private/Components/ThrobberComponent.cpp
@@ -12,7 +12,6 @@ UThrobberComponent::UThrobberComponent(const FPostConstructInitializeProperties&
 {
 	SThrobber::FArguments DefaultArgs;
 
-	PieceImage = *DefaultArgs._PieceImage;
 	NumberOfPieces = DefaultArgs._NumPieces;
 
 	bAnimateVertically = (DefaultArgs._Animate & SThrobber::Vertical) != 0;
@@ -22,12 +21,16 @@ UThrobberComponent::UThrobberComponent(const FPostConstructInitializeProperties&
 
 TSharedRef<SWidget> UThrobberComponent::RebuildWidget()
 {
+	SThrobber::FArguments DefaultArgs;
+
+	const FSlateBrush* Image = PieceImage ? &PieceImage->Brush : DefaultArgs._PieceImage;
+
 	const int32 AnimationParams = (bAnimateVertically ? SThrobber::Vertical : 0) |
 		(bAnimateHorizontally ? SThrobber::Horizontal : 0) |
 		(bAnimateOpacity ? SThrobber::Opacity : 0);
 
 	return SNew(SThrobber)
-		.PieceImage(&PieceImage)
+		.PieceImage(Image)
 		.NumPieces(NumberOfPieces)
 		.Animate(static_cast<SThrobber::EAnimation>(AnimationParams));
 }
diff --git a/Engine/Source/Runtime/UMG/Private/Components/VerticalBoxComponent.cpp b/Engine/Source/Runtime/UMG/Private/Components/VerticalBoxComponent.cpp
--- a/Engine/Source/Runtime/UMG/Private/Components/VerticalBoxComponent.cpp
+++ b/Engine/Source/Runtime/UMG/Private/Components/VerticalBoxComponent.cpp
@@ -21,7 +21,7 @@ USlateWrapperComponent* UVerticalBoxComponent::GetChildAt(int32 Index) const
 	return Slots[Index]->Content;
 }
 
-bool UVerticalBoxComponent::AddChild(USlateWrapperComponent* Child)
+bool UVerticalBoxComponent::AddChild(USlateWrapperComponent* Child, FVector2D Position)
 {
 	AddSlot(Child);
 	return true;
diff --git a/Engine/Source/Runtime/UMG/Public/Blueprint/WidgetTree.h b/Engine/Source/Runtime/UMG/Public/Blueprint/WidgetTree.h
--- a/Engine/Source/Runtime/UMG/Public/Blueprint/WidgetTree.h
+++ b/Engine/Source/Runtime/UMG/Public/Blueprint/WidgetTree.h
@@ -28,4 +28,19 @@ class UMG_API UWidgetTree : public UObject
 
 		return (T*)Widget;
 	}
+
+	class USlateWrapperComponent* FindWidget(FString& Name) const
+	{
+		FString ExistingName;
+		for ( USlateWrapperComponent* Widget : WidgetTemplates )
+		{
+			Widget->GetName(ExistingName);
+			if ( ExistingName.Equals(Name, ESearchCase::IgnoreCase) )
+			{
+				return Widget;
+			}
+		}
+
+		return NULL;
+	}
 };
diff --git a/Engine/Source/Runtime/UMG/Public/Components/CanvasPanelComponent.h b/Engine/Source/Runtime/UMG/Public/Components/CanvasPanelComponent.h
--- a/Engine/Source/Runtime/UMG/Public/Components/CanvasPanelComponent.h
+++ b/Engine/Source/Runtime/UMG/Public/Components/CanvasPanelComponent.h
@@ -18,6 +18,7 @@ class UMG_API UCanvasPanelComponent : public USlateNonLeafWidgetComponent
 	FVector2D DesiredCanvasSize;
 
 	UCanvasPanelSlot* AddSlot(USlateWrapperComponent* Content);
+	virtual bool AddChild(USlateWrapperComponent* Child, FVector2D Position) OVERRIDE;
 
 	// USlateNonLeafWidgetComponent
 	virtual int32 GetChildrenCount() const OVERRIDE;
diff --git a/Engine/Source/Runtime/UMG/Public/Components/CircularThrobberComponent.h b/Engine/Source/Runtime/UMG/Public/Components/CircularThrobberComponent.h
--- a/Engine/Source/Runtime/UMG/Public/Components/CircularThrobberComponent.h
+++ b/Engine/Source/Runtime/UMG/Public/Components/CircularThrobberComponent.h
@@ -25,8 +25,8 @@ protected:
 	float Radius;
 
 	/** Image to use for each segment of the throbber */
-	UPROPERTY(EditAnywhere, BlueprintReadOnly, Category=Appearance)
-	FSlateBrush PieceImage;
+	UPROPERTY(EditAnywhere, BlueprintReadWrite, Category=Appearance, meta=( DisplayThumbnail = "true" ))
+	USlateBrushAsset* PieceImage;
 
 protected:
 	// USlateWrapperComponent interface
diff --git a/Engine/Source/Runtime/UMG/Public/Components/ContentWidget.h b/Engine/Source/Runtime/UMG/Public/Components/ContentWidget.h
--- a/Engine/Source/Runtime/UMG/Public/Components/ContentWidget.h
+++ b/Engine/Source/Runtime/UMG/Public/Components/ContentWidget.h
@@ -17,7 +17,7 @@ class UMG_API UContentWidget : public USlateNonLeafWidgetComponent
 
 	virtual int32 GetChildrenCount() const OVERRIDE;
 	virtual USlateWrapperComponent* GetChildAt(int32 Index) const OVERRIDE;
-	virtual bool AddChild(USlateWrapperComponent* InContent) OVERRIDE;
+	virtual bool AddChild(USlateWrapperComponent* InContent, FVector2D Position) OVERRIDE;
 
 protected:
 
diff --git a/Engine/Source/Runtime/UMG/Public/Components/HorizontalBoxComponent.h b/Engine/Source/Runtime/UMG/Public/Components/HorizontalBoxComponent.h
--- a/Engine/Source/Runtime/UMG/Public/Components/HorizontalBoxComponent.h
+++ b/Engine/Source/Runtime/UMG/Public/Components/HorizontalBoxComponent.h
@@ -18,7 +18,7 @@ class UMG_API UHorizontalBoxComponent : public USlateNonLeafWidgetComponent
 	// USlateNonLeafWidgetComponent
 	virtual int32 GetChildrenCount() const OVERRIDE;
 	virtual USlateWrapperComponent* GetChildAt(int32 Index) const OVERRIDE;
-	virtual bool AddChild(USlateWrapperComponent* Child) OVERRIDE;
+	virtual bool AddChild(USlateWrapperComponent* Child, FVector2D Position) OVERRIDE;
 	// End USlateNonLeafWidgetComponent
 
 #if WITH_EDITOR
diff --git a/Engine/Source/Runtime/UMG/Public/Components/SlateNonLeafWidgetComponent.h b/Engine/Source/Runtime/UMG/Public/Components/SlateNonLeafWidgetComponent.h
--- a/Engine/Source/Runtime/UMG/Public/Components/SlateNonLeafWidgetComponent.h
+++ b/Engine/Source/Runtime/UMG/Public/Components/SlateNonLeafWidgetComponent.h
@@ -12,5 +12,5 @@ class UMG_API USlateNonLeafWidgetComponent : public USlateWrapperComponent
 	virtual int32 GetChildrenCount() const { return 0; }
 	virtual USlateWrapperComponent* GetChildAt(int32 Index) const { return NULL; }
 
-	virtual bool AddChild(USlateWrapperComponent* Child) { return false; }
+	virtual bool AddChild(USlateWrapperComponent* Child, FVector2D Position) { return false; }
 };
diff --git a/Engine/Source/Runtime/UMG/Public/Components/ThrobberComponent.h b/Engine/Source/Runtime/UMG/Public/Components/ThrobberComponent.h
--- a/Engine/Source/Runtime/UMG/Public/Components/ThrobberComponent.h
+++ b/Engine/Source/Runtime/UMG/Public/Components/ThrobberComponent.h
@@ -27,8 +27,8 @@ class UMG_API UThrobberComponent : public USlateLeafWidgetComponent
 	bool bAnimateOpacity;
 
 	/** Image to use for each segment of the throbber */
-	UPROPERTY(EditAnywhere, BlueprintReadOnly, Category=Appearance)
-	FSlateBrush PieceImage;
+	UPROPERTY(EditAnywhere, BlueprintReadWrite, Category=Appearance, meta=( DisplayThumbnail = "true" ))
+	USlateBrushAsset* PieceImage;
 
 protected:
 	// USlateWrapperComponent interface
diff --git a/Engine/Source/Runtime/UMG/Public/Components/VerticalBoxComponent.h b/Engine/Source/Runtime/UMG/Public/Components/VerticalBoxComponent.h
--- a/Engine/Source/Runtime/UMG/Public/Components/VerticalBoxComponent.h
+++ b/Engine/Source/Runtime/UMG/Public/Components/VerticalBoxComponent.h
@@ -18,7 +18,7 @@ class UMG_API UVerticalBoxComponent : public USlateNonLeafWidgetComponent
 	// USlateNonLeafWidgetComponent
 	virtual int32 GetChildrenCount() const OVERRIDE;
 	virtual USlateWrapperComponent* GetChildAt(int32 Index) const OVERRIDE;
-	virtual bool AddChild(USlateWrapperComponent* Child) OVERRIDE;
+	virtual bool AddChild(USlateWrapperComponent* Child, FVector2D Position) OVERRIDE;
 	// End USlateNonLeafWidgetComponent
 
 #if WITH_EDITOR
