UMG - Fixing non unity build - partial work on implementing drag resize handles.

diff --git a/Engine/Source/Editor/UMGEditor/Private/SUMGDesigner.cpp b/Engine/Source/Editor/UMGEditor/Private/SUMGDesigner.cpp
--- a/Engine/Source/Editor/UMGEditor/Private/SUMGDesigner.cpp
+++ b/Engine/Source/Editor/UMGEditor/Private/SUMGDesigner.cpp
@@ -5,6 +5,7 @@
 #include "SUMGDesigner.h"
 #include "BlueprintEditor.h"
 #include "SKismetInspector.h"
+#include "BlueprintEditorUtils.h"
 
 #include "WidgetTemplateDragDropOp.h"
 
@@ -98,6 +99,7 @@ void SUMGDesigner::Construct(const FArguments& InArgs, TSharedPtr<FBlueprintEdit
 {
 	PreviewWidgetActor = NULL;
 	BlueprintEditor = InBlueprintEditor;
+	CurrentHandle = DH_NONE;
 	
 	SDesignSurface::Construct(SDesignSurface::FArguments()
 		.Content()
@@ -215,26 +217,68 @@ USlateWrapperComponent* SUMGDesigner::GetTemplateAtCursor(const FGeometry& MyGeo
 
 FReply SUMGDesigner::OnMouseButtonDown(const FGeometry& MyGeometry, const FPointerEvent& MouseEvent)
 {
-	FArrangedWidget ArrangedWidget(SNullWidget::NullWidget, FGeometry());
-	USlateWrapperComponent* Template = GetTemplateAtCursor(MyGeometry, MouseEvent, ArrangedWidget);
+	CurrentHandle = HitTestDragHandles(MyGeometry, MouseEvent);
 
-	if ( Template )
+	if ( CurrentHandle == DH_NONE )
 	{
-		// Set the selected widget so that we can draw the highlight
-		//@TODO UMG Don't store the transient Widget as the selected object, store the template component and lookup the widget from that.
-		SelectedWidget = ArrangedWidget.Widget;
-
-		//@TODO UMG primary FBlueprintEditor needs to be inherited and selection control needs to be centralized.
-		// Set the template as selected in the details panel
-		TArray<USlateWrapperComponent*> SelectedTemplates;
-		SelectedTemplates.Add(Template);
-		ShowDetailsForObjects(SelectedTemplates);
+		FArrangedWidget ArrangedWidget(SNullWidget::NullWidget, FGeometry());
+		USlateWrapperComponent* Template = GetTemplateAtCursor(MyGeometry, MouseEvent, ArrangedWidget);
+
+		if ( Template )
+		{
+			// Set the selected widget so that we can draw the highlight
+			//@TODO UMG Don't store the transient Widget as the selected object, store the template component and lookup the widget from that.
+			SelectedWidget = ArrangedWidget.Widget;
+
+			//@TODO UMG primary FBlueprintEditor needs to be inherited and selection control needs to be centralized.
+			// Set the template as selected in the details panel
+			TArray<USlateWrapperComponent*> SelectedTemplates;
+			SelectedTemplates.Add(Template);
+			ShowDetailsForObjects(SelectedTemplates);
+		}
 	}
 
 	return FReply::Handled();
 }
 
-bool SUMGDesigner::GetLocallyArrangedWidget(TSharedRef<SWidget> Widget, FArrangedWidget& ArrangedWidget) const
+FReply SUMGDesigner::OnMouseButtonUp(const FGeometry& MyGeometry, const FPointerEvent& MouseEvent)
+{
+	CurrentHandle = DH_NONE;
+	return FReply::Handled();
+}
+
+FReply SUMGDesigner::OnMouseMove(const FGeometry& MyGeometry, const FPointerEvent& MouseEvent)
+{
+	if ( CurrentHandle != DH_NONE )
+	{
+
+	}
+
+	return FReply::Handled();
+}
+
+bool SUMGDesigner::GetArrangedWidget(TSharedRef<SWidget> Widget, FArrangedWidget& ArrangedWidget) const
+{
+	// We can't screenshot the widget unless there's a valid window handle to draw it in.
+	TSharedPtr<SWindow> WidgetWindow = FSlateApplication::Get().FindWidgetWindow(Widget);
+	if ( !WidgetWindow.IsValid() )
+	{
+		return false;
+	}
+
+	TSharedRef<SWindow> CurrentWindowRef = WidgetWindow.ToSharedRef();
+
+	FWidgetPath WidgetPath;
+	if ( FSlateApplication::Get().GeneratePathToWidgetUnchecked(Widget, WidgetPath) )
+	{
+		ArrangedWidget = WidgetPath.FindArrangedWidget(Widget);
+		return true;
+	}
+
+	return false;
+}
+
+bool SUMGDesigner::GetArrangedWidgetRelativeToWindow(TSharedRef<SWidget> Widget, FArrangedWidget& ArrangedWidget) const
 {
 	// We can't screenshot the widget unless there's a valid window handle to draw it in.
 	TSharedPtr<SWindow> WidgetWindow = FSlateApplication::Get().FindWidgetWindow(Widget);
@@ -265,7 +309,7 @@ int32 SUMGDesigner::OnPaint(const FGeometry& AllottedGeometry, const FSlateRect&
 		TSharedRef<SWidget> Widget = SelectedWidget.Pin().ToSharedRef();
 
 		FArrangedWidget ArrangedWidget(SNullWidget::NullWidget, FGeometry());
-		GetLocallyArrangedWidget(Widget, ArrangedWidget);
+		GetArrangedWidgetRelativeToWindow(Widget, ArrangedWidget);
 
 		const FLinearColor Tint(0, 1, 0);
 
@@ -286,44 +330,117 @@ int32 SUMGDesigner::OnPaint(const FGeometry& AllottedGeometry, const FSlateRect&
 			Tint
 		);
 
-		float X = SelectionGeometry.DrawPosition.X;
-		float Y = SelectionGeometry.DrawPosition.Y;
-		float Width = SelectionGeometry.DrawSize.X;
-		float Height = SelectionGeometry.DrawSize.Y;
-		
+		DrawDragHandles(SelectionGeometry, MyClippingRect, OutDrawElements, LayerId, InWidgetStyle);
+	}
+
+	return LayerId;
+}
+
+void SUMGDesigner::DrawDragHandles(const FPaintGeometry& SelectionGeometry, const FSlateRect& MyClippingRect, FSlateWindowElementList& OutDrawElements, int32 LayerId, const FWidgetStyle& InWidgetStyle) const
+{
+	float X = SelectionGeometry.DrawPosition.X;
+	float Y = SelectionGeometry.DrawPosition.Y;
+	float Width = SelectionGeometry.DrawSize.X;
+	float Height = SelectionGeometry.DrawSize.Y;
+
+	TArray<FVector2D> Handles;
+	Handles.Add(FVector2D(X, Y));					// Top - Left
+	Handles.Add(FVector2D(X + Width * 0.5f, Y));	// Top - Middle
+	Handles.Add(FVector2D(X + Width, Y));			// Top - Right
+
+	Handles.Add(FVector2D(X, Y + Height * 0.5f));			// Middle - Left
+	Handles.Add(FVector2D(X + Width, Y + Height * 0.5f));	// Middle - Right
+
+	Handles.Add(FVector2D(X, Y + Height));					// Bottom - Left
+	Handles.Add(FVector2D(X + Width * 0.5f, Y + Height));	// Bottom - Middle
+	Handles.Add(FVector2D(X + Width, Y + Height));			// Bottom - Right
+
+	const FVector2D HandleSize = FVector2D(10, 10);
+
+	// @TODO UMG - Don't use the curve editors brushes
+	const FSlateBrush* KeyBrush = FEditorStyle::GetBrush("CurveEd.CurveKey");
+	FLinearColor KeyColor = InWidgetStyle.GetColorAndOpacityTint();// IsEditingEnabled() ? InWidgetStyle.GetColorAndOpacityTint() : FLinearColor(0.1f, 0.1f, 0.1f, 1.f);
+
+	for ( FVector2D Handle : Handles )
+	{
+		FSlateDrawElement::MakeBox(
+			OutDrawElements,
+			++LayerId,
+			FPaintGeometry(FVector2D(Handle.X - HandleSize.X * 0.5f, Handle.Y - HandleSize.Y * 0.5f), HandleSize, 1.0f),
+			KeyBrush,
+			MyClippingRect,
+			ESlateDrawEffect::None,
+			KeyBrush->GetTint(InWidgetStyle) * InWidgetStyle.GetColorAndOpacityTint() * KeyColor
+			);
+	}
+}
+
+SUMGDesigner::DragHandle SUMGDesigner::HitTestDragHandles(const FGeometry& AllottedGeometry, const FPointerEvent& PointerEvent) const
+{
+	FVector2D LocalPointer = AllottedGeometry.AbsoluteToLocal(PointerEvent.GetScreenSpacePosition());
+
+	if ( SelectedWidget.IsValid() )
+	{
+		TSharedRef<SWidget> Widget = SelectedWidget.Pin().ToSharedRef();
+
+		FArrangedWidget ArrangedWidget(SNullWidget::NullWidget, FGeometry());
+		GetArrangedWidget(Widget, ArrangedWidget);
+
+		FVector2D WidgetLocalPosition = AllottedGeometry.AbsoluteToLocal(ArrangedWidget.Geometry.AbsolutePosition);
+		float X = WidgetLocalPosition.X;
+		float Y = WidgetLocalPosition.Y;
+		float Width = ArrangedWidget.Geometry.Size.X;
+		float Height = ArrangedWidget.Geometry.Size.Y;
+
 		TArray<FVector2D> Handles;
 		Handles.Add(FVector2D(X, Y));					// Top - Left
 		Handles.Add(FVector2D(X + Width * 0.5f, Y));	// Top - Middle
 		Handles.Add(FVector2D(X + Width, Y));			// Top - Right
 
+		Handles.Add(FVector2D(X, Y + Height * 0.5f));			// Middle - Left
+		Handles.Add(FVector2D(X + Width, Y + Height * 0.5f));	// Middle - Right
+
 		Handles.Add(FVector2D(X, Y + Height));					// Bottom - Left
 		Handles.Add(FVector2D(X + Width * 0.5f, Y + Height));	// Bottom - Middle
 		Handles.Add(FVector2D(X + Width, Y + Height));			// Bottom - Right
 
-		Handles.Add(FVector2D(X, Y + Height * 0.5f));			// Middle - Left
-		Handles.Add(FVector2D(X + Width, Y + Height * 0.5f));	// Middle - Right
-
 		const FVector2D HandleSize = FVector2D(10, 10);
 
-		// @TODO UMG - Don't use the curve editors brushes
-		const FSlateBrush* KeyBrush = FEditorStyle::GetBrush("CurveEd.CurveKey");
-		FLinearColor KeyColor = InWidgetStyle.GetColorAndOpacityTint();// IsEditingEnabled() ? InWidgetStyle.GetColorAndOpacityTint() : FLinearColor(0.1f, 0.1f, 0.1f, 1.f);
-
+		int32 i = 0;
 		for ( FVector2D Handle : Handles )
 		{
-			FSlateDrawElement::MakeBox(
-				OutDrawElements,
-				++LayerId,
-				FPaintGeometry(FVector2D(Handle.X - HandleSize.X * 0.5f, Handle.Y - HandleSize.Y * 0.5f), HandleSize, 1.0f),
-				KeyBrush,
-				MyClippingRect,
-				ESlateDrawEffect::None,
-				KeyBrush->GetTint(InWidgetStyle) * InWidgetStyle.GetColorAndOpacityTint() * KeyColor
-				);
+			FSlateRect Rect(FVector2D(Handle.X - HandleSize.X * 0.5f, Handle.Y - HandleSize.Y * 0.5f), Handle + HandleSize);
+			if ( Rect.ContainsPoint(LocalPointer) )
+			{
+				return (DragHandle)i;
+			}
+			i++;
 		}
 	}
 
-	return LayerId;
+	return DH_NONE;
+}
+
+FCursorReply SUMGDesigner::OnCursorQuery(const FGeometry& MyGeometry, const FPointerEvent& CursorEvent) const
+{
+	// Hit test the drag handles
+	switch ( HitTestDragHandles(MyGeometry, CursorEvent) )
+	{
+	case DH_TOP_LEFT:
+	case DH_BOTTOM_RIGHT:
+		return FCursorReply::Cursor(EMouseCursor::ResizeSouthEast);
+	case DH_TOP_RIGHT:
+	case DH_BOTTOM_LEFT:
+		return FCursorReply::Cursor(EMouseCursor::ResizeSouthWest);
+	case DH_TOP_CENTER:
+	case DH_BOTTOM_CENTER:
+		return FCursorReply::Cursor(EMouseCursor::ResizeUpDown);
+	case DH_MIDDLE_LEFT:
+	case DH_MIDDLE_RIGHT:
+		return FCursorReply::Cursor(EMouseCursor::ResizeLeftRight);
+	}
+
+	return FCursorReply::Unhandled();
 }
 
 void SUMGDesigner::Tick(const FGeometry& AllottedGeometry, const double InCurrentTime, const float InDeltaTime)
diff --git a/Engine/Source/Editor/UMGEditor/Private/SUMGDesigner.h b/Engine/Source/Editor/UMGEditor/Private/SUMGDesigner.h
--- a/Engine/Source/Editor/UMGEditor/Private/SUMGDesigner.h
+++ b/Engine/Source/Editor/UMGEditor/Private/SUMGDesigner.h
@@ -15,9 +15,13 @@ public:
 
 	void Construct(const FArguments& InArgs, TSharedPtr<class FBlueprintEditor> InBlueprintEditor);
 
+	// SWidget interface
 	virtual FReply OnMouseButtonDown(const FGeometry& MyGeometry, const FPointerEvent& MouseEvent) OVERRIDE;
+	virtual FReply OnMouseButtonUp(const FGeometry& MyGeometry, const FPointerEvent& MouseEvent) OVERRIDE;
+	virtual FReply OnMouseMove(const FGeometry& MyGeometry, const FPointerEvent& MouseEvent) OVERRIDE;
+
+	virtual FCursorReply OnCursorQuery(const FGeometry& MyGeometry, const FPointerEvent& CursorEvent) const OVERRIDE;
 
-	// SWidget interface
 	virtual int32 OnPaint(const FGeometry& AllottedGeometry, const FSlateRect& MyClippingRect, FSlateWindowElementList& OutDrawElements, int32 LayerId, const FWidgetStyle& InWidgetStyle, bool bParentEnabled) const OVERRIDE;
 	virtual void Tick(const FGeometry& AllottedGeometry, const double InCurrentTime, const float InDeltaTime) OVERRIDE;
 
@@ -35,7 +39,31 @@ private:
 	void OnObjectPropertyChanged(UObject* ObjectBeingModified);
 	void ShowDetailsForObjects(TArray<USlateWrapperComponent*> Widgets);
 
-	bool GetLocallyArrangedWidget(TSharedRef<SWidget> Widget, FArrangedWidget& ArrangedWidget) const;
+	bool GetArrangedWidget(TSharedRef<SWidget> Widget, FArrangedWidget& ArrangedWidget) const;
+	bool GetArrangedWidgetRelativeToWindow(TSharedRef<SWidget> Widget, FArrangedWidget& ArrangedWidget) const;
+
+private:
+	enum DragHandle
+	{
+		DH_NONE = -1,
+
+		DH_TOP_LEFT = 0,
+		DH_TOP_CENTER,
+		DH_TOP_RIGHT,
+
+		DH_MIDDLE_LEFT,
+		DH_MIDDLE_RIGHT,
+
+		DH_BOTTOM_LEFT,
+		DH_BOTTOM_CENTER,
+		DH_BOTTOM_RIGHT,
+		
+		DH_MAX,
+	};
+
+private:
+	void DrawDragHandles(const FPaintGeometry& SelectionGeometry, const FSlateRect& MyClippingRect, FSlateWindowElementList& OutDrawElements, int32 LayerId, const FWidgetStyle& InWidgetStyle) const;
+	DragHandle HitTestDragHandles(const FGeometry& AllottedGeometry, const FPointerEvent& PointerEvent) const;
 
 private:
 	TWeakPtr<FBlueprintEditor> BlueprintEditor;
@@ -45,4 +73,6 @@ private:
 	TWeakPtr<SWidget> SelectedWidget;
 
 	TSharedPtr<SBorder> PreviewSurface;
+
+	DragHandle CurrentHandle;
 };
