UMG - Fixing localization errors.

diff --git a/Engine/Source/Editor/UMGEditor/Private/SDesignSurface.cpp b/Engine/Source/Editor/UMGEditor/Private/SDesignSurface.cpp
--- a/Engine/Source/Editor/UMGEditor/Private/SDesignSurface.cpp
+++ b/Engine/Source/Editor/UMGEditor/Private/SDesignSurface.cpp
@@ -4,13 +4,13 @@
 
 #include "SDesignSurface.h"
 
-#define LOCTEXT_NAMESPACE "Slate"
+#define LOCTEXT_NAMESPACE "UMG"
 
 struct FZoomLevelEntry
 {
 public:
 	FZoomLevelEntry(float InZoomAmount, const FText& InDisplayText, EGraphRenderingLOD::Type InLOD)
-		: DisplayText(FText::Format(NSLOCTEXT("GraphEditor", "Zoom", "Zoom {0}"), InDisplayText))
+		: DisplayText(FText::Format(LOCTEXT("Zoom", "Zoom {0}"), InDisplayText))
 		, ZoomAmount(InZoomAmount)
 		, LOD(InLOD)
 	{
@@ -30,26 +30,26 @@ struct FFixedZoomLevelsContainer : public FZoomLevelsContainer
 		if ( ZoomLevels.Num() == 0 )
 		{
 			ZoomLevels.Reserve(20);
-			ZoomLevels.Add(FZoomLevelEntry(0.100f, NSLOCTEXT("GraphEditor", "ZoomLevel", "-12"), EGraphRenderingLOD::LowestDetail));
-			ZoomLevels.Add(FZoomLevelEntry(0.125f, NSLOCTEXT("GraphEditor", "ZoomLevel", "-11"), EGraphRenderingLOD::LowestDetail));
-			ZoomLevels.Add(FZoomLevelEntry(0.150f, NSLOCTEXT("GraphEditor", "ZoomLevel", "-10"), EGraphRenderingLOD::LowestDetail));
-			ZoomLevels.Add(FZoomLevelEntry(0.175f, NSLOCTEXT("GraphEditor", "ZoomLevel", "-9"), EGraphRenderingLOD::LowestDetail));
-			ZoomLevels.Add(FZoomLevelEntry(0.200f, NSLOCTEXT("GraphEditor", "ZoomLevel", "-8"), EGraphRenderingLOD::LowestDetail));
-			ZoomLevels.Add(FZoomLevelEntry(0.225f, NSLOCTEXT("GraphEditor", "ZoomLevel", "-7"), EGraphRenderingLOD::LowDetail));
-			ZoomLevels.Add(FZoomLevelEntry(0.250f, NSLOCTEXT("GraphEditor", "ZoomLevel", "-6"), EGraphRenderingLOD::LowDetail));
-			ZoomLevels.Add(FZoomLevelEntry(0.375f, NSLOCTEXT("GraphEditor", "ZoomLevel", "-5"), EGraphRenderingLOD::MediumDetail));
-			ZoomLevels.Add(FZoomLevelEntry(0.500f, NSLOCTEXT("GraphEditor", "ZoomLevel", "-4"), EGraphRenderingLOD::MediumDetail));
-			ZoomLevels.Add(FZoomLevelEntry(0.675f, NSLOCTEXT("GraphEditor", "ZoomLevel", "-3"), EGraphRenderingLOD::MediumDetail));
-			ZoomLevels.Add(FZoomLevelEntry(0.750f, NSLOCTEXT("GraphEditor", "ZoomLevel", "-2"), EGraphRenderingLOD::DefaultDetail));
-			ZoomLevels.Add(FZoomLevelEntry(0.875f, NSLOCTEXT("GraphEditor", "ZoomLevel", "-1"), EGraphRenderingLOD::DefaultDetail));
-			ZoomLevels.Add(FZoomLevelEntry(1.000f, NSLOCTEXT("GraphEditor", "ZoomLevel", "1:1"), EGraphRenderingLOD::DefaultDetail));
-			ZoomLevels.Add(FZoomLevelEntry(1.250f, NSLOCTEXT("GraphEditor", "ZoomLevel", "+1"), EGraphRenderingLOD::DefaultDetail));
-			ZoomLevels.Add(FZoomLevelEntry(1.375f, NSLOCTEXT("GraphEditor", "ZoomLevel", "+2"), EGraphRenderingLOD::DefaultDetail));
-			ZoomLevels.Add(FZoomLevelEntry(1.500f, NSLOCTEXT("GraphEditor", "ZoomLevel", "+3"), EGraphRenderingLOD::FullyZoomedIn));
-			ZoomLevels.Add(FZoomLevelEntry(1.675f, NSLOCTEXT("GraphEditor", "ZoomLevel", "+4"), EGraphRenderingLOD::FullyZoomedIn));
-			ZoomLevels.Add(FZoomLevelEntry(1.750f, NSLOCTEXT("GraphEditor", "ZoomLevel", "+5"), EGraphRenderingLOD::FullyZoomedIn));
-			ZoomLevels.Add(FZoomLevelEntry(1.875f, NSLOCTEXT("GraphEditor", "ZoomLevel", "+6"), EGraphRenderingLOD::FullyZoomedIn));
-			ZoomLevels.Add(FZoomLevelEntry(2.000f, NSLOCTEXT("GraphEditor", "ZoomLevel", "+7"), EGraphRenderingLOD::FullyZoomedIn));
+			ZoomLevels.Add(FZoomLevelEntry(0.100f, LOCTEXT("ZoomLevel", "-12"), EGraphRenderingLOD::LowestDetail));
+			ZoomLevels.Add(FZoomLevelEntry(0.125f, LOCTEXT("ZoomLevel", "-11"), EGraphRenderingLOD::LowestDetail));
+			ZoomLevels.Add(FZoomLevelEntry(0.150f, LOCTEXT("ZoomLevel", "-10"), EGraphRenderingLOD::LowestDetail));
+			ZoomLevels.Add(FZoomLevelEntry(0.175f, LOCTEXT("ZoomLevel", "-9"), EGraphRenderingLOD::LowestDetail));
+			ZoomLevels.Add(FZoomLevelEntry(0.200f, LOCTEXT("ZoomLevel", "-8"), EGraphRenderingLOD::LowestDetail));
+			ZoomLevels.Add(FZoomLevelEntry(0.225f, LOCTEXT("ZoomLevel", "-7"), EGraphRenderingLOD::LowDetail));
+			ZoomLevels.Add(FZoomLevelEntry(0.250f, LOCTEXT("ZoomLevel", "-6"), EGraphRenderingLOD::LowDetail));
+			ZoomLevels.Add(FZoomLevelEntry(0.375f, LOCTEXT("ZoomLevel", "-5"), EGraphRenderingLOD::MediumDetail));
+			ZoomLevels.Add(FZoomLevelEntry(0.500f, LOCTEXT("ZoomLevel", "-4"), EGraphRenderingLOD::MediumDetail));
+			ZoomLevels.Add(FZoomLevelEntry(0.675f, LOCTEXT("ZoomLevel", "-3"), EGraphRenderingLOD::MediumDetail));
+			ZoomLevels.Add(FZoomLevelEntry(0.750f, LOCTEXT("ZoomLevel", "-2"), EGraphRenderingLOD::DefaultDetail));
+			ZoomLevels.Add(FZoomLevelEntry(0.875f, LOCTEXT("ZoomLevel", "-1"), EGraphRenderingLOD::DefaultDetail));
+			ZoomLevels.Add(FZoomLevelEntry(1.000f, LOCTEXT("ZoomLevel", "1:1"), EGraphRenderingLOD::DefaultDetail));
+			ZoomLevels.Add(FZoomLevelEntry(1.250f, LOCTEXT("ZoomLevel", "+1"), EGraphRenderingLOD::DefaultDetail));
+			ZoomLevels.Add(FZoomLevelEntry(1.375f, LOCTEXT("ZoomLevel", "+2"), EGraphRenderingLOD::DefaultDetail));
+			ZoomLevels.Add(FZoomLevelEntry(1.500f, LOCTEXT("ZoomLevel", "+3"), EGraphRenderingLOD::FullyZoomedIn));
+			ZoomLevels.Add(FZoomLevelEntry(1.675f, LOCTEXT("ZoomLevel", "+4"), EGraphRenderingLOD::FullyZoomedIn));
+			ZoomLevels.Add(FZoomLevelEntry(1.750f, LOCTEXT("ZoomLevel", "+5"), EGraphRenderingLOD::FullyZoomedIn));
+			ZoomLevels.Add(FZoomLevelEntry(1.875f, LOCTEXT("ZoomLevel", "+6"), EGraphRenderingLOD::FullyZoomedIn));
+			ZoomLevels.Add(FZoomLevelEntry(2.000f, LOCTEXT("ZoomLevel", "+7"), EGraphRenderingLOD::FullyZoomedIn));
 		}
 	}
 
