UMG - Prototyping better selection, grab handles to perform resizing, each slot is allowed to control resizing different directions.

diff --git a/Engine/Source/Editor/UMGEditor/Private/SUMGDesigner.cpp b/Engine/Source/Editor/UMGEditor/Private/SUMGDesigner.cpp
--- a/Engine/Source/Editor/UMGEditor/Private/SUMGDesigner.cpp
+++ b/Engine/Source/Editor/UMGEditor/Private/SUMGDesigner.cpp
@@ -100,7 +100,23 @@ void SUMGDesigner::Construct(const FArguments& InArgs, TSharedPtr<FBlueprintEdit
 	PreviewWidgetActor = NULL;
 	BlueprintEditor = InBlueprintEditor;
 	CurrentHandle = DH_NONE;
-	
+	SelectedTemplate = NULL;
+
+	DragDirections.Init((int32)DH_MAX);
+	DragDirections[DH_TOP_LEFT] = FVector2D(-1, -1);
+	DragDirections[DH_TOP_CENTER] = FVector2D(0, -1);
+	DragDirections[DH_TOP_RIGHT] = FVector2D(1, -1);
+
+	DragDirections[DH_MIDDLE_LEFT] = FVector2D(-1, 0);
+	DragDirections[DH_MIDDLE_RIGHT] = FVector2D(1, 0);
+
+	DragDirections[DH_BOTTOM_LEFT] = FVector2D(-1, 1);
+	DragDirections[DH_BOTTOM_CENTER] = FVector2D(0, 1);
+	DragDirections[DH_BOTTOM_RIGHT] = FVector2D(1, 1);
+
+	UWidgetBlueprint* Blueprint = GetBlueprint();
+	Blueprint->OnChanged().AddSP(this, &SUMGDesigner::OnBlueprintChanged);
+
 	SDesignSurface::Construct(SDesignSurface::FArguments()
 		.Content()
 		[
@@ -124,6 +140,15 @@ void SUMGDesigner::Construct(const FArguments& InArgs, TSharedPtr<FBlueprintEdit
 	);
 }
 
+SUMGDesigner::~SUMGDesigner()
+{
+	UWidgetBlueprint* Blueprint = GetBlueprint();
+	if ( Blueprint )
+	{
+		Blueprint->OnChanged().RemoveAll(this);
+	}
+}
+
 UWidgetBlueprint* SUMGDesigner::GetBlueprint() const
 {
 	if ( BlueprintEditor.IsValid() )
@@ -139,7 +164,6 @@ void SUMGDesigner::OnBlueprintChanged(UBlueprint* InBlueprint)
 {
 	if ( InBlueprint )
 	{
-		
 	}
 }
 
@@ -222,21 +246,21 @@ FReply SUMGDesigner::OnMouseButtonDown(const FGeometry& MyGeometry, const FPoint
 	if ( CurrentHandle == DH_NONE )
 	{
 		FArrangedWidget ArrangedWidget(SNullWidget::NullWidget, FGeometry());
-		USlateWrapperComponent* Template = GetTemplateAtCursor(MyGeometry, MouseEvent, ArrangedWidget);
+		SelectedTemplate = GetTemplateAtCursor(MyGeometry, MouseEvent, ArrangedWidget);
 
-		if ( Template )
+		if ( SelectedTemplate )
 		{
-			// Set the selected widget so that we can draw the highlight
-			//@TODO UMG Don't store the transient Widget as the selected object, store the template component and lookup the widget from that.
-			SelectedWidget = ArrangedWidget.Widget;
-
 			//@TODO UMG primary FBlueprintEditor needs to be inherited and selection control needs to be centralized.
 			// Set the template as selected in the details panel
 			TArray<USlateWrapperComponent*> SelectedTemplates;
-			SelectedTemplates.Add(Template);
+			SelectedTemplates.Add(SelectedTemplate);
 			ShowDetailsForObjects(SelectedTemplates);
 		}
 	}
+	else
+	{
+		return FReply::Handled().PreventThrottling().CaptureMouse(AsShared());
+	}
 
 	return FReply::Handled();
 }
@@ -244,14 +268,22 @@ FReply SUMGDesigner::OnMouseButtonDown(const FGeometry& MyGeometry, const FPoint
 FReply SUMGDesigner::OnMouseButtonUp(const FGeometry& MyGeometry, const FPointerEvent& MouseEvent)
 {
 	CurrentHandle = DH_NONE;
-	return FReply::Handled();
+	return FReply::Handled().ReleaseMouseCapture();
 }
 
 FReply SUMGDesigner::OnMouseMove(const FGeometry& MyGeometry, const FPointerEvent& MouseEvent)
 {
 	if ( CurrentHandle != DH_NONE )
 	{
+		if ( SelectedTemplate->Slot && !MouseEvent.GetCursorDelta().IsZero() )
+		{
+			//@TODO UMG - Implement some system to query slots to know if they support dragging so we can provide visual feedback by hiding handles that wouldnt work and such.
+			//SelectedTemplate->Slot->CanResize(DH_TOP_CENTER)
 
+			SelectedTemplate->Slot->Resize(DragDirections[CurrentHandle], MouseEvent.GetCursorDelta());
+			FBlueprintEditorUtils::MarkBlueprintAsStructurallyModified(GetBlueprint());
+			return FReply::Handled().PreventThrottling();
+		}
 	}
 
 	return FReply::Handled();
@@ -338,40 +370,51 @@ int32 SUMGDesigner::OnPaint(const FGeometry& AllottedGeometry, const FSlateRect&
 
 void SUMGDesigner::DrawDragHandles(const FPaintGeometry& SelectionGeometry, const FSlateRect& MyClippingRect, FSlateWindowElementList& OutDrawElements, int32 LayerId, const FWidgetStyle& InWidgetStyle) const
 {
-	float X = SelectionGeometry.DrawPosition.X;
-	float Y = SelectionGeometry.DrawPosition.Y;
-	float Width = SelectionGeometry.DrawSize.X;
-	float Height = SelectionGeometry.DrawSize.Y;
+	if ( SelectedTemplate && SelectedTemplate->Slot )
+	{
 
-	TArray<FVector2D> Handles;
-	Handles.Add(FVector2D(X, Y));					// Top - Left
-	Handles.Add(FVector2D(X + Width * 0.5f, Y));	// Top - Middle
-	Handles.Add(FVector2D(X + Width, Y));			// Top - Right
+		float X = SelectionGeometry.DrawPosition.X;
+		float Y = SelectionGeometry.DrawPosition.Y;
+		float Width = SelectionGeometry.DrawSize.X;
+		float Height = SelectionGeometry.DrawSize.Y;
 
-	Handles.Add(FVector2D(X, Y + Height * 0.5f));			// Middle - Left
-	Handles.Add(FVector2D(X + Width, Y + Height * 0.5f));	// Middle - Right
+		TArray<FVector2D> Handles;
+		Handles.Add(FVector2D(X, Y));					// Top - Left
+		Handles.Add(FVector2D(X + Width * 0.5f, Y));	// Top - Middle
+		Handles.Add(FVector2D(X + Width, Y));			// Top - Right
 
-	Handles.Add(FVector2D(X, Y + Height));					// Bottom - Left
-	Handles.Add(FVector2D(X + Width * 0.5f, Y + Height));	// Bottom - Middle
-	Handles.Add(FVector2D(X + Width, Y + Height));			// Bottom - Right
+		Handles.Add(FVector2D(X, Y + Height * 0.5f));			// Middle - Left
+		Handles.Add(FVector2D(X + Width, Y + Height * 0.5f));	// Middle - Right
 
-	const FVector2D HandleSize = FVector2D(10, 10);
+		Handles.Add(FVector2D(X, Y + Height));					// Bottom - Left
+		Handles.Add(FVector2D(X + Width * 0.5f, Y + Height));	// Bottom - Middle
+		Handles.Add(FVector2D(X + Width, Y + Height));			// Bottom - Right
 
-	// @TODO UMG - Don't use the curve editors brushes
-	const FSlateBrush* KeyBrush = FEditorStyle::GetBrush("CurveEd.CurveKey");
-	FLinearColor KeyColor = InWidgetStyle.GetColorAndOpacityTint();// IsEditingEnabled() ? InWidgetStyle.GetColorAndOpacityTint() : FLinearColor(0.1f, 0.1f, 0.1f, 1.f);
+		const FVector2D HandleSize = FVector2D(10, 10);
 
-	for ( FVector2D Handle : Handles )
-	{
-		FSlateDrawElement::MakeBox(
-			OutDrawElements,
-			++LayerId,
-			FPaintGeometry(FVector2D(Handle.X - HandleSize.X * 0.5f, Handle.Y - HandleSize.Y * 0.5f), HandleSize, 1.0f),
-			KeyBrush,
-			MyClippingRect,
-			ESlateDrawEffect::None,
-			KeyBrush->GetTint(InWidgetStyle) * InWidgetStyle.GetColorAndOpacityTint() * KeyColor
-			);
+		// @TODO UMG - Don't use the curve editors brushes
+		const FSlateBrush* KeyBrush = FEditorStyle::GetBrush("CurveEd.CurveKey");
+		FLinearColor KeyColor = InWidgetStyle.GetColorAndOpacityTint();// IsEditingEnabled() ? InWidgetStyle.GetColorAndOpacityTint() : FLinearColor(0.1f, 0.1f, 0.1f, 1.f);
+
+		for ( int32 HandleIndex = 0; HandleIndex < Handles.Num(); HandleIndex++ )
+		{
+			const FVector2D& Handle = Handles[HandleIndex];
+			if ( !SelectedTemplate->Slot->CanResize(DragDirections[HandleIndex]) )
+			{
+				// This isn't a valid handle
+				continue;
+			}
+
+			FSlateDrawElement::MakeBox(
+				OutDrawElements,
+				++LayerId,
+				FPaintGeometry(FVector2D(Handle.X - HandleSize.X * 0.5f, Handle.Y - HandleSize.Y * 0.5f), HandleSize, 1.0f),
+				KeyBrush,
+				MyClippingRect,
+				ESlateDrawEffect::None,
+				KeyBrush->GetTint(InWidgetStyle) * InWidgetStyle.GetColorAndOpacityTint() * KeyColor
+				);
+		}
 	}
 }
 
@@ -379,7 +422,7 @@ SUMGDesigner::DragHandle SUMGDesigner::HitTestDragHandles(const FGeometry& Allot
 {
 	FVector2D LocalPointer = AllottedGeometry.AbsoluteToLocal(PointerEvent.GetScreenSpacePosition());
 
-	if ( SelectedWidget.IsValid() )
+	if ( SelectedTemplate && SelectedTemplate->Slot && SelectedWidget.IsValid() )
 	{
 		TSharedRef<SWidget> Widget = SelectedWidget.Pin().ToSharedRef();
 
@@ -412,6 +455,12 @@ SUMGDesigner::DragHandle SUMGDesigner::HitTestDragHandles(const FGeometry& Allot
 			FSlateRect Rect(FVector2D(Handle.X - HandleSize.X * 0.5f, Handle.Y - HandleSize.Y * 0.5f), Handle + HandleSize);
 			if ( Rect.ContainsPoint(LocalPointer) )
 			{
+				if ( !SelectedTemplate->Slot->CanResize(DragDirections[i]) )
+				{
+					// This isn't a valid handle
+					break;
+				}
+
 				return (DragHandle)i;
 			}
 			i++;
@@ -474,6 +523,20 @@ void SUMGDesigner::Tick(const FGeometry& AllottedGeometry, const double InCurren
 		];
 	}
 
+	// Update the selected widget to match the selected template.
+	if ( SelectedTemplate )
+	{
+		AUserWidget* WidgetActor = PreviewWidgetActor.Get();
+		if ( WidgetActor )
+		{
+			// Set the selected widget so that we can draw the highlight
+			//@TODO UMG Don't store the transient Widget as the selected object, store the template component and lookup the widget from that.
+			//SelectedWidget = ArrangedWidget.Widget;
+
+			SelectedWidget = WidgetActor->GetWidgetFromName(SelectedTemplate->GetName());
+		}
+	}
+
 	SCompoundWidget::Tick(AllottedGeometry, InCurrentTime, InDeltaTime);
 }
 
diff --git a/Engine/Source/Editor/UMGEditor/Private/SUMGDesigner.h b/Engine/Source/Editor/UMGEditor/Private/SUMGDesigner.h
--- a/Engine/Source/Editor/UMGEditor/Private/SUMGDesigner.h
+++ b/Engine/Source/Editor/UMGEditor/Private/SUMGDesigner.h
@@ -14,6 +14,7 @@ public:
 	SLATE_END_ARGS()
 
 	void Construct(const FArguments& InArgs, TSharedPtr<class FBlueprintEditor> InBlueprintEditor);
+	virtual ~SUMGDesigner();
 
 	// SWidget interface
 	virtual FReply OnMouseButtonDown(const FGeometry& MyGeometry, const FPointerEvent& MouseEvent) OVERRIDE;
@@ -61,6 +62,8 @@ private:
 		DH_MAX,
 	};
 
+	TArray< FVector2D > DragDirections;
+
 private:
 	void DrawDragHandles(const FPaintGeometry& SelectionGeometry, const FSlateRect& MyClippingRect, FSlateWindowElementList& OutDrawElements, int32 LayerId, const FWidgetStyle& InWidgetStyle) const;
 	DragHandle HitTestDragHandles(const FGeometry& AllottedGeometry, const FPointerEvent& PointerEvent) const;
@@ -70,6 +73,7 @@ private:
 	TWeakObjectPtr<AUserWidget> PreviewWidgetActor;
 	TWeakPtr<SWidget> PreviewWidget;
 
+	USlateWrapperComponent* SelectedTemplate;
 	TWeakPtr<SWidget> SelectedWidget;
 
 	TSharedPtr<SBorder> PreviewSurface;
diff --git a/Engine/Source/Runtime/UMG/Private/Components/CanvasPanelSlot.cpp b/Engine/Source/Runtime/UMG/Private/Components/CanvasPanelSlot.cpp
--- a/Engine/Source/Runtime/UMG/Private/Components/CanvasPanelSlot.cpp
+++ b/Engine/Source/Runtime/UMG/Private/Components/CanvasPanelSlot.cpp
@@ -13,3 +13,30 @@ UCanvasPanelSlot::UCanvasPanelSlot(const FPostConstructInitializeProperties& PCI
 	HorizontalAlignment = HAlign_Left;
 	VerticalAlignment = VAlign_Top;
 }
+
+void UCanvasPanelSlot::Resize(const FVector2D& Direction, const FVector2D& Amount)
+{
+	if ( Direction.X < 0 )
+	{
+		Position.X -= Amount.X * Direction.X;
+		Size.X += Amount.X * Direction.X;
+	}
+	if ( Direction.Y < 0 )
+	{
+		Position.Y -= Amount.Y * Direction.Y;
+		Size.Y += Amount.Y * Direction.Y;
+	}
+	if ( Direction.X > 0 )
+	{
+		Size.X += Amount.X * Direction.X;
+	}
+	if ( Direction.Y > 0 )
+	{
+		Size.Y += Amount.Y * Direction.Y;
+	}
+}
+
+bool UCanvasPanelSlot::CanResize(const FVector2D& Direction) const
+{
+	return true;
+}
diff --git a/Engine/Source/Runtime/UMG/Private/Components/PanelSlot.cpp b/Engine/Source/Runtime/UMG/Private/Components/PanelSlot.cpp
--- a/Engine/Source/Runtime/UMG/Private/Components/PanelSlot.cpp
+++ b/Engine/Source/Runtime/UMG/Private/Components/PanelSlot.cpp
@@ -8,4 +8,14 @@
 UPanelSlot::UPanelSlot(const FPostConstructInitializeProperties& PCIP)
 	: Super(PCIP)
 {
+}
+
+void UPanelSlot::Resize(const FVector2D& Direction, const FVector2D& Amount)
+{
+
+}
+
+bool UPanelSlot::CanResize(const FVector2D& Direction) const
+{
+	return false;
 }
\ No newline at end of file
diff --git a/Engine/Source/Runtime/UMG/Private/UserWidget.cpp b/Engine/Source/Runtime/UMG/Private/UserWidget.cpp
--- a/Engine/Source/Runtime/UMG/Private/UserWidget.cpp
+++ b/Engine/Source/Runtime/UMG/Private/UserWidget.cpp
@@ -31,7 +31,7 @@ void AUserWidget::RerunConstructionScripts()
 
 USlateWrapperComponent* AUserWidget::GetWidgetHandle(TSharedRef<SWidget> InWidget)
 {
-	return WidgetToComponent.FindRef(&InWidget.Get());
+	return WidgetToComponent.FindRef(InWidget);
 }
 
 void AUserWidget::RebuildWrapperWidget()
@@ -54,7 +54,7 @@ void AUserWidget::RebuildWrapperWidget()
 		USlateWrapperComponent* Handle = SlateWrapperComponents[ComponentIndex];
 		TSharedRef<SWidget> Widget = Handle->GetWidget();
 
-		WidgetToComponent.Add(&Widget.Get(), Handle);
+		WidgetToComponent.Add(Widget, Handle);
 	}
 	
 	// If this is a game world add the widget to the current worlds viewport.
@@ -66,6 +66,32 @@ void AUserWidget::RebuildWrapperWidget()
 	}
 }
 
+TSharedPtr<SWidget> AUserWidget::GetWidgetFromName(const FString& Name) const
+{
+	for ( auto& Entry : WidgetToComponent )
+	{
+		if ( Entry.Value->GetName().Equals(Name, ESearchCase::IgnoreCase) )
+		{
+			return Entry.Key.Pin();
+		}
+	}
+
+	return TSharedPtr<SWidget>();
+}
+
+USlateWrapperComponent* AUserWidget::GetHandleFromName(const FString& Name) const
+{
+	for ( auto& Entry : WidgetToComponent )
+	{
+		if ( Entry.Value->GetName().Equals(Name, ESearchCase::IgnoreCase) )
+		{
+			return Entry.Value;
+		}
+	}
+
+	return NULL;
+}
+
 TSharedRef<SWidget> AUserWidget::GetRootWidget()
 {
 	return RootWidget.ToSharedRef();
diff --git a/Engine/Source/Runtime/UMG/Public/Blueprint/UserWidget.h b/Engine/Source/Runtime/UMG/Public/Blueprint/UserWidget.h
--- a/Engine/Source/Runtime/UMG/Public/Blueprint/UserWidget.h
+++ b/Engine/Source/Runtime/UMG/Public/Blueprint/UserWidget.h
@@ -19,9 +19,12 @@ class UMG_API AUserWidget : public AActor
 
 	TSharedRef<SWidget> GetRootWidget();
 
+	TSharedPtr<SWidget> GetWidgetFromName(const FString& Name) const;
+	USlateWrapperComponent* GetHandleFromName(const FString& Name) const;
+
 private:
 	TSharedPtr<SWidget> RootWidget;
-	TMap< SWidget*, USlateWrapperComponent* > WidgetToComponent;
+	TMap< TWeakPtr<SWidget>, USlateWrapperComponent* > WidgetToComponent;
 
 private:
 	void RebuildWrapperWidget();
diff --git a/Engine/Source/Runtime/UMG/Public/Components/CanvasPanelSlot.h b/Engine/Source/Runtime/UMG/Public/Components/CanvasPanelSlot.h
--- a/Engine/Source/Runtime/UMG/Public/Components/CanvasPanelSlot.h
+++ b/Engine/Source/Runtime/UMG/Public/Components/CanvasPanelSlot.h
@@ -47,4 +47,8 @@ class UMG_API UCanvasPanelSlot : public UPanelSlot
 	*/
 	UPROPERTY(EditAnywhere, Category=Appearance)
 	TEnumAsByte<EVerticalAlignment> VerticalAlignment;
+
+	virtual void Resize(const FVector2D& Direction, const FVector2D& Amount) OVERRIDE;
+
+	virtual bool CanResize(const FVector2D& Direction) const OVERRIDE;
 };
\ No newline at end of file
diff --git a/Engine/Source/Runtime/UMG/Public/Components/PanelSlot.h b/Engine/Source/Runtime/UMG/Public/Components/PanelSlot.h
--- a/Engine/Source/Runtime/UMG/Public/Components/PanelSlot.h
+++ b/Engine/Source/Runtime/UMG/Public/Components/PanelSlot.h
@@ -11,4 +11,8 @@ class UMG_API UPanelSlot : public UObject
 
 	UPROPERTY()
 	class USlateWrapperComponent* Content;
+
+	virtual void Resize(const FVector2D& Direction, const FVector2D& Amount);
+
+	virtual bool CanResize(const FVector2D& Direction) const;
 };
