UGM - General cleanup and removal of unused code, and some reorganization of the editor.

diff --git a/Engine/Source/Editor/UMGEditor/Private/SUMGDesigner.cpp b/Engine/Source/Editor/UMGEditor/Private/SUMGDesigner.cpp
--- a/Engine/Source/Editor/UMGEditor/Private/SUMGDesigner.cpp
+++ b/Engine/Source/Editor/UMGEditor/Private/SUMGDesigner.cpp
@@ -372,12 +372,12 @@ void SUMGDesigner::DrawDragHandles(const FPaintGeometry& SelectionGeometry, cons
 {
 	if ( SelectedTemplate && SelectedTemplate->Slot )
 	{
-
 		float X = SelectionGeometry.DrawPosition.X;
 		float Y = SelectionGeometry.DrawPosition.Y;
 		float Width = SelectionGeometry.DrawSize.X;
 		float Height = SelectionGeometry.DrawSize.Y;
 
+		// @TODO UMG Handles should come from the slot/container to express how its slots can be transformed.
 		TArray<FVector2D> Handles;
 		Handles.Add(FVector2D(X, Y));					// Top - Left
 		Handles.Add(FVector2D(X + Width * 0.5f, Y));	// Top - Middle
diff --git a/Engine/Source/Editor/UMGEditor/Private/SUMGEditorTree.cpp b/Engine/Source/Editor/UMGEditor/Private/SUMGEditorTree.cpp
--- a/Engine/Source/Editor/UMGEditor/Private/SUMGEditorTree.cpp
+++ b/Engine/Source/Editor/UMGEditor/Private/SUMGEditorTree.cpp
@@ -3,8 +3,6 @@
 #include "UMGEditorPrivatePCH.h"
 
 #include "SUMGEditorTree.h"
-#include "UMGEditor.h"
-#include "UMGEditorViewportClient.h"
 #include "UMGEditorActions.h"
 #include "SUMGEditorTreeItem.h"
 
diff --git a/Engine/Source/Editor/UMGEditor/Private/SUMGEditorTree.h b/Engine/Source/Editor/UMGEditor/Private/SUMGEditorTree.h
--- a/Engine/Source/Editor/UMGEditor/Private/SUMGEditorTree.h
+++ b/Engine/Source/Editor/UMGEditor/Private/SUMGEditorTree.h
@@ -3,7 +3,6 @@
 #pragma once
 
 #include "SCompoundWidget.h"
-#include "IUMGEditor.h"
 #include "BlueprintEditor.h"
 
 //TODO rename SUMGEditorHierarchy
diff --git a/Engine/Source/Editor/UMGEditor/Private/SUMGEditorTreeItem.cpp b/Engine/Source/Editor/UMGEditor/Private/SUMGEditorTreeItem.cpp
--- a/Engine/Source/Editor/UMGEditor/Private/SUMGEditorTreeItem.cpp
+++ b/Engine/Source/Editor/UMGEditor/Private/SUMGEditorTreeItem.cpp
@@ -3,8 +3,6 @@
 #include "UMGEditorPrivatePCH.h"
 
 #include "SUMGEditorTreeItem.h"
-#include "UMGEditor.h"
-#include "UMGEditorViewportClient.h"
 #include "UMGEditorActions.h"
 #include "WidgetTemplateDragDropOp.h"
 
diff --git a/Engine/Source/Editor/UMGEditor/Private/SUMGEditorTreeItem.h b/Engine/Source/Editor/UMGEditor/Private/SUMGEditorTreeItem.h
--- a/Engine/Source/Editor/UMGEditor/Private/SUMGEditorTreeItem.h
+++ b/Engine/Source/Editor/UMGEditor/Private/SUMGEditorTreeItem.h
@@ -3,7 +3,6 @@
 #pragma once
 
 #include "SCompoundWidget.h"
-#include "IUMGEditor.h"
 #include "BlueprintEditor.h"
 
 //TODO rename SUMGEditorHierarchy
diff --git a/Engine/Source/Editor/UMGEditor/Private/SUMGEditorViewport.cpp b/Engine/Source/Editor/UMGEditor/Private/SUMGEditorViewport.cpp
deleted file mode 100644
diff --git a/Engine/Source/Editor/UMGEditor/Private/SUMGEditorViewport.h b/Engine/Source/Editor/UMGEditor/Private/SUMGEditorViewport.h
deleted file mode 100644
diff --git a/Engine/Source/Editor/UMGEditor/Private/SUMGEditorWidgetTemplates.cpp b/Engine/Source/Editor/UMGEditor/Private/SUMGEditorWidgetTemplates.cpp
--- a/Engine/Source/Editor/UMGEditor/Private/SUMGEditorWidgetTemplates.cpp
+++ b/Engine/Source/Editor/UMGEditor/Private/SUMGEditorWidgetTemplates.cpp
@@ -3,8 +3,6 @@
 #include "UMGEditorPrivatePCH.h"
 
 #include "SUMGEditorWidgetTemplates.h"
-#include "UMGEditor.h"
-#include "UMGEditorViewportClient.h"
 #include "UMGEditorActions.h"
 
 #include "PreviewScene.h"
diff --git a/Engine/Source/Editor/UMGEditor/Private/SUMGEditorWidgetTemplates.h b/Engine/Source/Editor/UMGEditor/Private/SUMGEditorWidgetTemplates.h
--- a/Engine/Source/Editor/UMGEditor/Private/SUMGEditorWidgetTemplates.h
+++ b/Engine/Source/Editor/UMGEditor/Private/SUMGEditorWidgetTemplates.h
@@ -3,7 +3,6 @@
 #pragma once
 
 #include "SCompoundWidget.h"
-#include "IUMGEditor.h"
 #include "BlueprintEditor.h"
 
 #include "WidgetTemplate.h"
diff --git a/Engine/Source/Editor/UMGEditor/Private/Templates/WidgetTemplate.cpp b/Engine/Source/Editor/UMGEditor/Private/Templates/WidgetTemplate.cpp
new file mode 100644
--- /dev/null
+++ b/Engine/Source/Editor/UMGEditor/Private/Templates/WidgetTemplate.cpp
@@ -0,0 +1,9 @@
+// Copyright 1998-2014 Epic Games, Inc. All Rights Reserved.
+
+#include "UMGEditorPrivatePCH.h"
+
+#include "WidgetTemplate.h"
+
+FWidgetTemplate::FWidgetTemplate()
+{
+}
diff --git a/Engine/Source/Editor/UMGEditor/Private/Templates/WidgetTemplate.h b/Engine/Source/Editor/UMGEditor/Private/Templates/WidgetTemplate.h
new file mode 100644
--- /dev/null
+++ b/Engine/Source/Editor/UMGEditor/Private/Templates/WidgetTemplate.h
@@ -0,0 +1,20 @@
+// Copyright 1998-2014 Epic Games, Inc. All Rights Reserved.
+
+#pragma once
+
+class FWidgetTemplate : public TSharedFromThis<FWidgetTemplate>
+{
+public:
+	FWidgetTemplate();
+
+	virtual FText GetCategory() = 0;
+
+	virtual USlateWrapperComponent* Create(class UWidgetTree* Tree) = 0;
+
+public:
+	FText Name;
+
+	FText ToolTip;
+
+	FSlateIcon Icon;
+};
diff --git a/Engine/Source/Editor/UMGEditor/Private/Templates/WidgetTemplateButton.cpp b/Engine/Source/Editor/UMGEditor/Private/Templates/WidgetTemplateButton.cpp
new file mode 100644
--- /dev/null
+++ b/Engine/Source/Editor/UMGEditor/Private/Templates/WidgetTemplateButton.cpp
@@ -0,0 +1,33 @@
+// Copyright 1998-2014 Epic Games, Inc. All Rights Reserved.
+
+#include "UMGEditorPrivatePCH.h"
+
+#include "WidgetTemplateButton.h"
+
+#define LOCTEXT_NAMESPACE "UMGEditor"
+
+FWidgetTemplateButton::FWidgetTemplateButton()
+{
+	Name = LOCTEXT("Button", "Button");
+	ToolTip = LOCTEXT("Button_Tooltip", "Button With Text");
+	//Icon = 
+}
+
+FText FWidgetTemplateButton::GetCategory()
+{
+	return LOCTEXT("Common", "Common");
+}
+
+USlateWrapperComponent* FWidgetTemplateButton::Create(UWidgetTree* Tree)
+{
+	UButtonComponent* Button = Tree->ConstructWidget<UButtonComponent>(UButtonComponent::StaticClass());
+	UTextBlockComponent* Text = Tree->ConstructWidget<UTextBlockComponent>(UTextBlockComponent::StaticClass());
+
+	Text->Text = FText::FromString(TEXT("Button Text"));
+	Text->Font.Size = 10;
+	Button->SetContent(Text);
+
+	return Button;
+}
+
+#undef LOCTEXT_NAMESPACE
diff --git a/Engine/Source/Editor/UMGEditor/Private/Templates/WidgetTemplateButton.h b/Engine/Source/Editor/UMGEditor/Private/Templates/WidgetTemplateButton.h
new file mode 100644
--- /dev/null
+++ b/Engine/Source/Editor/UMGEditor/Private/Templates/WidgetTemplateButton.h
@@ -0,0 +1,15 @@
+// Copyright 1998-2014 Epic Games, Inc. All Rights Reserved.
+
+#pragma once
+
+#include "WidgetTemplate.h"
+
+class FWidgetTemplateButton : public FWidgetTemplate
+{
+public:
+	FWidgetTemplateButton();
+
+	virtual FText GetCategory() OVERRIDE;
+
+	virtual USlateWrapperComponent* Create(UWidgetTree* Tree) OVERRIDE;
+};
diff --git a/Engine/Source/Editor/UMGEditor/Private/Templates/WidgetTemplateClass.cpp b/Engine/Source/Editor/UMGEditor/Private/Templates/WidgetTemplateClass.cpp
new file mode 100644
--- /dev/null
+++ b/Engine/Source/Editor/UMGEditor/Private/Templates/WidgetTemplateClass.cpp
@@ -0,0 +1,25 @@
+// Copyright 1998-2014 Epic Games, Inc. All Rights Reserved.
+
+#include "UMGEditorPrivatePCH.h"
+
+#include "WidgetTemplateClass.h"
+
+#define LOCTEXT_NAMESPACE "UMGEditor"
+
+FWidgetTemplateClass::FWidgetTemplateClass(TSubclassOf<USlateWrapperComponent> InWidgetClass)
+	: WidgetClass(InWidgetClass)
+{
+}
+
+FText FWidgetTemplateClass::GetCategory()
+{
+	// TODO this should be pulled from metadata from the class, that describes its categorization, is it a panel, a common control...etc?
+	return LOCTEXT("Class", "Class");
+}
+
+USlateWrapperComponent* FWidgetTemplateClass::Create(UWidgetTree* Tree)
+{
+	return Tree->ConstructWidget<USlateWrapperComponent>(WidgetClass);
+}
+
+#undef LOCTEXT_NAMESPACE
diff --git a/Engine/Source/Editor/UMGEditor/Private/Templates/WidgetTemplateClass.h b/Engine/Source/Editor/UMGEditor/Private/Templates/WidgetTemplateClass.h
new file mode 100644
--- /dev/null
+++ b/Engine/Source/Editor/UMGEditor/Private/Templates/WidgetTemplateClass.h
@@ -0,0 +1,18 @@
+// Copyright 1998-2014 Epic Games, Inc. All Rights Reserved.
+
+#pragma once
+
+#include "WidgetTemplate.h"
+
+class FWidgetTemplateClass : public FWidgetTemplate
+{
+public:
+	FWidgetTemplateClass(TSubclassOf<USlateWrapperComponent> InWidgetClass);
+
+	virtual FText GetCategory() OVERRIDE;
+
+	virtual USlateWrapperComponent* Create(UWidgetTree* Tree) OVERRIDE;
+
+protected:
+	TSubclassOf<USlateWrapperComponent> WidgetClass;
+};
diff --git a/Engine/Source/Editor/UMGEditor/Private/UMGEditor.cpp b/Engine/Source/Editor/UMGEditor/Private/UMGEditor.cpp
deleted file mode 100644
diff --git a/Engine/Source/Editor/UMGEditor/Private/UMGEditor.h b/Engine/Source/Editor/UMGEditor/Private/UMGEditor.h
deleted file mode 100644
diff --git a/Engine/Source/Editor/UMGEditor/Private/UMGEditorActions.cpp b/Engine/Source/Editor/UMGEditor/Private/UMGEditorActions.cpp
--- a/Engine/Source/Editor/UMGEditor/Private/UMGEditorActions.cpp
+++ b/Engine/Source/Editor/UMGEditor/Private/UMGEditorActions.cpp
@@ -2,9 +2,7 @@
 
 #include "UMGEditorPrivatePCH.h"
 
-#include "UMGEditor.h"
 #include "UMGEditorActions.h"
-#include "SUMGEditorViewport.h"
 
 void FUMGEditorCommands::RegisterCommands()
 {
diff --git a/Engine/Source/Editor/UMGEditor/Private/UMGEditorModule.cpp b/Engine/Source/Editor/UMGEditor/Private/UMGEditorModule.cpp
--- a/Engine/Source/Editor/UMGEditor/Private/UMGEditorModule.cpp
+++ b/Engine/Source/Editor/UMGEditor/Private/UMGEditorModule.cpp
@@ -4,7 +4,7 @@
 
 #include "UMGEditorModule.h"
 #include "ModuleManager.h"
-#include "UMGEditor.h"
+
 #include "AssetToolsModule.h"
 #include "UMGBlueprintEditorExtensionHook.h"
 #include "AssetTypeActions_WidgetBlueprint.h"
@@ -79,35 +79,6 @@ public:
 		return FReply::Unhandled();
 	}
 
-	virtual TSharedRef<IUMGEditor> CreateUMGEditor(const EToolkitMode::Type Mode, const TSharedPtr< IToolkitHost >& InitToolkitHost, UParticleSystem* ParticleSystem) OVERRIDE
-	{
-		TSharedRef<FUMGEditor> NewUMGEditor(new FUMGEditor());
-		NewUMGEditor->Initialize(Mode, InitToolkitHost, ParticleSystem);
-		UMGEditorToolkits.Add(&*NewUMGEditor);
-		return NewUMGEditor;
-	}
-
-	virtual void UMGEditorClosed(FUMGEditor* UMGEditorInstance) OVERRIDE
-	{
-		UMGEditorToolkits.Remove(UMGEditorInstance);
-	}
-
-	virtual void RefreshUMGEditor(UParticleSystem* ParticleSystem) OVERRIDE
-	{
-		for (int32 Idx = 0; Idx < UMGEditorToolkits.Num(); ++Idx)
-		{
-			//if (UMGEditorToolkits[Idx]->GetParticleSystem() == ParticleSystem)
-			//{
-			//	UMGEditorToolkits[Idx]->ForceUpdate();
-			//}
-		}
-	}
-
-	virtual void ConvertModulesToSeeded(UParticleSystem* ParticleSystem) OVERRIDE
-	{
-		//FUMGEditor::ConvertAllModulesToSeeded( ParticleSystem );
-	}
-
 	/** Gets the extensibility managers for outside entities to extend gui page editor's menus and toolbars */
 	virtual TSharedPtr<FExtensibilityManager> GetMenuExtensibilityManager() {return MenuExtensibilityManager;}
 	virtual TSharedPtr<FExtensibilityManager> GetToolBarExtensibilityManager() {return ToolBarExtensibilityManager;}
diff --git a/Engine/Source/Editor/UMGEditor/Private/UMGEditorViewportClient.cpp b/Engine/Source/Editor/UMGEditor/Private/UMGEditorViewportClient.cpp
deleted file mode 100644
diff --git a/Engine/Source/Editor/UMGEditor/Private/UMGEditorViewportClient.h b/Engine/Source/Editor/UMGEditor/Private/UMGEditorViewportClient.h
deleted file mode 100644
diff --git a/Engine/Source/Editor/UMGEditor/Private/WidgetTemplate.cpp b/Engine/Source/Editor/UMGEditor/Private/WidgetTemplate.cpp
deleted file mode 100644
diff --git a/Engine/Source/Editor/UMGEditor/Private/WidgetTemplate.h b/Engine/Source/Editor/UMGEditor/Private/WidgetTemplate.h
deleted file mode 100644
diff --git a/Engine/Source/Editor/UMGEditor/Private/WidgetTemplateButton.cpp b/Engine/Source/Editor/UMGEditor/Private/WidgetTemplateButton.cpp
deleted file mode 100644
diff --git a/Engine/Source/Editor/UMGEditor/Private/WidgetTemplateButton.h b/Engine/Source/Editor/UMGEditor/Private/WidgetTemplateButton.h
deleted file mode 100644
diff --git a/Engine/Source/Editor/UMGEditor/Private/WidgetTemplateClass.cpp b/Engine/Source/Editor/UMGEditor/Private/WidgetTemplateClass.cpp
deleted file mode 100644
diff --git a/Engine/Source/Editor/UMGEditor/Private/WidgetTemplateClass.h b/Engine/Source/Editor/UMGEditor/Private/WidgetTemplateClass.h
deleted file mode 100644
diff --git a/Engine/Source/Editor/UMGEditor/Public/IUMGEditor.h b/Engine/Source/Editor/UMGEditor/Public/IUMGEditor.h
deleted file mode 100644
diff --git a/Engine/Source/Editor/UMGEditor/Public/UMGEditorModule.h b/Engine/Source/Editor/UMGEditor/Public/UMGEditorModule.h
--- a/Engine/Source/Editor/UMGEditor/Public/UMGEditorModule.h
+++ b/Engine/Source/Editor/UMGEditor/Public/UMGEditorModule.h
@@ -6,7 +6,6 @@
 #include "ModuleInterface.h"
 #include "Toolkits/AssetEditorToolkit.h"
 #include "Toolkits/IToolkit.h"
-#include "IUMGEditor.h"
 
 extern const FName UMGEditorAppIdentifier;
 
@@ -15,15 +14,5 @@ class FUMGEditor;
 class IUMGEditorModule : public IModuleInterface, public IHasMenuExtensibility, public IHasToolBarExtensibility
 {
 public:
-	/** Creates a new UMGEditor instance */
-	virtual TSharedRef<IUMGEditor> CreateUMGEditor(const EToolkitMode::Type Mode, const TSharedPtr< IToolkitHost >& InitToolkitHost, UParticleSystem* ParticleSystem) = 0;
 
-	/** Removes the specified instance from the list of open UMGEditor toolkits */
-	virtual void UMGEditorClosed(FUMGEditor* UMGEditorInstance) = 0;
-
-	/** Refreshes the toolkit inspecting the specified particle system */
-	virtual void RefreshUMGEditor(UParticleSystem* ParticleSystem) = 0;
-
-	/** Converts all the modules in the specified particle system to seeded modules */
-	virtual void ConvertModulesToSeeded(UParticleSystem* ParticleSystem) = 0;
 };
diff --git a/Engine/Source/Editor/UMGEditor/UMGEditor.Build.cs b/Engine/Source/Editor/UMGEditor/UMGEditor.Build.cs
--- a/Engine/Source/Editor/UMGEditor/UMGEditor.Build.cs
+++ b/Engine/Source/Editor/UMGEditor/UMGEditor.Build.cs
@@ -6,19 +6,20 @@ public class UMGEditor : ModuleRules
 {
 	public UMGEditor(TargetInfo Target)
 	{
-        PrivateIncludePaths.AddRange(
-            new string[] {
-                "Editor/GraphEditor/Private",
+		PrivateIncludePaths.AddRange(
+			new string[] {
+				"Editor/GraphEditor/Private",
 				"Editor/Kismet/Private",
-                "Developer/AssetTools/Private/AssetTypeActions",
-                "Editor/UMGEditor/Private" // For PCH includes (because they don't work with relative paths, yet)
-            });
+				"Developer/AssetTools/Private/AssetTypeActions",
+				"Editor/UMGEditor/Private", // For PCH includes (because they don't work with relative paths, yet)
+				"Editor/UMGEditor/Private/Templates",
+			});
 
-        PrivateIncludePathModuleNames.AddRange(
-            new string[] {
-                "AssetTools",
-                "UMG",
-            });
+		PrivateIncludePathModuleNames.AddRange(
+			new string[] {
+				"AssetTools",
+				"UMG",
+			});
 
 		PublicIncludePaths.AddRange(
 			new string[] {
@@ -33,19 +34,19 @@ public class UMGEditor : ModuleRules
 				"CoreUObject",
 				"Slate",
 				"Engine",
-                "InputCore",
+				"InputCore",
 				"AssetTools",
 				"UnrealEd", // for FAssetEditorManager
 				"KismetWidgets",
-                "KismetCompiler",
-                "BlueprintGraph",
+				"KismetCompiler",
+				"BlueprintGraph",
 				"GraphEditor",
 				"Kismet",  // for FWorkflowCentricApplication
 				"PropertyEditor",
 				"RenderCore",
 				"LevelEditor", // for EdModes to get a toolkit host  //@TODO: PAPER: Should be a better way to do this (see the @todo in EdModeTileMap.cpp)
 				"UMG",
-                "DistCurveEditor",
+				"DistCurveEditor",
 				"ContentBrowser",
 				"WorkspaceMenuStructure",
 				"EditorStyle",
@@ -58,8 +59,8 @@ public class UMGEditor : ModuleRules
 			new string[] {
 				"MainFrame",
 				"PropertyEditor",
-                "AssetTools",
-                "UMG",
+				"AssetTools",
+				"UMG",
 			}
 			);
 	}
