UMG - Adding missing files.

diff --git a/Engine/Source/Editor/UMGEditor/Private/Templates/WidgetTemplateCheckBox.cpp b/Engine/Source/Editor/UMGEditor/Private/Templates/WidgetTemplateCheckBox.cpp
new file mode 100644
--- /dev/null
+++ b/Engine/Source/Editor/UMGEditor/Private/Templates/WidgetTemplateCheckBox.cpp
@@ -0,0 +1,36 @@
+// Copyright 1998-2014 Epic Games, Inc. All Rights Reserved.
+
+#include "UMGEditorPrivatePCH.h"
+
+#include "WidgetTemplateCheckBox.h"
+
+#define LOCTEXT_NAMESPACE "UMGEditor"
+
+FWidgetTemplateCheckBox::FWidgetTemplateCheckBox()
+{
+	Name = LOCTEXT("CheckBoxLabel", "CheckBox & Label");
+	ToolTip = LOCTEXT("CheckBox_Tooltip", "CheckBox with Text");
+	//Icon = 
+}
+
+FText FWidgetTemplateCheckBox::GetCategory()
+{
+	return LOCTEXT("Common", "Common");
+}
+
+USlateWrapperComponent* FWidgetTemplateCheckBox::Create(UWidgetTree* Tree)
+{
+	UHorizontalBoxComponent* Horizontal = Tree->ConstructWidget<UHorizontalBoxComponent>(UHorizontalBoxComponent::StaticClass());
+	UCheckBoxComponent* Checkbox = Tree->ConstructWidget<UCheckBoxComponent>(UCheckBoxComponent::StaticClass());
+	UTextBlockComponent* Text = Tree->ConstructWidget<UTextBlockComponent>(UTextBlockComponent::StaticClass());
+	Text->Text = LOCTEXT("CheckboxText", "Checkbox Text");
+	Text->Font.Size = 10;
+
+	Horizontal->AddSlot(Checkbox)
+		->Size = FSlateChildSize(ESlateSizeRule::Automatic);
+	Horizontal->AddSlot(Text);
+
+	return Horizontal;
+}
+
+#undef LOCTEXT_NAMESPACE
diff --git a/Engine/Source/Editor/UMGEditor/Private/Templates/WidgetTemplateCheckBox.h b/Engine/Source/Editor/UMGEditor/Private/Templates/WidgetTemplateCheckBox.h
new file mode 100644
--- /dev/null
+++ b/Engine/Source/Editor/UMGEditor/Private/Templates/WidgetTemplateCheckBox.h
@@ -0,0 +1,15 @@
+// Copyright 1998-2014 Epic Games, Inc. All Rights Reserved.
+
+#pragma once
+
+#include "WidgetTemplate.h"
+
+class FWidgetTemplateCheckBox : public FWidgetTemplate
+{
+public:
+	FWidgetTemplateCheckBox();
+
+	virtual FText GetCategory() OVERRIDE;
+
+	virtual USlateWrapperComponent* Create(UWidgetTree* Tree) OVERRIDE;
+};
