UMG - Checkpoint - Prototyping function binding pins and function binding through customization, lots of TODOs.

diff --git a/Engine/Source/Editor/Kismet/Public/SKismetInspector.h b/Engine/Source/Editor/Kismet/Public/SKismetInspector.h
--- a/Engine/Source/Editor/Kismet/Public/SKismetInspector.h
+++ b/Engine/Source/Editor/Kismet/Public/SKismetInspector.h
@@ -60,6 +60,8 @@ public:
 	/** Turns on or off details customization for components */
 	void EnableComponentDetailsCustomization(bool bEnable);
 
+	TSharedPtr<class IDetailsView> GetPropertyView() const { return PropertyView; }
+
 protected:
 	/** Update the inspector window to show information on the supplied objects */
 	void UpdateFromObjects(const TArray<UObject*>& PropertyObjects, struct FKismetSelectionInfo& SelectionInfo, const FShowDetailsOptions& Options);
diff --git a/Engine/Source/Editor/KismetCompiler/Private/KismetCompiler.cpp b/Engine/Source/Editor/KismetCompiler/Private/KismetCompiler.cpp
--- a/Engine/Source/Editor/KismetCompiler/Private/KismetCompiler.cpp
+++ b/Engine/Source/Editor/KismetCompiler/Private/KismetCompiler.cpp
@@ -162,7 +162,7 @@ void FKismetCompilerContext::CleanAndSanitizeClass(UBlueprintGeneratedClass* Cla
 		ParentClass = UObject::StaticClass();
 	}
 	TransientClass->ClassAddReferencedObjects = ParentClass->AddReferencedObjects;
-    
+	
 	NewClass = ClassToClean;
 	OldCDO = ClassToClean->ClassDefaultObject; // we don't need to create the CDO at this point
 	
@@ -2114,21 +2114,21 @@ void FKismetCompilerContext::CreateFunctionStubForEvent(UK2Node_Event* SrcEventN
 	}
 
 	// If this is a customizable event, make sure to copy over the user defined pins
- 	if (UK2Node_CustomEvent const* SrcCustomEventNode = Cast<UK2Node_CustomEvent const>(SrcEventNode))
- 	{
- 		EntryNode->UserDefinedPins = SrcCustomEventNode->UserDefinedPins;
+	if (UK2Node_CustomEvent const* SrcCustomEventNode = Cast<UK2Node_CustomEvent const>(SrcEventNode))
+	{
+		EntryNode->UserDefinedPins = SrcCustomEventNode->UserDefinedPins;
 		// CustomEvents may inherit net flags (so let's use their GetNetFlags() incase this is an override)
 		StubContext.MarkAsNetFunction(SrcCustomEventNode->GetNetFlags());
- 	}
+	}
 	EntryNode->AllocateDefaultPins();
 
 	// Confirm that the event node matches the latest function signature, which the newly created EntryNode should have
- 	if( !SrcEventNode->IsFunctionEntryCompatible(EntryNode) )
- 	{
- 		// There is no match, so the function parameters must have changed.  Throw an error, and force them to refresh
- 		MessageLog.Error(*FString::Printf(*LOCTEXT("EventNodeOutOfDate_Error", "Event node @@ is out-of-date.  Please refresh it.").ToString()), SrcEventNode);
- 		return;
- 	}
+	if( !SrcEventNode->IsFunctionEntryCompatible(EntryNode) )
+	{
+		// There is no match, so the function parameters must have changed.  Throw an error, and force them to refresh
+		MessageLog.Error(*FString::Printf(*LOCTEXT("EventNodeOutOfDate_Error", "Event node @@ is out-of-date.  Please refresh it.").ToString()), SrcEventNode);
+		return;
+	}
 
 	// Copy each event parameter to the assignment node, if there are any inputs
 	UK2Node_VariableSet* AssignmentNode = NULL;
@@ -2173,7 +2173,7 @@ void FKismetCompilerContext::CreateFunctionStubForEvent(UK2Node_Event* SrcEventN
 	CallIntoUbergraph->NodePosX = 300;
 
 	// Use the ExecuteUbergraph base function to generate the pins...
- 	CallIntoUbergraph->FunctionReference.SetExternalMember(Schema->FN_ExecuteUbergraphBase, UObject::StaticClass());
+	CallIntoUbergraph->FunctionReference.SetExternalMember(Schema->FN_ExecuteUbergraphBase, UObject::StaticClass());
 	CallIntoUbergraph->AllocateDefaultPins();
 	
 	// ...then swap to the generated version for this level
@@ -2799,6 +2799,11 @@ void FKismetCompilerContext::CreateFunctionList()
 	}
 }
 
+FKismetFunctionContext* FKismetCompilerContext::CreateFunctionContext()
+{
+	return new (FunctionList)FKismetFunctionContext(MessageLog, Schema, NewClass, Blueprint);
+}
+
 /** Compile a blueprint into a class and a set of functions */
 void FKismetCompilerContext::Compile()
 {
@@ -2837,25 +2842,25 @@ void FKismetCompilerContext::Compile()
 	UBlueprintGeneratedClass* TargetClass = Cast<UBlueprintGeneratedClass>(TargetUClass);
 
 	// >>> Backwards Compatibility: Make sure that skeleton generated classes have the proper "SKEL_" naming convention
- 	const FString SkeletonPrefix(TEXT("SKEL_"));
- 	if( bIsSkeletonOnly && TargetClass && !TargetClass->GetName().StartsWith(SkeletonPrefix) )
- 	{
- 		FString NewName = SkeletonPrefix + TargetClass->GetName();
+	const FString SkeletonPrefix(TEXT("SKEL_"));
+	if( bIsSkeletonOnly && TargetClass && !TargetClass->GetName().StartsWith(SkeletonPrefix) )
+	{
+		FString NewName = SkeletonPrefix + TargetClass->GetName();
  
- 		// Ensure we have a free name for this class
- 		UClass* AnyClassWithGoodName = (UClass*)StaticFindObject(UClass::StaticClass(), Blueprint->GetOutermost(), *NewName, false);
- 		if( AnyClassWithGoodName )
- 		{
- 			// Special Case:  If the CDO of the class has become dissociated from its actual CDO, attempt to find the proper named CDO, and get rid of it.
- 			if( AnyClassWithGoodName->ClassDefaultObject == TargetClass->ClassDefaultObject )
- 			{
- 				AnyClassWithGoodName->ClassDefaultObject = NULL;
- 				FString DefaultObjectName = FString(DEFAULT_OBJECT_PREFIX) + NewName;
- 				AnyClassWithGoodName->ClassDefaultObject = (UObject*)StaticFindObject(UObject::StaticClass(), Blueprint->GetOutermost(), *DefaultObjectName, false);
- 			}
+		// Ensure we have a free name for this class
+		UClass* AnyClassWithGoodName = (UClass*)StaticFindObject(UClass::StaticClass(), Blueprint->GetOutermost(), *NewName, false);
+		if( AnyClassWithGoodName )
+		{
+			// Special Case:  If the CDO of the class has become dissociated from its actual CDO, attempt to find the proper named CDO, and get rid of it.
+			if( AnyClassWithGoodName->ClassDefaultObject == TargetClass->ClassDefaultObject )
+			{
+				AnyClassWithGoodName->ClassDefaultObject = NULL;
+				FString DefaultObjectName = FString(DEFAULT_OBJECT_PREFIX) + NewName;
+				AnyClassWithGoodName->ClassDefaultObject = (UObject*)StaticFindObject(UObject::StaticClass(), Blueprint->GetOutermost(), *DefaultObjectName, false);
+			}
  
 			// Get rid of the old class to make room for renaming our class to the final SKEL name
- 			FKismetCompilerUtilities::ConsignToOblivion(AnyClassWithGoodName, Blueprint->bIsRegeneratingOnLoad);
+			FKismetCompilerUtilities::ConsignToOblivion(AnyClassWithGoodName, Blueprint->bIsRegeneratingOnLoad);
 
 			// Update the refs to the old SKC
 			TMap<UObject*, UObject*> ClassReplacementMap;
@@ -2866,11 +2871,11 @@ void FKismetCompilerContext::Compile()
 			{
 				FArchiveReplaceObjectRef<UObject> ReplaceInBlueprintAr(AllGraphs[i], ClassReplacementMap, /*bNullPrivateRefs=*/ false, /*bIgnoreOuterRef=*/ false, /*bIgnoreArchetypeRef=*/ false);
 			}
- 		}
+		}
  
- 		ERenameFlags RenameFlags = (REN_DontCreateRedirectors|REN_NonTransactional|(Blueprint->bIsRegeneratingOnLoad ? REN_ForceNoResetLoaders : 0));
- 		TargetClass->Rename(*NewName, NULL, RenameFlags);
- 	}
+		ERenameFlags RenameFlags = (REN_DontCreateRedirectors|REN_NonTransactional|(Blueprint->bIsRegeneratingOnLoad ? REN_ForceNoResetLoaders : 0));
+		TargetClass->Rename(*NewName, NULL, RenameFlags);
+	}
 	// <<< End Backwards Compatibility
 
 	// >>> Backwards compatibility:  If SkeletonGeneratedClass == GeneratedClass, we need to make a new generated class the first time we need it
@@ -3252,29 +3257,29 @@ void FKismetCompilerContext::Compile()
 	}
 
 	// For full compiles, find other blueprints that may need refreshing, and mark them dirty, in case they try to run
-  	if( bIsFullCompile && !Blueprint->bIsRegeneratingOnLoad )
-  	{
-  		TArray<UObject*> AllBlueprints;
-  		GetObjectsOfClass(UBlueprint::StaticClass(), AllBlueprints, true);
+	if( bIsFullCompile && !Blueprint->bIsRegeneratingOnLoad )
+	{
+		TArray<UObject*> AllBlueprints;
+		GetObjectsOfClass(UBlueprint::StaticClass(), AllBlueprints, true);
   
-  		// Mark any blueprints that implement this interface as dirty
-  		for( auto CurrentObj = AllBlueprints.CreateIterator(); CurrentObj; ++CurrentObj )
-  		{
-  			UBlueprint* CurrentBP = Cast<UBlueprint>(*CurrentObj);
+		// Mark any blueprints that implement this interface as dirty
+		for( auto CurrentObj = AllBlueprints.CreateIterator(); CurrentObj; ++CurrentObj )
+		{
+			UBlueprint* CurrentBP = Cast<UBlueprint>(*CurrentObj);
   
-  			if( FBlueprintEditorUtils::IsBlueprintDependentOn(CurrentBP, Blueprint) )
-  			{
+			if( FBlueprintEditorUtils::IsBlueprintDependentOn(CurrentBP, Blueprint) )
+			{
 				CurrentBP->Status = BS_Dirty;
-  				FBlueprintEditorUtils::RefreshExternalBlueprintDependencyNodes(CurrentBP);
+				FBlueprintEditorUtils::RefreshExternalBlueprintDependencyNodes(CurrentBP);
 				CurrentBP->BroadcastChanged();
-  			}
-  		}
-  	}
+			}
+		}
+	}
 
 	// Clear out pseudo-local members that are only valid within a Compile call
 	UbergraphContext = NULL;
- 	CallsIntoUbergraph.Empty();
- 	TimelineToMemberVariableMap.Empty();
+	CallsIntoUbergraph.Empty();
+	TimelineToMemberVariableMap.Empty();
 
 
 	check(NewClass->PropertiesSize >= UObject::StaticClass()->PropertiesSize);
diff --git a/Engine/Source/Editor/KismetCompiler/Public/KismetCompiler.h b/Engine/Source/Editor/KismetCompiler/Public/KismetCompiler.h
--- a/Engine/Source/Editor/KismetCompiler/Public/KismetCompiler.h
+++ b/Engine/Source/Editor/KismetCompiler/Public/KismetCompiler.h
@@ -239,6 +239,9 @@ protected:
 	 */
 	virtual void CreateFunctionList();
 
+	/** Creates a new function context and adds it to the function list to be processed. */
+	FKismetFunctionContext* CreateFunctionContext();
+
 	/**
 	 * Merges macros/subgraphs into the graph and validates it, creating a function list entry if it's reasonable.
 	 */
diff --git a/Engine/Source/Editor/UMGEditor/Classes/WidgetGraphNode_Base.h b/Engine/Source/Editor/UMGEditor/Classes/WidgetGraphNode_Base.h
new file mode 100644
--- /dev/null
+++ b/Engine/Source/Editor/UMGEditor/Classes/WidgetGraphNode_Base.h
@@ -0,0 +1,67 @@
+// Copyright 1998-2014 Epic Games, Inc. All Rights Reserved.
+
+#pragma once
+
+#include "WidgetGraphNode_Base.generated.h"
+
+UCLASS(Abstract, MinimalAPI)
+class UWidgetGraphNode_Base : public UK2Node
+{
+	GENERATED_UCLASS_BODY()
+
+	UPROPERTY(EditAnywhere, BlueprintReadWrite, Category=PinOptions, EditFixedSize)
+	TArray<FOptionalPinFromProperty> ShowPinForProperties;
+
+	//// UObject interface
+	UMGEDITOR_API virtual void PostEditChangeProperty(struct FPropertyChangedEvent& PropertyChangedEvent) OVERRIDE;
+	//// End of UObject interface
+
+	//// UEdGraphNode interface
+	UMGEDITOR_API virtual void AllocateDefaultPins() OVERRIDE;
+	UMGEDITOR_API virtual FLinearColor GetNodeTitleColor() const OVERRIDE;
+	//UMGEDITOR_API virtual FString GetDocumentationLink() const OVERRIDE;
+	//UMGEDITOR_API virtual void GetPinHoverText(const UEdGraphPin& Pin, FString& HoverTextOut) const OVERRIDE;
+	UMGEDITOR_API virtual bool ShowPaletteIconOnNode() const OVERRIDE { return false; }
+	//	// End of UEdGraphNode interface
+
+
+	// UK2Node interface
+	virtual bool NodeCausesStructuralBlueprintChange() const OVERRIDE { return true; }
+	virtual bool ShouldShowNodeProperties() const OVERRIDE { return true; }
+	virtual bool CanPlaceBreakpoints() const OVERRIDE { return false; }
+	UMGEDITOR_API virtual void ReallocatePinsDuringReconstruction(TArray<UEdGraphPin*>& OldPins) OVERRIDE;
+	UMGEDITOR_API virtual bool CanCreateUnderSpecifiedSchema(const UEdGraphSchema* DesiredSchema) const OVERRIDE;
+	UMGEDITOR_API virtual void GetMenuEntries(FGraphContextMenuBuilder& ContextMenuBuilder) const OVERRIDE;
+	// End of UK2Node interface
+
+	// UWidgetGraphNode_Base interface
+
+	// Gets the menu category this node belongs in
+	UMGEDITOR_API virtual FString GetNodeCategory() const;
+
+	// Create any output pins necessary for this node
+	UMGEDITOR_API virtual void CreateOutputPins();
+
+	UMGEDITOR_API void GetPinAssociatedProperty(UScriptStruct* NodeType, UEdGraphPin* InputPin, UProperty*& OutProperty, int32& OutIndex);
+
+	// customize pin data based on the input
+	virtual void CustomizePinData(UEdGraphPin* Pin, FName SourcePropertyName, int32 ArrayIndex) const;
+
+	/** Get the animation blueprint to which this node belongs */
+	UWidgetBlueprint* GetWidgetBlueprint() const { return CastChecked<UWidgetBlueprint>(GetBlueprint()); }
+
+	// UWidgetGraphNode_Base
+
+protected:
+	friend class FWidgetBlueprintCompiler;
+
+	// Gets the animation FNode type represented by this ed graph node
+	UMGEDITOR_API UScriptStruct* GetFNodeType() const;
+
+	// Gets the animation FNode property represented by this ed graph node
+	UMGEDITOR_API UStructProperty* GetFNodeProperty() const;
+
+	void InternalPinCreation(TArray<UEdGraphPin*>* OldPins);
+
+	TSharedPtr<FEdGraphSchemaAction_K2NewNode> CreateDefaultMenuEntry(FGraphContextMenuBuilder& ContextMenuBuilder) const;
+};
diff --git a/Engine/Source/Editor/UMGEditor/Classes/WidgetGraphNode_Text.h b/Engine/Source/Editor/UMGEditor/Classes/WidgetGraphNode_Text.h
new file mode 100644
--- /dev/null
+++ b/Engine/Source/Editor/UMGEditor/Classes/WidgetGraphNode_Text.h
@@ -0,0 +1,25 @@
+// Copyright 1998-2014 Epic Games, Inc. All Rights Reserved.
+
+#pragma once
+
+#include "WidgetGraphNode_Text.generated.h"
+
+UCLASS(MinimalAPI)
+class UWidgetGraphNode_Text : public UWidgetGraphNode_Base
+{
+	GENERATED_UCLASS_BODY()
+
+	UPROPERTY(EditAnywhere, Category=Settings)
+	FWidgetNode_Text Node;
+
+	// UEdGraphNode interface
+	virtual FString GetTooltip() const OVERRIDE;
+	virtual FText GetNodeTitle(ENodeTitleType::Type TitleType) const OVERRIDE;
+	// End of UEdGraphNode interface
+
+	/*
+	// UAnimGraphNode_Base interface
+	virtual void CustomizePinData(UEdGraphPin* Pin, FName SourcePropertyName, int32 ArrayIndex) const OVERRIDE;
+	// End of UAnimGraphNode_Base interface
+	*/
+};
diff --git a/Engine/Source/Editor/UMGEditor/Classes/WidgetGraphSchema.h b/Engine/Source/Editor/UMGEditor/Classes/WidgetGraphSchema.h
new file mode 100644
--- /dev/null
+++ b/Engine/Source/Editor/UMGEditor/Classes/WidgetGraphSchema.h
@@ -0,0 +1,33 @@
+// Copyright 1998-2014 Epic Games, Inc. All Rights Reserved.
+
+#pragma once
+
+#include "EdGraphSchema_K2.h"
+#include "WidgetGraphSchema.generated.h"
+
+UCLASS(MinimalAPI)
+class UWidgetGraphSchema : public UEdGraphSchema_K2
+{
+	GENERATED_UCLASS_BODY()
+
+	UPROPERTY()
+	FName NAME_NeverAsPin;
+
+	UPROPERTY()
+	FName NAME_PinHiddenByDefault;
+
+	UPROPERTY()
+	FName NAME_PinShownByDefault;
+
+	UPROPERTY()
+	FName NAME_AlwaysAsPin;
+
+	UPROPERTY()
+	FName NAME_OnEvaluate;
+	
+	UPROPERTY()
+	FName DefaultEvaluationHandlerName;
+	
+	
+	// TODO UMG - Add more things specific to the needs of the widget blueprint graph.
+};
\ No newline at end of file
diff --git a/Engine/Source/Editor/UMGEditor/Private/DetailCustomizations.cpp b/Engine/Source/Editor/UMGEditor/Private/DetailCustomizations.cpp
new file mode 100644
--- /dev/null
+++ b/Engine/Source/Editor/UMGEditor/Private/DetailCustomizations.cpp
@@ -0,0 +1,259 @@
+// Copyright 1998-2014 Epic Games, Inc. All Rights Reserved.
+
+#include "UMGEditorPrivatePCH.h"
+
+#include "DetailCustomizations.h"
+
+#include "PropertyEditing.h"
+#include "ObjectEditorUtils.h"
+#include "WidgetGraphSchema.h"
+#include "ScopedTransaction.h"
+
+#define LOCTEXT_NAMESPACE "UMG"
+
+void FBlueprintWidgetCustomization::CustomizeDetails( IDetailLayoutBuilder& DetailLayout )
+{
+	TArray< TWeakObjectPtr<UObject> > OutObjects;
+	DetailLayout.GetObjectsBeingCustomized(OutObjects);
+
+	if ( OutObjects.Num() == 1 )
+	{
+		UClass* PropertyClass = OutObjects[0].Get()->GetClass();
+
+		for ( TFieldIterator<UProperty> PropertyIt(PropertyClass, EFieldIteratorFlags::IncludeSuper); PropertyIt; ++PropertyIt )
+		{
+			if ( UDelegateProperty* Property = Cast<UDelegateProperty>(*PropertyIt) )
+			{
+				if ( !Property->GetName().EndsWith(TEXT("Delegate")) )
+				{
+					continue;
+				}
+
+				FString ConstPropertyName = Property->GetName();
+				ConstPropertyName.RemoveFromEnd(TEXT("Delegate"));
+
+				//TSharedRef<IPropertyHandle> DelegatePropertyHandle = DetailLayout.GetProperty(Property->GetFName(), PropertyClass);
+				TSharedRef<IPropertyHandle> ConstPropertyHandle = DetailLayout.GetProperty(FName(*ConstPropertyName), PropertyClass);
+
+				const bool bHasValidHandle = ConstPropertyHandle->IsValidHandle();
+
+				if ( !bHasValidHandle )
+				{
+					continue;
+				}
+
+				UProperty* ConstProperty = ConstPropertyHandle->GetProperty();
+
+				const bool bIsEditable = ConstProperty->HasAnyPropertyFlags(CPF_Edit | CPF_EditConst);
+				const bool bDoSignaturesMatch = Property->SignatureFunction->GetReturnProperty()->SameType(ConstProperty);
+
+				if ( !(bIsEditable && bDoSignaturesMatch) )
+				{
+					continue;
+				}
+
+				IDetailCategoryBuilder& PropertyCategory = DetailLayout.EditCategory(FObjectEditorUtils::GetCategoryFName(ConstProperty));
+
+				IDetailPropertyRow& PropertyRow = PropertyCategory.AddProperty(ConstPropertyHandle);
+
+				TSharedPtr<SWidget> NameWidget;
+				TSharedPtr<SWidget> ValueWidget;
+				FDetailWidgetRow Row;
+				PropertyRow.GetDefaultWidgets(NameWidget, ValueWidget, Row);
+
+				TSharedRef<SComboButton> DelegateButton = SNew(SComboButton)
+				.OnGetMenuContent(this, &FBlueprintWidgetCustomization::OnGenerateDelegateMenu, ConstPropertyHandle)
+				.ContentPadding(0)
+				.ButtonContent()
+				[
+					SNew(STextBlock)
+					.Text(this, &FBlueprintWidgetCustomization::GetCurrentBindingText, ConstPropertyHandle)
+					.Font(IDetailLayoutBuilder::GetDetailFont())
+				];
+
+				const bool bShowChildren = true;
+				PropertyRow.CustomWidget(bShowChildren)
+				.NameContent()
+				.MinDesiredWidth(Row.NameWidget.MinWidth)
+				.MaxDesiredWidth(Row.NameWidget.MaxWidth)
+				[
+					NameWidget.ToSharedRef()
+				]
+				.ValueContent()
+				.MinDesiredWidth(Row.ValueWidget.MinWidth * 2)
+				.MaxDesiredWidth(Row.ValueWidget.MaxWidth * 2)
+				[
+					SNew(SHorizontalBox)
+
+					+ SHorizontalBox::Slot()
+					.FillWidth(1)
+					.Padding(FMargin(0, 0, 4, 0))
+					[
+						ValueWidget.ToSharedRef()
+					]
+					
+					+ SHorizontalBox::Slot()
+					.AutoWidth()
+					[
+						DelegateButton
+					]
+				];
+			}
+		}
+	}
+}
+
+void FBlueprintWidgetCustomization::RefreshBlueprintFunctionCache()
+{
+	const UWidgetGraphSchema* Schema = GetDefault<UWidgetGraphSchema>();
+	const FSlateFontInfo DetailFontInfo = IDetailLayoutBuilder::GetDetailFont();
+
+	BlueprintFunctionCache.Reset();
+
+	//TODO UMG Filter the list of functions to only be functions matching the signature.
+
+	// Grab functions implemented by the blueprint
+	for ( int32 i = 0; i < Blueprint->FunctionGraphs.Num(); i++ )
+	{
+		UEdGraph* Graph = Blueprint->FunctionGraphs[i];
+		check(Graph);
+
+		// Ignore the construction script, we just want user created functions.
+		if ( Graph->GetFName() == Schema->FN_UserConstructionScript )
+		{
+			continue;
+		}
+
+		FGraphDisplayInfo DisplayInfo;
+		Graph->GetSchema()->GetGraphDisplayInformation(*Graph, DisplayInfo);
+
+		TSharedPtr<FunctionInfo> NewFuncAction = MakeShareable(new FunctionInfo());
+		NewFuncAction->DisplayName = DisplayInfo.PlainName;
+		NewFuncAction->Tooltip = DisplayInfo.Tooltip;
+		NewFuncAction->FuncName = Graph->GetFName();
+		NewFuncAction->EdGraph = Graph;
+
+		BlueprintFunctionCache.Add(NewFuncAction);
+	}
+}
+
+TSharedRef<SWidget> FBlueprintWidgetCustomization::OnGenerateDelegateMenu(TSharedRef<IPropertyHandle> PropertyHandle)
+{
+	RefreshBlueprintFunctionCache();
+
+	const bool bInShouldCloseWindowAfterMenuSelection = true;
+	FMenuBuilder MenuBuilder(bInShouldCloseWindowAfterMenuSelection, NULL);
+
+	MenuBuilder.BeginSection("ClearBinding");
+	{
+		MenuBuilder.AddMenuEntry(
+			LOCTEXT("FilterListResetFilters", "Remove Binding"),
+			LOCTEXT("FilterListResetToolTip", "Removes the current binding"),
+			FSlateIcon(),
+			FUIAction(FExecuteAction::CreateSP(this, &FBlueprintWidgetCustomization::HandleRemoveBinding, PropertyHandle))
+			);
+	}
+	MenuBuilder.EndSection(); //ClearBinding
+
+	MenuBuilder.BeginSection("Functions");
+	{
+		for ( TSharedPtr<FunctionInfo>& Function : BlueprintFunctionCache )
+		{
+			MenuBuilder.AddMenuEntry(
+				Function->DisplayName,
+				FText::FromString(Function->Tooltip),
+				FSlateIcon(),
+				FUIAction(FExecuteAction::CreateSP(this, &FBlueprintWidgetCustomization::HandleAddBinding, PropertyHandle, Function))
+				);
+		}
+	}
+	MenuBuilder.EndSection(); //Functions
+
+
+	FDisplayMetrics DisplayMetrics;
+	FSlateApplication::Get().GetDisplayMetrics(DisplayMetrics);
+
+	return
+		SNew(SVerticalBox)
+
+		+ SVerticalBox::Slot()
+		.MaxHeight(DisplayMetrics.PrimaryDisplayHeight * 0.5)
+		[
+			MenuBuilder.MakeWidget()
+		];;
+}
+
+FText FBlueprintWidgetCustomization::GetCurrentBindingText(TSharedRef<IPropertyHandle> PropertyHandle) const
+{
+	TArray<UObject*> OuterObjects;
+	PropertyHandle->GetOuterObjects(OuterObjects);
+
+	//TODO UMG O(N) Isn't good for this, needs to be map, but map isn't serialized, need cached runtime map for fast lookups.
+
+	FName PropertyName = PropertyHandle->GetProperty()->GetFName();
+	for ( int32 ObjectIndex = 0; ObjectIndex < OuterObjects.Num(); ObjectIndex++ )
+	{
+		//TODO UMG handle multiple things selected
+
+		for ( const FDelegateEditorBinding& Binding : Blueprint->Bindings )
+		{
+			if ( Binding.PropertyName == PropertyName )
+			{
+				FName FoundName = Blueprint->GetFieldNameFromClassByGuid<UFunction>(Blueprint->GeneratedClass, Binding.MemberGuid);
+				return FText::FromName(FoundName);
+				//return FText::FromName(Binding.FunctionName);
+			}
+		}
+
+		//TODO UMG Do something about missing functions, little exclamation points if they're missing and such.
+
+		break;
+	}
+
+	return LOCTEXT("Bind", "Bind");
+}
+
+void FBlueprintWidgetCustomization::HandleRemoveBinding(TSharedRef<IPropertyHandle> PropertyHandle)
+{
+	const FScopedTransaction Transaction(LOCTEXT("BindDelegate", "Remove Binding"));
+
+	Blueprint->Modify();
+
+	TArray<UObject*> OuterObjects;
+	PropertyHandle->GetOuterObjects(OuterObjects);
+	for ( UObject* SelectedObject : OuterObjects )
+	{
+		FDelegateEditorBinding Binding;
+		Binding.ObjectName = SelectedObject->GetName();
+		Binding.PropertyName = PropertyHandle->GetProperty()->GetFName();
+
+		Blueprint->Bindings.Remove(Binding);
+	}
+
+	FBlueprintEditorUtils::MarkBlueprintAsStructurallyModified(Blueprint);
+}
+
+void FBlueprintWidgetCustomization::HandleAddBinding(TSharedRef<IPropertyHandle> PropertyHandle, TSharedPtr<FunctionInfo> SelectedFunction)
+{
+	const FScopedTransaction Transaction(LOCTEXT("BindDelegate", "Set Binding"));
+
+	Blueprint->Modify();
+
+	TArray<UObject*> OuterObjects;
+	PropertyHandle->GetOuterObjects(OuterObjects);
+	for ( UObject* SelectedObject : OuterObjects )
+	{
+		FDelegateEditorBinding Binding;
+		Binding.ObjectName = SelectedObject->GetName();
+		Binding.PropertyName = PropertyHandle->GetProperty()->GetFName();
+		Binding.FunctionName = SelectedFunction->FuncName;
+		Binding.MemberGuid = SelectedFunction->EdGraph->GraphGuid;
+
+		Blueprint->Bindings.Remove(Binding);
+		Blueprint->Bindings.AddUnique(Binding);
+	}
+
+	FBlueprintEditorUtils::MarkBlueprintAsStructurallyModified(Blueprint);
+}
+
+#undef LOCTEXT_NAMESPACE
diff --git a/Engine/Source/Editor/UMGEditor/Private/DetailCustomizations.h b/Engine/Source/Editor/UMGEditor/Private/DetailCustomizations.h
new file mode 100644
--- /dev/null
+++ b/Engine/Source/Editor/UMGEditor/Private/DetailCustomizations.h
@@ -0,0 +1,45 @@
+// Copyright 1998-2014 Epic Games, Inc. All Rights Reserved.
+
+#pragma once
+
+class FBlueprintWidgetCustomization : public IDetailCustomization
+{
+public:
+	/** Makes a new instance of this detail layout class for a specific detail view requesting it */
+	static TSharedRef<class IDetailCustomization> MakeInstance(UBlueprint* Blueprint)
+	{
+		return MakeShareable(new FBlueprintWidgetCustomization(Blueprint));
+	}
+
+	FBlueprintWidgetCustomization(UBlueprint* Blueprint)
+		: Blueprint(CastChecked<UWidgetBlueprint>(Blueprint))
+	{
+	}
+	
+	/** IDetailCustomization interface */
+	virtual void CustomizeDetails( IDetailLayoutBuilder& DetailLayout ) OVERRIDE;
+
+private:
+	struct FunctionInfo
+	{
+		FText DisplayName;
+		FString Tooltip;
+
+		FName FuncName;
+		UEdGraph* EdGraph;
+	};
+
+	void RefreshBlueprintFunctionCache();
+
+	TSharedRef<SWidget> OnGenerateDelegateMenu(TSharedRef<IPropertyHandle> PropertyHandle);
+
+	FText GetCurrentBindingText(TSharedRef<IPropertyHandle> PropertyHandle) const;
+	void HandleRemoveBinding(TSharedRef<IPropertyHandle> PropertyHandle);
+	void HandleAddBinding(TSharedRef<IPropertyHandle> PropertyHandle, TSharedPtr<FunctionInfo> SelectedFunction);
+
+private:
+
+	UWidgetBlueprint* Blueprint;
+
+	TArray< TSharedPtr<FunctionInfo> > BlueprintFunctionCache;
+};
diff --git a/Engine/Source/Editor/UMGEditor/Private/UMGBlueprintEditorExtensionHook.cpp b/Engine/Source/Editor/UMGEditor/Private/UMGBlueprintEditorExtensionHook.cpp
--- a/Engine/Source/Editor/UMGEditor/Private/UMGBlueprintEditorExtensionHook.cpp
+++ b/Engine/Source/Editor/UMGEditor/Private/UMGBlueprintEditorExtensionHook.cpp
@@ -15,6 +15,9 @@
 #include "SUMGDesigner.h"
 #include "SUMGEditorWidgetTemplates.h"
 #include "SKismetInspector.h"
+#include "IDetailsView.h"
+
+#include "DetailCustomizations.h"
 
 #define LOCTEXT_NAMESPACE "UMG_EXTENSION"
 
@@ -231,6 +234,11 @@ public:
 		FBlueprintComponentsApplicationMode::PostActivateMode();
 
 		GetBlueprintEditor()->GetInspector()->EnableComponentDetailsCustomization(false);
+
+		TSharedRef<class SKismetInspector> Inspector = MyBlueprintEditor.Pin()->GetInspector();
+		FOnGetDetailCustomizationInstance LayoutDelegateDetails = FOnGetDetailCustomizationInstance::CreateStatic(&FBlueprintWidgetCustomization::MakeInstance, MyBlueprintEditor.Pin()->GetBlueprintObj());
+		Inspector->GetPropertyView()->RegisterInstancedCustomPropertyLayout(USlateWrapperComponent::StaticClass(), LayoutDelegateDetails);
+		//Inspector->GetPropertyView()->RegisterInstancedCustomPropertyLayout(UDelegateProperty::StaticClass(), LayoutDelegateDetails);
 	}
 
 	UWidgetBlueprint* GetBlueprint() const
@@ -251,6 +259,44 @@ public:
 	}
 };
 
+//////////////////////////////////////////////////////////////////////////
+// FGraphEditorModeOverride
+
+class FGraphEditorModeOverride : public FBlueprintEditorApplicationMode
+{
+public:
+	FGraphEditorModeOverride(TSharedPtr<class FBlueprintEditor> InBlueprintEditor, FName InModeName, const bool bRegisterViewport = true, const bool bRegisterDefaultsTab = true)
+		: FBlueprintEditorApplicationMode(InBlueprintEditor, InModeName, bRegisterViewport, bRegisterDefaultsTab)
+	{}
+
+	virtual void PreDeactivateMode() OVERRIDE
+	{
+		FBlueprintEditorApplicationMode::PreDeactivateMode();
+	}
+
+	virtual void PostActivateMode() OVERRIDE
+	{
+		FBlueprintEditorApplicationMode::PostActivateMode();
+	}
+
+	UWidgetBlueprint* GetBlueprint() const
+	{
+		if ( FBlueprintEditor* Editor = MyBlueprintEditor.Pin().Get() )
+		{
+			return Cast<UWidgetBlueprint>(Editor->GetBlueprintObj());
+		}
+		else
+		{
+			return NULL;
+		}
+	}
+
+	TSharedPtr<FBlueprintEditor> GetBlueprintEditor() const
+	{
+		return MyBlueprintEditor.Pin();
+	}
+};
+
 //////////////////////////////////////////////////////////////////////////
 // FBlueprintEditorExtenderForSlatePreview
 
@@ -272,6 +318,19 @@ public:
 				}
 			}
 		}
+		else if ( ModeName == FBlueprintEditorApplicationModes::StandardBlueprintEditorMode )
+		{
+			//@TODO: Bit of a lie - push GetBlueprint up, or pass in editor!
+			auto LieMode = StaticCastSharedRef<FGraphEditorModeOverride>(InMode);
+
+			if ( UBlueprint* BP = LieMode->GetBlueprint() )
+			{
+				if ( BP->ParentClass->IsChildOf(AUserWidget::StaticClass()) )
+				{
+					return MakeShareable(new FGraphEditorModeOverride(LieMode->GetBlueprintEditor(), ModeName));
+				}
+			}
+		}
 
 		return InMode;
 	}
diff --git a/Engine/Source/Editor/UMGEditor/Private/UMGEditorPrivatePCH.h b/Engine/Source/Editor/UMGEditor/Private/UMGEditorPrivatePCH.h
--- a/Engine/Source/Editor/UMGEditor/Private/UMGEditorPrivatePCH.h
+++ b/Engine/Source/Editor/UMGEditor/Private/UMGEditorPrivatePCH.h
@@ -5,6 +5,7 @@
 #include "Slate.h"
 #include "SlateCore.h"
 #include "UnrealEd.h"
+#include "BlueprintGraphDefinitions.h"
 
 #include "UMG.h"
 
diff --git a/Engine/Source/Editor/UMGEditor/Private/WidgetBlueprintCompiler.cpp b/Engine/Source/Editor/UMGEditor/Private/WidgetBlueprintCompiler.cpp
--- a/Engine/Source/Editor/UMGEditor/Private/WidgetBlueprintCompiler.cpp
+++ b/Engine/Source/Editor/UMGEditor/Private/WidgetBlueprintCompiler.cpp
@@ -6,6 +6,50 @@
 #include "Kismet2NameValidators.h"
 #include "KismetReinstanceUtilities.h"
 
+template<>
+FString FNetNameMapping::MakeBaseName<UWidgetGraphNode_Base>(const UWidgetGraphNode_Base* Net)
+{
+	return FString::Printf(TEXT("%s_%s"), *Net->GetDescriptiveCompiledName(), *Net->NodeGuid.ToString());
+}
+
+///-------------------------------------------------------------
+
+FWidgetBlueprintCompiler::FEvaluationHandlerRecord::FEvaluationHandlerRecord()
+	: NodeVariableProperty(NULL)
+	, EvaluationHandlerProperty(NULL)
+	, HandlerFunctionName(NAME_None)
+{
+}
+
+bool FWidgetBlueprintCompiler::FEvaluationHandlerRecord::IsValid() const
+{
+	return ( NodeVariableProperty != NULL ) && ( EvaluationHandlerProperty != NULL );
+}
+
+void FWidgetBlueprintCompiler::FEvaluationHandlerRecord::PatchFunctionNameInto(UObject* TargetObject) const
+{
+	EvaluationHandlerProperty->ContainerPtrToValuePtr<FExposedValueHandler>(
+		NodeVariableProperty->ContainerPtrToValuePtr<void>(TargetObject)
+		)->BoundFunction = HandlerFunctionName;
+}
+
+void FWidgetBlueprintCompiler::FEvaluationHandlerRecord::RegisterPin(UEdGraphPin* SourcePin, UProperty* AssociatedProperty, int32 AssociatedPropertyArrayIndex)
+{
+	FAnimNodeSinglePropertyHandler& Handler = ServicedProperties.FindOrAdd(AssociatedProperty->GetFName());
+
+	if ( AssociatedPropertyArrayIndex != INDEX_NONE )
+	{
+		Handler.ArrayPins.Add(AssociatedPropertyArrayIndex, SourcePin);
+	}
+	else
+	{
+		check(Handler.SinglePin == NULL);
+		Handler.SinglePin = SourcePin;
+	}
+}
+
+///-------------------------------------------------------------
+
 FWidgetBlueprintCompiler::FWidgetBlueprintCompiler(UWidgetBlueprint* SourceSketch, FCompilerResultsLog& InMessageLog, const FKismetCompilerOptions& InCompilerOptions, TArray<UObject*>* InObjLoaded)
 	: Super(SourceSketch, InMessageLog, InCompilerOptions, InObjLoaded)
 {
@@ -15,6 +59,469 @@ FWidgetBlueprintCompiler::~FWidgetBlueprintCompiler()
 {
 }
 
+void FWidgetBlueprintCompiler::MergeUbergraphPagesIn(UEdGraph* Ubergraph)
+{
+	Super::MergeUbergraphPagesIn(Ubergraph);
+
+	// Build the raw node list
+	TArray<UWidgetGraphNode_Base*> WidgetNodeList;
+	ConsolidatedEventGraph->GetNodesOfClass<UWidgetGraphNode_Base>(/*out*/ WidgetNodeList);
+	
+	AllocateNodeIndexCounter = 0;
+
+	// Process the remaining nodes
+	for ( UWidgetGraphNode_Base* VisualWidgetNode : WidgetNodeList )
+	{
+		ProcessWidgetNode(VisualWidgetNode);
+	}
+}
+
+void FWidgetBlueprintCompiler::ProcessWidgetNode(UWidgetGraphNode_Base* VisualWidgetNode)
+{
+	const UWidgetGraphSchema* Schema = GetDefault<UWidgetGraphSchema>();
+
+	// Early out if this node has already been processed
+	//if ( AllocatedAnimNodes.Contains(VisualWidgetNode) )
+	//{
+	//	return;
+	//}
+
+	// Make sure the visual node has a runtime node template
+	UScriptStruct* NodeType = VisualWidgetNode->GetFNodeType();
+	if ( NodeType == NULL )
+	{
+		MessageLog.Error(TEXT("@@ has no widget value node member"), VisualWidgetNode);
+		return;
+	}
+
+	// Give the visual node a chance to do validation
+	{
+		const int32 PreValidationErrorCount = MessageLog.NumErrors;
+		//VisualWidgetNode->ValidateAnimNodeDuringCompilation(VisualWidgetNode->TargetSkeleton, MessageLog);
+		//VisualWidgetNode->BakeDataDuringCompilation(MessageLog);
+		if ( MessageLog.NumErrors != PreValidationErrorCount )
+		{
+			return;
+		}
+	}
+
+	// Create a property for the node
+	const FString NodeVariableName = ClassScopeNetNameMap.MakeValidName(VisualWidgetNode);
+
+	FEdGraphPinType NodeVariableType;
+	NodeVariableType.PinCategory = Schema->PC_Struct;
+	NodeVariableType.PinSubCategoryObject = NodeType;
+
+	UStructProperty* NewProperty = Cast<UStructProperty>(CreateVariable(FName(*NodeVariableName), NodeVariableType));
+
+	if ( NewProperty == NULL )
+	{
+		MessageLog.Error(TEXT("Failed to create node property for @@"), VisualWidgetNode);
+	}
+
+	// Register this node with the compile-time data structures
+	const int32 AllocatedIndex = AllocateNodeIndexCounter++;
+	AllocatedAnimNodes.Add(VisualWidgetNode, NewProperty);
+	AllocatedNodePropertiesToNodes.Add(NewProperty, VisualWidgetNode);
+	AllocatedAnimNodeIndices.Add(VisualWidgetNode, AllocatedIndex);
+	AllocatedPropertiesByIndex.Add(AllocatedIndex, NewProperty);
+
+	UWidgetGraphNode_Base* TrueSourceObject = MessageLog.FindSourceObjectTypeChecked<UWidgetGraphNode_Base>(VisualWidgetNode);
+	SourceNodeToProcessedNodeMap.Add(TrueSourceObject, VisualWidgetNode);
+
+
+	// TODO UMG add debug data.
+	// Register the slightly more permanent debug information
+//	NewAnimBlueprintClass->GetAnimBlueprintDebugData().NodePropertyToIndexMap.Add(TrueSourceObject, AllocatedIndex);
+//	NewAnimBlueprintClass->GetDebugData().RegisterClassPropertyAssociation(TrueSourceObject, NewProperty);
+
+	// Record pose pins for later patchup and gather pins that have an associated evaluation handler
+	TMap<FName, FEvaluationHandlerRecord> EvaluationHandlers;
+
+	for ( auto SourcePinIt = VisualWidgetNode->Pins.CreateIterator(); SourcePinIt; ++SourcePinIt )
+	{
+		UEdGraphPin* SourcePin = *SourcePinIt;
+		bool bConsumed = false;
+
+		// Does this pin have an associated evaluation handler?
+		UProperty* SourcePinProperty;
+		int32 SourceArrayIndex;
+		VisualWidgetNode->GetPinAssociatedProperty(NodeType, SourcePin, /*out*/ SourcePinProperty, /*out*/ SourceArrayIndex);
+
+		if ( SourcePinProperty != NULL )
+		{
+			if ( SourcePin->LinkedTo.Num() == 0 )
+			{
+				// Literal that can be pushed into the CDO instead of re-evaluated every frame
+				new (ValidAnimNodePinConstants)FEffectiveConstantRecord(NewProperty, SourcePin, SourcePinProperty, SourceArrayIndex);
+				bConsumed = true;
+			}
+			else
+			{
+				// Dynamic value that needs to be wired up and evaluated each frame
+				FString EvaluationHandlerStr = SourcePinProperty->GetMetaData(Schema->NAME_OnEvaluate);
+				FName EvaluationHandlerName(*EvaluationHandlerStr);
+				if ( EvaluationHandlerName == NAME_None )
+				{
+					EvaluationHandlerName = Schema->DefaultEvaluationHandlerName;
+				}
+
+				EvaluationHandlers.FindOrAdd(EvaluationHandlerName).RegisterPin(SourcePin, SourcePinProperty, SourceArrayIndex);
+				bConsumed = true;
+			}
+		}
+
+		if ( !bConsumed && ( SourcePin->Direction == EGPD_Input ) )
+		{
+			//@TODO: ANIMREFACTOR: It's probably OK to have certain pins ignored eventually, but this is very helpful during development
+			MessageLog.Note(TEXT("@@ was visible but ignored"), SourcePin);
+		}
+	}
+
+	// Match the associated property to each evaluation handler
+	for ( TFieldIterator<UProperty> NodePropIt(NodeType); NodePropIt; ++NodePropIt )
+	{
+		if ( UStructProperty* StructProp = Cast<UStructProperty>(*NodePropIt) )
+		{
+			if ( StructProp->Struct == FExposedWidgetValueHandler::StaticStruct() )
+			{
+				// Register this property to the list of pins that need to be updated
+				// (it's OK if there isn't an entry for this handler, it means that the values are static and don't need to be calculated every frame)
+				FName EvaluationHandlerName(StructProp->GetFName());
+				if ( FEvaluationHandlerRecord* pRecord = EvaluationHandlers.Find(EvaluationHandlerName) )
+				{
+					pRecord->NodeVariableProperty = NewProperty;
+					pRecord->EvaluationHandlerProperty = StructProp;
+				}
+			}
+		}
+	}
+
+	// Generate a new event to update the value of these properties
+	for ( auto HandlerIt = EvaluationHandlers.CreateIterator(); HandlerIt; ++HandlerIt )
+	{
+		FName EvaluationHandlerName = HandlerIt.Key();
+		FEvaluationHandlerRecord& Record = HandlerIt.Value();
+
+		if ( Record.IsValid() )
+		{
+			CreateEvaluationHandler(VisualWidgetNode, Record);
+			ValidEvaluationHandlerList.Add(Record);
+		}
+		else
+		{
+			MessageLog.Error(*FString::Printf(TEXT("A property on @@ references a non-existent %s property named %s"), *( Schema->NAME_OnEvaluate.ToString() ), *( EvaluationHandlerName.ToString() )), VisualWidgetNode );
+		}
+	}
+}
+
+//void FWidgetBlueprintCompiler::CreateEvaluationHandler(UWidgetGraphNode_Base* VisualAnimNode, FEvaluationHandlerRecord& Record)
+//{
+//	// Shouldn't create a handler if there is nothing to work with
+//	check(Record.ServicedProperties.Num() > 0);
+//	check(Record.NodeVariableProperty != NULL);
+//
+//
+//	//@TODO: Want to name these better
+//	const FString FunctionName = FString::Printf(TEXT("%s_%s_%s"), *Record.EvaluationHandlerProperty->GetName(), *VisualAnimNode->GetOuter()->GetName(), *VisualAnimNode->GetName());
+//	Record.HandlerFunctionName = FName(*FunctionName);
+//	
+//
+//	//UWidgetBlueprint* Blueprint = WidgetBlueprint();
+//
+//
+//	// Create the stub graph and add it to the list of functions to compile
+//
+//	UEdGraph* ChildStubGraph = NewNamedObject<UEdGraph>(Blueprint, Record.HandlerFunctionName);
+//	Blueprint->EventGraphs.Add(ChildStubGraph);
+//	ChildStubGraph->Schema = UEdGraphSchema_K2::StaticClass();
+//	ChildStubGraph->SetFlags(RF_Transient);
+//	MessageLog.NotifyIntermediateObjectCreation(ChildStubGraph, VisualAnimNode);
+//
+//	FKismetFunctionContext* StubContext = CreateFunctionContext();
+//	StubContext->SourceGraph = ChildStubGraph;
+//
+//	// A stub graph has no visual representation and is thus not suited to be debugged via the debugger
+//	StubContext->bCreateDebugData = false;
+//	StubContext->MarkAsInternalOrCppUseOnly();
+//
+//
+//
+//
+//
+//
+//
+//
+//
+//
+//
+//	
+//
+//	// Add a custom event in the graph
+//	//UK2Node_CustomEvent* EntryNode = SpawnIntermediateNode<UK2Node_CustomEvent>(VisualAnimNode, ConsolidatedEventGraph);
+//	//EntryNode->bInternalEvent = true;
+//	//EntryNode->CustomFunctionName = Record.HandlerFunctionName;
+//	UK2Node_FunctionEntry* EntryNode = SpawnIntermediateNode<UK2Node_FunctionEntry>(VisualAnimNode, ChildStubGraph);
+//	//EntryNode->SignatureClass = UObject::StaticClass();
+//	//EntryNode->SignatureName = Schema->FN_ExecuteUbergraphBase;
+//	EntryNode->CustomGeneratedFunctionName = Record.HandlerFunctionName;
+//	EntryNode->AllocateDefaultPins();
+//
+//	// The ExecChain is the current exec output pin in the linear chain
+//	UEdGraphPin* ExecChain = Schema->FindExecutionPin(*EntryNode, EGPD_Output);
+//
+//	// Create a struct member write node to store the parameters into the animation node
+//	UK2Node_StructMemberSet* AssignmentNode = SpawnIntermediateNode<UK2Node_StructMemberSet>(VisualAnimNode, ChildStubGraph);
+//	AssignmentNode->VariableReference.SetSelfMember(Record.NodeVariableProperty->GetFName());
+//	AssignmentNode->StructType = Record.NodeVariableProperty->Struct;
+//	AssignmentNode->AllocateDefaultPins();
+//
+//	// Wire up the variable node execution wires
+//	//UEdGraphPin* ExecVariablesIn = Schema->FindExecutionPin(*AssignmentNode, EGPD_Input);
+//	//ExecChain->MakeLinkTo(ExecVariablesIn);
+//	//ExecChain = Schema->FindExecutionPin(*AssignmentNode, EGPD_Output);
+//
+//	UK2Node_FunctionResult* ReturnNode = SpawnIntermediateNode<UK2Node_FunctionResult>(VisualAnimNode, ChildStubGraph);
+//	ReturnNode->AllocateDefaultPins();
+//
+//	// Wire up the variable node execution wires
+//	UEdGraphPin* ExecReturnIn = Schema->FindExecutionPin(*ReturnNode, EGPD_Input);
+//	ExecChain->MakeLinkTo(ExecReturnIn);
+//	ExecChain = Schema->FindExecutionPin(*ReturnNode, EGPD_Output);
+//
+//	for ( auto TargetPinIt = AssignmentNode->Pins.CreateIterator(); TargetPinIt; ++TargetPinIt )
+//	{
+//		UEdGraphPin* TargetPin = *TargetPinIt;
+//		FString PropertyNameStr = TargetPin->PinName;
+//		FName PropertyName(*PropertyNameStr);
+//
+//		// Does it get serviced by this handler?
+//		if ( FAnimNodeSinglePropertyHandler* SourceInfo = Record.ServicedProperties.Find(PropertyName) )
+//		{
+//			if ( TargetPin->PinType.bIsArray )
+//			{
+//				// Grab the array that we need to set members for
+//				UK2Node_StructMemberGet* FetchArrayNode = SpawnIntermediateNode<UK2Node_StructMemberGet>(VisualAnimNode, ChildStubGraph);
+//				FetchArrayNode->VariableReference.SetSelfMember(Record.NodeVariableProperty->GetFName());
+//				FetchArrayNode->StructType = Record.NodeVariableProperty->Struct;
+//				FetchArrayNode->AllocatePinsForSingleMemberGet(PropertyName);
+//
+//				UEdGraphPin* ArrayVariableNode = FetchArrayNode->FindPin(PropertyNameStr);
+//
+//				if ( SourceInfo->ArrayPins.Num() > 0 )
+//				{
+//					// Set each element in the array
+//					for ( auto SourcePinIt = SourceInfo->ArrayPins.CreateIterator(); SourcePinIt; ++SourcePinIt )
+//					{
+//						int32 ArrayIndex = SourcePinIt.Key();
+//						UEdGraphPin* SourcePin = SourcePinIt.Value();
+//
+//						// Create an array element set node
+//						UK2Node_CallArrayFunction* ArrayNode = SpawnIntermediateNode<UK2Node_CallArrayFunction>(VisualAnimNode, ChildStubGraph);
+//						ArrayNode->FunctionReference.SetExternalMember(GET_FUNCTION_NAME_CHECKED(UKismetArrayLibrary, Array_Set), UKismetArrayLibrary::StaticClass());
+//						ArrayNode->AllocateDefaultPins();
+//
+//						// Connect the execution chain
+//						ExecChain->MakeLinkTo(ArrayNode->GetExecPin());
+//						ExecChain = ArrayNode->GetThenPin();
+//
+//						// Connect the input array
+//						UEdGraphPin* TargetArrayPin = ArrayNode->FindPinChecked(TEXT("TargetArray"));
+//						TargetArrayPin->MakeLinkTo(ArrayVariableNode);
+//						ArrayNode->PinConnectionListChanged(TargetArrayPin);
+//
+//						// Set the array index
+//						UEdGraphPin* TargetIndexPin = ArrayNode->FindPinChecked(TEXT("Index"));
+//						TargetIndexPin->DefaultValue = FString::FromInt(ArrayIndex);
+//
+//						// Wire up the data input
+//						UEdGraphPin* TargetItemPin = ArrayNode->FindPinChecked(TEXT("Item"));
+//						TargetItemPin->CopyPersistentDataFromOldPin(*SourcePin);
+//						MessageLog.NotifyIntermediateObjectCreation(TargetItemPin, SourcePin);
+//						SourcePin->BreakAllPinLinks();
+//					}
+//				}
+//			}
+//			else
+//			{
+//				// Single property
+//				if ( SourceInfo->SinglePin != NULL )
+//				{
+//					UEdGraphPin* SourcePin = SourceInfo->SinglePin;
+//
+//					// First, add this pin to the user-defined pins
+//					TSharedPtr<FUserPinInfo> NewPinInfo = MakeShareable(new FUserPinInfo());
+//					NewPinInfo->PinName = TEXT("return");
+//					NewPinInfo->PinType = SourcePin->PinType;
+//					//UserDefinedPins.Add(NewPinInfo);
+//					TargetPin = ReturnNode->CreatePinFromUserDefinition(NewPinInfo);
+//
+//					//PropertiesBeingSet.Add(FName(*SourcePin->PinName));
+//					//TargetPin->CopyPersistentDataFromOldPin(*SourcePin);
+//					MessageLog.NotifyIntermediateObjectCreation(TargetPin, SourcePin);
+//					SourcePin->BreakAllPinLinks();
+//				}
+//			}
+//		}
+//	}
+//
+//	// Remove any unused pins from the assignment node to avoid smashing constant values
+//	//for ( int32 PinIndex = 0; PinIndex < AssignmentNode->ShowPinForProperties.Num(); ++PinIndex )
+//	//{
+//	//	FOptionalPinFromProperty& TestProperty = AssignmentNode->ShowPinForProperties[PinIndex];
+//	//	TestProperty.bShowPin = PropertiesBeingSet.Contains(TestProperty.PropertyName);
+//	//}
+//
+//	AssignmentNode->ReconstructNode();
+//	ReturnNode->ReconstructNode();
+//
+//	// TODO
+//}
+
+void FWidgetBlueprintCompiler::CreateEvaluationHandler(UWidgetGraphNode_Base* VisualAnimNode, FEvaluationHandlerRecord& Record)
+{
+	// Shouldn't create a handler if there is nothing to work with
+	check(Record.ServicedProperties.Num() > 0);
+	check(Record.NodeVariableProperty != NULL);
+
+	//@TODO: Want to name these better
+	const FString FunctionName = FString::Printf(TEXT("%s_%s_%s"), *Record.EvaluationHandlerProperty->GetName(), *VisualAnimNode->GetOuter()->GetName(), *VisualAnimNode->GetName());
+	Record.HandlerFunctionName = FName(*FunctionName);
+
+	// Add a custom event in the graph
+	UK2Node_CustomEvent* EntryNode = SpawnIntermediateNode<UK2Node_CustomEvent>(VisualAnimNode, ConsolidatedEventGraph);
+	EntryNode->bInternalEvent = true;
+	EntryNode->CustomFunctionName = Record.HandlerFunctionName;
+	EntryNode->AllocateDefaultPins();
+
+	// The ExecChain is the current exec output pin in the linear chain
+	UEdGraphPin* ExecChain = Schema->FindExecutionPin(*EntryNode, EGPD_Output);
+
+	// Create a struct member write node to store the parameters into the animation node
+	UK2Node_StructMemberSet* AssignmentNode = SpawnIntermediateNode<UK2Node_StructMemberSet>(VisualAnimNode, ConsolidatedEventGraph);
+	AssignmentNode->VariableReference.SetSelfMember(Record.NodeVariableProperty->GetFName());
+	AssignmentNode->StructType = Record.NodeVariableProperty->Struct;
+	AssignmentNode->AllocateDefaultPins();
+
+	// Wire up the variable node execution wires
+	UEdGraphPin* ExecVariablesIn = Schema->FindExecutionPin(*AssignmentNode, EGPD_Input);
+	ExecChain->MakeLinkTo(ExecVariablesIn);
+	ExecChain = Schema->FindExecutionPin(*AssignmentNode, EGPD_Output);
+
+	// Run thru each property
+	TSet<FName> PropertiesBeingSet;
+
+	//UK2Node_FunctionResult* ReturnNode = SpawnIntermediateNode<UK2Node_FunctionResult>(VisualAnimNode, ConsolidatedEventGraph);
+	//ReturnNode->AllocateDefaultPins();
+
+	//// Wire up the variable node execution wires
+	//UEdGraphPin* ExecReturnIn = Schema->FindExecutionPin(*ReturnNode, EGPD_Input);
+	//ExecChain->MakeLinkTo(ExecReturnIn);
+	//ExecChain = Schema->FindExecutionPin(*ReturnNode, EGPD_Output);
+
+	for ( auto TargetPinIt = AssignmentNode->Pins.CreateIterator(); TargetPinIt; ++TargetPinIt )
+	{
+		UEdGraphPin* TargetPin = *TargetPinIt;
+		FString PropertyNameStr = TargetPin->PinName;
+		FName PropertyName(*PropertyNameStr);
+
+		// Does it get serviced by this handler?
+		if ( FAnimNodeSinglePropertyHandler* SourceInfo = Record.ServicedProperties.Find(PropertyName) )
+		{
+			if ( TargetPin->PinType.bIsArray )
+			{
+				// Grab the array that we need to set members for
+				UK2Node_StructMemberGet* FetchArrayNode = SpawnIntermediateNode<UK2Node_StructMemberGet>(VisualAnimNode, ConsolidatedEventGraph);
+				FetchArrayNode->VariableReference.SetSelfMember(Record.NodeVariableProperty->GetFName());
+				FetchArrayNode->StructType = Record.NodeVariableProperty->Struct;
+				FetchArrayNode->AllocatePinsForSingleMemberGet(PropertyName);
+
+				UEdGraphPin* ArrayVariableNode = FetchArrayNode->FindPin(PropertyNameStr);
+
+				if ( SourceInfo->ArrayPins.Num() > 0 )
+				{
+					// Set each element in the array
+					for ( auto SourcePinIt = SourceInfo->ArrayPins.CreateIterator(); SourcePinIt; ++SourcePinIt )
+					{
+						int32 ArrayIndex = SourcePinIt.Key();
+						UEdGraphPin* SourcePin = SourcePinIt.Value();
+
+						// Create an array element set node
+						UK2Node_CallArrayFunction* ArrayNode = SpawnIntermediateNode<UK2Node_CallArrayFunction>(VisualAnimNode, ConsolidatedEventGraph);
+						ArrayNode->FunctionReference.SetExternalMember(GET_FUNCTION_NAME_CHECKED(UKismetArrayLibrary, Array_Set), UKismetArrayLibrary::StaticClass());
+						ArrayNode->AllocateDefaultPins();
+
+						// Connect the execution chain
+						ExecChain->MakeLinkTo(ArrayNode->GetExecPin());
+						ExecChain = ArrayNode->GetThenPin();
+
+						// Connect the input array
+						UEdGraphPin* TargetArrayPin = ArrayNode->FindPinChecked(TEXT("TargetArray"));
+						TargetArrayPin->MakeLinkTo(ArrayVariableNode);
+						ArrayNode->PinConnectionListChanged(TargetArrayPin);
+
+						// Set the array index
+						UEdGraphPin* TargetIndexPin = ArrayNode->FindPinChecked(TEXT("Index"));
+						TargetIndexPin->DefaultValue = FString::FromInt(ArrayIndex);
+
+						// Wire up the data input
+						UEdGraphPin* TargetItemPin = ArrayNode->FindPinChecked(TEXT("Item"));
+						TargetItemPin->CopyPersistentDataFromOldPin(*SourcePin);
+						MessageLog.NotifyIntermediateObjectCreation(TargetItemPin, SourcePin);
+						SourcePin->BreakAllPinLinks();
+					}
+				}
+			}
+			else
+			{
+				// Single property
+				if ( SourceInfo->SinglePin != NULL )
+				{
+					UEdGraphPin* SourcePin = SourceInfo->SinglePin;
+
+					PropertiesBeingSet.Add(FName(*SourcePin->PinName));
+					TargetPin->CopyPersistentDataFromOldPin(*SourcePin);
+					MessageLog.NotifyIntermediateObjectCreation(TargetPin, SourcePin);
+					SourcePin->BreakAllPinLinks();
+
+
+					//// First, add this pin to the user-defined pins
+					//TSharedPtr<FUserPinInfo> NewPinInfo = MakeShareable(new FUserPinInfo());
+					//NewPinInfo->PinName = TEXT("return");
+					//NewPinInfo->PinType = SourcePin->PinType;
+
+					//TargetPin = ReturnNode->CreatePinFromUserDefinition(NewPinInfo);
+					//TargetPin->CopyPersistentDataFromOldPin(*SourcePin);
+					//MessageLog.NotifyIntermediateObjectCreation(TargetPin, SourcePin);
+
+					//SourcePin->BreakAllPinLinks();
+				}
+			}
+		}
+	}
+
+	// Remove any unused pins from the assignment node to avoid smashing constant values
+	for ( int32 PinIndex = 0; PinIndex < AssignmentNode->ShowPinForProperties.Num(); ++PinIndex )
+	{
+		FOptionalPinFromProperty& TestProperty = AssignmentNode->ShowPinForProperties[PinIndex];
+		TestProperty.bShowPin = PropertiesBeingSet.Contains(TestProperty.PropertyName);
+	}
+	AssignmentNode->ReconstructNode();
+
+}
+
+void FWidgetBlueprintCompiler::CopyTermDefaultsToDefaultObject(UObject* DefaultObject)
+{
+	Super::CopyTermDefaultsToDefaultObject(DefaultObject);
+
+	// And patch evaluation function entry names
+	for ( auto EvalLinkIt = ValidEvaluationHandlerList.CreateIterator(); EvalLinkIt; ++EvalLinkIt )
+	{
+		FEvaluationHandlerRecord& Record = *EvalLinkIt;
+		Record.PatchFunctionNameInto(DefaultObject);
+	}
+}
+
 void FWidgetBlueprintCompiler::ValidateWidgetNames()
 {
 	UWidgetBlueprint* Blueprint = WidgetBlueprint();
@@ -113,6 +620,12 @@ void FWidgetBlueprintCompiler::FinishCompilingClass(UClass* Class)
 	UWidgetBlueprintGeneratedClass* BPGClass = CastChecked<UWidgetBlueprintGeneratedClass>(Class);
 	BPGClass->WidgetTree = NULL;
 	BPGClass->WidgetTree = Blueprint->WidgetTree;
+	BPGClass->Bindings.Reset();
+
+	for ( const FDelegateEditorBinding& EditorBinding : Blueprint->Bindings )
+	{
+		BPGClass->Bindings.Add(EditorBinding.ToRuntimeBinding(Blueprint));
+	}
 
 	Super::FinishCompilingClass(Class);
 }
diff --git a/Engine/Source/Editor/UMGEditor/Private/WidgetBlueprintCompiler.h b/Engine/Source/Editor/UMGEditor/Private/WidgetBlueprintCompiler.h
--- a/Engine/Source/Editor/UMGEditor/Private/WidgetBlueprintCompiler.h
+++ b/Engine/Source/Editor/UMGEditor/Private/WidgetBlueprintCompiler.h
@@ -26,20 +26,122 @@ protected:
 	void ValidateWidgetNames();
 
 	// FKismetCompilerContext
+	//virtual UEdGraphSchema_K2* CreateSchema() OVERRIDE;
+	virtual void MergeUbergraphPagesIn(UEdGraph* Ubergraph) OVERRIDE;
 	virtual void SpawnNewClass(const FString& NewClassName) OVERRIDE;
 	virtual void EnsureProperGeneratedClass(UClass*& TargetClass) OVERRIDE;
 	virtual void CreateClassVariablesFromBlueprint() OVERRIDE;
 	virtual void FinishCompilingClass(UClass* Class) OVERRIDE;
 	virtual void CleanAndSanitizeClass(UBlueprintGeneratedClass* ClassToClean, UObject*& OldCDO) OVERRIDE;
 	virtual bool ValidateGeneratedClass(UBlueprintGeneratedClass* Class) OVERRIDE;
+	virtual void CopyTermDefaultsToDefaultObject(UObject* DefaultObject) OVERRIDE;
 	// End FKismetCompilerContext
 
 protected:
+	// Wireup record for a single anim node property (which might be an array)
+	struct FAnimNodeSinglePropertyHandler
+	{
+		TMap<int32, UEdGraphPin*> ArrayPins;
+
+		UEdGraphPin* SinglePin;
+
+		FAnimNodeSinglePropertyHandler()
+			: SinglePin(NULL)
+		{
+		}
+	};
+
+	// Record for a property that was exposed as a pin, but wasn't wired up (just a literal)
+	struct FEffectiveConstantRecord
+	{
+	public:
+		// The node variable that the handler is in
+		class UStructProperty* NodeVariableProperty;
+
+		// The property within the struct to set
+		class UProperty* ConstantProperty;
+
+		// The array index if ConstantProperty is an array property, or INDEX_NONE otherwise
+		int32 ArrayIndex;
+
+		// The pin to pull the DefaultValue/DefaultObject from
+		UEdGraphPin* LiteralSourcePin;
+
+		FEffectiveConstantRecord()
+			: NodeVariableProperty(NULL)
+			, ConstantProperty(NULL)
+			, ArrayIndex(INDEX_NONE)
+			, LiteralSourcePin(NULL)
+		{
+		}
+
+		FEffectiveConstantRecord(UStructProperty* ContainingNodeProperty, UEdGraphPin* SourcePin, UProperty* SourcePinProperty, int32 SourceArrayIndex)
+			: NodeVariableProperty(ContainingNodeProperty)
+			, ConstantProperty(SourcePinProperty)
+			, ArrayIndex(SourceArrayIndex)
+			, LiteralSourcePin(SourcePin)
+		{
+		}
+
+		bool Apply(UObject* Object);
+	};
+
+	// Struct used to record information about one evaluation handler
+	struct FEvaluationHandlerRecord
+	{
+	public:
+		// The node variable that the handler is in
+		class UStructProperty* NodeVariableProperty;
+
+		// The specific evaluation handler inside the specified node
+		class UStructProperty* EvaluationHandlerProperty;
+
+		// Set of properties serviced by this handler (Map from property name to the record for that property)
+		TMap<FName, FAnimNodeSinglePropertyHandler> ServicedProperties;
+
+		// The resulting function name
+		FName HandlerFunctionName;
+	public:
+		FEvaluationHandlerRecord();
+
+		bool IsValid() const;
+
+		void PatchFunctionNameInto(UObject* TargetObject) const;
+
+		void RegisterPin(UEdGraphPin* SourcePin, UProperty* AssociatedProperty, int32 AssociatedPropertyArrayIndex);
+	};
+
+protected:
+	void ProcessWidgetNode(UWidgetGraphNode_Base* VisualWidgetNode);
+
+	void CreateEvaluationHandler(UWidgetGraphNode_Base* VisualAnimNode, FEvaluationHandlerRecord& Record);
+
 	UWidgetBlueprintGeneratedClass* NewWidgetBlueprintClass;
-	//UWidgetBlueprint* AnimBlueprint;
 
 	// Map of properties created for widgets; to aid in debug data generation
 	TMap<class USlateWrapperComponent*, class UProperty*> WidgetToMemberVariableMap;
 
+	///----------------------------------------------------------------
+	// Callback system properties.
+
+	// Map of allocated v3 nodes that are members of the class
+	TMap<class UWidgetGraphNode_Base*, UProperty*> AllocatedAnimNodes;
+	TMap<UProperty*, class UWidgetGraphNode_Base*> AllocatedNodePropertiesToNodes;
+	TMap<int32, UProperty*> AllocatedPropertiesByIndex;
+
+	// Map of true source objects (user edited ones) to the cloned ones that are actually compiled
+	TMap<class UWidgetGraphNode_Base*, UWidgetGraphNode_Base*> SourceNodeToProcessedNodeMap;
+
+	// Index of the nodes (must match up with the runtime discovery process of nodes, which runs thru the property chain)
+	int32 AllocateNodeIndexCounter;
+	TMap<class UWidgetGraphNode_Base*, int32> AllocatedAnimNodeIndices;
+
+	// List of successfully created evaluation handlers
+	TArray<FEvaluationHandlerRecord> ValidEvaluationHandlerList;
+
+	// List of animation node literals (values exposed as pins but never wired up) that need to be pushed into the CDO
+	TArray<FEffectiveConstantRecord> ValidAnimNodePinConstants;
+
+	///----------------------------------------------------------------
 };
 
diff --git a/Engine/Source/Editor/UMGEditor/Private/WidgetGraphNode_Base.cpp b/Engine/Source/Editor/UMGEditor/Private/WidgetGraphNode_Base.cpp
new file mode 100644
--- /dev/null
+++ b/Engine/Source/Editor/UMGEditor/Private/WidgetGraphNode_Base.cpp
@@ -0,0 +1,270 @@
+// Copyright 1998-2014 Epic Games, Inc. All Rights Reserved.
+
+#include "UMGEditorPrivatePCH.h"
+
+#include "BlueprintEditorUtils.h"
+#include "WidgetGraphSchema.h"
+#include "K2ActionMenuBuilder.h"
+
+#define LOCTEXT_NAMESPACE "UMG_EDITOR"
+
+/////////////////////////////////////////////////////
+// FWidgetOptionalPinManager
+
+struct FWidgetOptionalPinManager : public FOptionalPinManager
+{
+protected:
+	class UWidgetGraphNode_Base* BaseNode;
+	TArray<UEdGraphPin*>* OldPins;
+
+	TMap<FString, UEdGraphPin*> OldPinMap;
+
+public:
+	FWidgetOptionalPinManager(class UWidgetGraphNode_Base* Node, TArray<UEdGraphPin*>* InOldPins)
+		: BaseNode(Node)
+		, OldPins(InOldPins)
+	{
+		if ( OldPins != NULL )
+		{
+			for ( auto PinIt = OldPins->CreateIterator(); PinIt; ++PinIt )
+			{
+				UEdGraphPin* Pin = *PinIt;
+				OldPinMap.Add(Pin->PinName, Pin);
+			}
+		}
+	}
+
+	virtual void GetRecordDefaults(UProperty* TestProperty, FOptionalPinFromProperty& Record) const OVERRIDE
+	{
+		const UWidgetGraphSchema* Schema = GetDefault<UWidgetGraphSchema>();
+
+		// Determine if this is a pose or array of poses
+		UArrayProperty* ArrayProp = Cast<UArrayProperty>(TestProperty);
+		UStructProperty* StructProp = Cast<UStructProperty>(( ArrayProp != NULL ) ? ArrayProp->Inner : TestProperty);
+		const bool bIsPoseInput = ( StructProp != NULL ) && ( StructProp->Struct->IsChildOf(FPoseLinkBase::StaticStruct()) );
+
+		//@TODO: Error if they specified two or more of these flags
+		const bool bAlwaysShow = TestProperty->HasMetaData(Schema->NAME_AlwaysAsPin) || bIsPoseInput;
+		const bool bOptional_ShowByDefault = TestProperty->HasMetaData(Schema->NAME_PinShownByDefault);
+		const bool bOptional_HideByDefault = TestProperty->HasMetaData(Schema->NAME_PinHiddenByDefault);
+		const bool bNeverShow = TestProperty->HasMetaData(Schema->NAME_NeverAsPin);
+
+		Record.bCanToggleVisibility = bOptional_ShowByDefault || bOptional_HideByDefault;
+		Record.bShowPin = bAlwaysShow || bOptional_ShowByDefault;
+	}
+
+	virtual void CustomizePinData(UEdGraphPin* Pin, FName SourcePropertyName, int32 ArrayIndex, UProperty* Property) const OVERRIDE
+	{
+		if ( BaseNode != NULL )
+		{
+			BaseNode->CustomizePinData(Pin, SourcePropertyName, ArrayIndex);
+		}
+	}
+
+	void PostInitNewPin(UEdGraphPin* Pin, FOptionalPinFromProperty& Record, int32 ArrayIndex, UProperty* Property, uint8* PropertyAddress) const OVERRIDE
+	{
+		check(PropertyAddress != NULL);
+		check(Record.bShowPin);
+
+		if ( OldPins == NULL )
+		{
+			// Initial construction of a visible pin; copy values from the struct
+			FBlueprintEditorUtils::ExportPropertyToKismetDefaultValue(Pin, Property, PropertyAddress, BaseNode);
+		}
+		else if ( Record.bCanToggleVisibility )
+		{
+			if ( UEdGraphPin* OldPin = OldPinMap.FindRef(Pin->PinName) )
+			{
+				// Was already visible
+			}
+			else
+			{
+				// Showing a pin that was previously hidden, during a reconstruction
+				// Convert the struct property into DefaultValue/DefaultValueObject
+				FBlueprintEditorUtils::ExportPropertyToKismetDefaultValue(Pin, Property, PropertyAddress, BaseNode);
+			}
+		}
+	}
+
+	void PostRemovedOldPin(FOptionalPinFromProperty& Record, int32 ArrayIndex, UProperty* Property, uint8* PropertyAddress) const OVERRIDE
+	{
+		check(PropertyAddress != NULL);
+		check(!Record.bShowPin);
+
+		if ( Record.bCanToggleVisibility && ( OldPins != NULL ) )
+		{
+			const FString OldPinName = ( ArrayIndex != INDEX_NONE ) ? FString::Printf(TEXT("%s_%d"), *( Record.PropertyName.ToString() ), ArrayIndex) : Record.PropertyName.ToString();
+			if ( UEdGraphPin* OldPin = OldPinMap.FindRef(OldPinName) )
+			{
+				// Pin was visible but it's now hidden
+				// Convert DefaultValue/DefaultValueObject and push back into the struct
+				FBlueprintEditorUtils::ImportKismetDefaultValueToProperty(OldPin, Property, PropertyAddress, BaseNode);
+			}
+		}
+	}
+
+	void AllocateDefaultPins(UStruct* SourceStruct, uint8* StructBasePtr)
+	{
+			RebuildPropertyList(BaseNode->ShowPinForProperties, SourceStruct);
+			CreateVisiblePins(BaseNode->ShowPinForProperties, SourceStruct, EGPD_Input, BaseNode, StructBasePtr);
+	}
+};
+
+UWidgetGraphNode_Base::UWidgetGraphNode_Base(const class FPostConstructInitializeProperties& PCIP)
+	: Super(PCIP)
+{
+}
+
+void UWidgetGraphNode_Base::PostEditChangeProperty(struct FPropertyChangedEvent& PropertyChangedEvent)
+{
+	FName PropertyName = ( PropertyChangedEvent.Property != NULL ) ? PropertyChangedEvent.Property->GetFName() : NAME_None;
+
+	if ( ( PropertyName == TEXT("bShowPin") ) )
+	{
+		GetSchema()->ReconstructNode(*this);
+	}
+
+	Super::PostEditChangeProperty(PropertyChangedEvent);
+}
+
+void UWidgetGraphNode_Base::CreateOutputPins()
+{
+}
+
+void UWidgetGraphNode_Base::InternalPinCreation(TArray<UEdGraphPin*>* OldPins)
+{
+	// preload required assets first before creating pins
+	//PreloadRequiredAssets();
+
+	if ( UStructProperty* NodeStruct = GetFNodeProperty() )
+	{
+		// Display any currently visible optional pins
+		{
+			FWidgetOptionalPinManager OptionalPinManager(this, OldPins);
+			OptionalPinManager.AllocateDefaultPins(NodeStruct->Struct, NodeStruct->ContainerPtrToValuePtr<uint8>(this));
+		}
+
+		// Create the output pin, if needed
+		CreateOutputPins();
+	}
+}
+
+void UWidgetGraphNode_Base::AllocateDefaultPins()
+{
+	InternalPinCreation(NULL);
+}
+
+void UWidgetGraphNode_Base::ReallocatePinsDuringReconstruction(TArray<UEdGraphPin*>& OldPins)
+{
+	InternalPinCreation(&OldPins);
+}
+
+bool UWidgetGraphNode_Base::CanCreateUnderSpecifiedSchema(const UEdGraphSchema* DesiredSchema) const
+{
+	return true;
+	//return DesiredSchema->GetClass()->IsChildOf(UWidgetGraphSchema::StaticClass());
+}
+
+void UWidgetGraphNode_Base::GetMenuEntries(FGraphContextMenuBuilder& ContextMenuBuilder) const
+{
+	CreateDefaultMenuEntry(ContextMenuBuilder);
+}
+
+TSharedPtr<FEdGraphSchemaAction_K2NewNode> UWidgetGraphNode_Base::CreateDefaultMenuEntry(FGraphContextMenuBuilder& ContextMenuBuilder) const
+{
+	UWidgetGraphNode_Base* TemplateNode = NewObject<UWidgetGraphNode_Base>(GetTransientPackage(), GetClass());
+
+	FString Category = TemplateNode->GetNodeCategory();
+	FText MenuDesc = TemplateNode->GetNodeTitle(ENodeTitleType::ListView);
+	FString Tooltip = TemplateNode->GetTooltip();
+	FString Keywords = TemplateNode->GetKeywords();
+
+	TSharedPtr<FEdGraphSchemaAction_K2NewNode> NodeAction = FK2ActionMenuBuilder::AddNewNodeAction(ContextMenuBuilder, Category, MenuDesc, Tooltip, 0, Keywords);
+	NodeAction->NodeTemplate = TemplateNode;
+	NodeAction->SearchTitle = TemplateNode->GetNodeSearchTitle();
+
+	return NodeAction;
+}
+
+void UWidgetGraphNode_Base::GetPinAssociatedProperty(UScriptStruct* NodeType, UEdGraphPin* InputPin, UProperty*& OutProperty, int32& OutIndex)
+{
+	OutProperty = NULL;
+	OutIndex = INDEX_NONE;
+
+	//@TODO: Name-based hackery, avoid the roundtrip and better indicate when it's an array pose pin
+	int32 UnderscoreIndex = InputPin->PinName.Find(TEXT("_"));
+	if (UnderscoreIndex != INDEX_NONE)
+	{
+		FString ArrayName = InputPin->PinName.Left(UnderscoreIndex);
+		int32 ArrayIndex = FCString::Atoi(*(InputPin->PinName.Mid(UnderscoreIndex + 1)));
+
+		if (UArrayProperty* ArrayProperty = FindField<UArrayProperty>(NodeType, *ArrayName))
+		{
+			OutProperty = ArrayProperty;
+			OutIndex = ArrayIndex;
+		}
+	}
+	else
+	{
+		if (UProperty* Property = FindField<UProperty>(NodeType, *(InputPin->PinName)))
+		{
+			OutProperty = Property;
+			OutIndex = INDEX_NONE;
+		}
+	}
+}
+
+// Gets the menu category this node belongs in
+FString UWidgetGraphNode_Base::GetNodeCategory() const
+{
+	return LOCTEXT("Widget", "Widget").ToString();
+}
+
+// customize pin data based on the input
+void UWidgetGraphNode_Base::CustomizePinData(UEdGraphPin* Pin, FName SourcePropertyName, int32 ArrayIndex) const
+{
+
+}
+
+FLinearColor UWidgetGraphNode_Base::GetNodeTitleColor() const
+{
+	return FLinearColor::Black;
+}
+
+UScriptStruct* UWidgetGraphNode_Base::GetFNodeType() const
+{
+	UScriptStruct* BaseFStruct = FWidgetNode_Base::StaticStruct();
+
+	for ( TFieldIterator<UProperty> PropIt(GetClass(), EFieldIteratorFlags::IncludeSuper); PropIt; ++PropIt )
+	{
+		if ( UStructProperty* StructProp = Cast<UStructProperty>(*PropIt) )
+		{
+			if ( StructProp->Struct->IsChildOf(BaseFStruct) )
+			{
+				return StructProp->Struct;
+			}
+		}
+	}
+
+	return NULL;
+}
+
+UStructProperty* UWidgetGraphNode_Base::GetFNodeProperty() const
+{
+	UScriptStruct* BaseFStruct = FWidgetNode_Base::StaticStruct();
+
+	for ( TFieldIterator<UProperty> PropIt(GetClass(), EFieldIteratorFlags::IncludeSuper); PropIt; ++PropIt )
+	{
+		if ( UStructProperty* StructProp = Cast<UStructProperty>(*PropIt) )
+		{
+			if ( StructProp->Struct->IsChildOf(BaseFStruct) )
+			{
+				return StructProp;
+			}
+		}
+	}
+
+	return NULL;
+}
+
+#undef LOCTEXT_NAMESPACE
\ No newline at end of file
diff --git a/Engine/Source/Editor/UMGEditor/Private/WidgetGraphNode_Text.cpp b/Engine/Source/Editor/UMGEditor/Private/WidgetGraphNode_Text.cpp
new file mode 100644
--- /dev/null
+++ b/Engine/Source/Editor/UMGEditor/Private/WidgetGraphNode_Text.cpp
@@ -0,0 +1,22 @@
+// Copyright 1998-2014 Epic Games, Inc. All Rights Reserved.
+
+#include "UMGEditorPrivatePCH.h"
+
+#define LOCTEXT_NAMESPACE "UMG_EDITOR"
+
+UWidgetGraphNode_Text::UWidgetGraphNode_Text(const class FPostConstructInitializeProperties& PCIP)
+	: Super(PCIP)
+{
+}
+
+FString UWidgetGraphNode_Text::GetTooltip() const
+{
+	return LOCTEXT("TextBlockBindingTooltip", "Text Block (Binding)").ToString();
+}
+
+FText UWidgetGraphNode_Text::GetNodeTitle(ENodeTitleType::Type TitleType) const
+{
+	return LOCTEXT("TextBlockBinding", "Text Block (Binding)");
+}
+
+#undef LOCTEXT_NAMESPACE
\ No newline at end of file
diff --git a/Engine/Source/Editor/UMGEditor/Private/WidgetGraphSchema.cpp b/Engine/Source/Editor/UMGEditor/Private/WidgetGraphSchema.cpp
new file mode 100644
--- /dev/null
+++ b/Engine/Source/Editor/UMGEditor/Private/WidgetGraphSchema.cpp
@@ -0,0 +1,14 @@
+// Copyright 1998-2014 Epic Games, Inc. All Rights Reserved.
+
+#include "UMGEditorPrivatePCH.h"
+
+UWidgetGraphSchema::UWidgetGraphSchema(const FPostConstructInitializeProperties& PCIP)
+: Super(PCIP)
+{
+	NAME_NeverAsPin = TEXT("NeverAsPin");
+	NAME_PinHiddenByDefault = TEXT("PinHiddenByDefault");
+	NAME_PinShownByDefault = TEXT("PinShownByDefault");
+	NAME_AlwaysAsPin = TEXT("AlwaysAsPin");
+	NAME_OnEvaluate = TEXT("OnEvaluate");
+	DefaultEvaluationHandlerName = TEXT("EvaluateGraphExposedInputs");
+}
\ No newline at end of file
diff --git a/Engine/Source/Runtime/UMG/Private/Components/SlateWrapperComponent.cpp b/Engine/Source/Runtime/UMG/Private/Components/SlateWrapperComponent.cpp
--- a/Engine/Source/Runtime/UMG/Private/Components/SlateWrapperComponent.cpp
+++ b/Engine/Source/Runtime/UMG/Private/Components/SlateWrapperComponent.cpp
@@ -64,7 +64,8 @@ TSharedRef<SWidget> USlateWrapperComponent::GetWidget()
 	{
 		MyWidget = RebuildWidget();
 		MyWidget->SetEnabled(bIsEnabled);
-		
+		MyWidget->SetVisibility(ConvertSerializedVisibilityToRuntime(Visiblity));
+
 		if ( !ToolTipText.IsEmpty() )
 		{
 			MyWidget->SetToolTipText(ToolTipText);
diff --git a/Engine/Source/Runtime/UMG/Private/Components/TextBlockComponent.cpp b/Engine/Source/Runtime/UMG/Private/Components/TextBlockComponent.cpp
--- a/Engine/Source/Runtime/UMG/Private/Components/TextBlockComponent.cpp
+++ b/Engine/Source/Runtime/UMG/Private/Components/TextBlockComponent.cpp
@@ -34,8 +34,10 @@ TSharedRef<SWidget> UTextBlockComponent::RebuildWidget()
 		StylePtr = Defaults._TextStyle;
 	}
 
+//	TAttribute<FText>::Create(TAttribute<FText>::FGetter::From::CreateUObject(this, &GlueText));
+
 	return SNew(STextBlock)
-		.Text(BIND_UOBJECT_ATTRIBUTE(FString, GetText))
+		.Text(BIND_UOBJECT_ATTRIBUTE(FText, GetText))
 		.TextStyle(StylePtr)
 		.Font(FSlateFontInfo(FontPath, Font.Size))
 		.ColorAndOpacity( BIND_UOBJECT_ATTRIBUTE(FSlateColor, GetColorAndOpacity) )
@@ -43,9 +45,14 @@ TSharedRef<SWidget> UTextBlockComponent::RebuildWidget()
 		.ShadowColorAndOpacity(BIND_UOBJECT_ATTRIBUTE(FLinearColor, GetShadowColorAndOpacity));
 }
 
-FString UTextBlockComponent::GetText() const
+FText UTextBlockComponent::GetText() const
 {
-	return Text.ToString();
+	if ( TextDelegate.IsBound() )
+	{
+		return TextDelegate.Execute();
+	}
+
+	return Text;
 }
 
 FSlateColor UTextBlockComponent::GetColorAndOpacity() const
diff --git a/Engine/Source/Runtime/UMG/Private/UserWidget.cpp b/Engine/Source/Runtime/UMG/Private/UserWidget.cpp
--- a/Engine/Source/Runtime/UMG/Private/UserWidget.cpp
+++ b/Engine/Source/Runtime/UMG/Private/UserWidget.cpp
@@ -140,16 +140,40 @@ void AUserWidget::RerunConstructionScripts()
 
 void AUserWidget::TickActor(float DeltaTime, enum ELevelTick TickType, FActorTickFunction& ThisTickFunction)
 {
-	//UWorld* World = GetWorld();
+	//if ( BindingEvaluator )
+	//{
+	//	// TODO UMG pass in this as a parameter, so it has the context to set properties?
+	//	ProcessEvent(BindingEvaluator, NULL);
+	//}
+	//static AUserWidget* TestWidget;
+	//static FWidgetNode_Base* LinkedPtr = NULL;
+
+	//FText GlueText()
+	//{
+	//	LinkedPtr->EvaluateGraphExposedInputs.Execute(TestWidget);
+	//	FWidgetNode_Text* TextNode = (FWidgetNode_Text*)LinkedPtr;
+	//	return TextNode->Text;
+	//}
+
+	//TextBlock->LinkedTextAttribute = TAttribute<FText>::Create(TAttribute<FText>::FGetter::CreateStatic(&GlueText));
 
+	//UWorld* World = GetWorld();
 	//if ( World && World->IsGameWorld() )
 	//{
-	//	if ( bShowCursorWhenVisible && GetIsVisible() )
+	//	if ( UWidgetBlueprintGeneratedClass* WidgetBlueprintClass = Cast<UWidgetBlueprintGeneratedClass>(GetClass()) )
 	//	{
-	//		World->GetFirstLocalPlayerFromController()->PlayerController->bShowMouseCursor = true;
+	//		for ( UStructProperty* Attribute : WidgetBlueprintClass->WidgetNodeProperties )
+	//		{
 
-	//		UGameViewportClient* Viewport = World->GetGameViewport();
+	//		}
 	//	}
+
+
+	//	//if ( bShowCursorWhenVisible && GetIsVisible() )
+	//	//{
+	//	//	World->GetFirstLocalPlayerFromController()->PlayerController->bShowMouseCursor = true;
+	//	//	UGameViewportClient* Viewport = World->GetGameViewport();
+	//	//}
 	//}
 }
 
diff --git a/Engine/Source/Runtime/UMG/Private/WidgetBlueprint.cpp b/Engine/Source/Runtime/UMG/Private/WidgetBlueprint.cpp
--- a/Engine/Source/Runtime/UMG/Private/WidgetBlueprint.cpp
+++ b/Engine/Source/Runtime/UMG/Private/WidgetBlueprint.cpp
@@ -4,6 +4,16 @@
 
 #define LOCTEXT_NAMESPACE "UMG"
 
+FDelegateRuntimeBinding FDelegateEditorBinding::ToRuntimeBinding(UWidgetBlueprint* Blueprint) const
+{
+	FDelegateRuntimeBinding Binding;
+	Binding.ObjectName = ObjectName;
+	Binding.PropertyName = PropertyName;
+	Binding.FunctionName = Blueprint->GetFieldNameFromClassByGuid<UFunction>(Blueprint->SkeletonGeneratedClass, MemberGuid);
+
+	return Binding;
+}
+
 /////////////////////////////////////////////////////
 // UWidgetBlueprint
 
diff --git a/Engine/Source/Runtime/UMG/Private/WidgetBlueprintGeneratedClass.cpp b/Engine/Source/Runtime/UMG/Private/WidgetBlueprintGeneratedClass.cpp
--- a/Engine/Source/Runtime/UMG/Private/WidgetBlueprintGeneratedClass.cpp
+++ b/Engine/Source/Runtime/UMG/Private/WidgetBlueprintGeneratedClass.cpp
@@ -13,6 +13,26 @@ UWidgetBlueprintGeneratedClass::UWidgetBlueprintGeneratedClass(const FPostConstr
 	WidgetTree = ConstructObject<UWidgetTree>(UWidgetTree::StaticClass(), this);
 }
 
+void UWidgetBlueprintGeneratedClass::Link(FArchive& Ar, bool bRelinkExistingProperties)
+{
+	Super::Link(Ar, bRelinkExistingProperties);
+
+	// @TODO: Shouldn't be necessary to clear these, but currently the class gets linked twice during compilation
+	WidgetNodeProperties.Empty();
+
+	// Initialize derived members
+	for ( TFieldIterator<UProperty> It(this); It; ++It )
+	{
+		if ( UStructProperty* StructProp = Cast<UStructProperty>(*It) )
+		{
+			if ( StructProp->Struct->IsChildOf(FWidgetNode_Base::StaticStruct()) )
+			{
+				WidgetNodeProperties.Add(StructProp);
+			}
+		}
+	}
+}
+
 void UWidgetBlueprintGeneratedClass::CreateComponentsForActor(AActor* Actor) const
 {
 	Super::CreateComponentsForActor(Actor);
@@ -51,6 +71,31 @@ void UWidgetBlueprintGeneratedClass::CreateComponentsForActor(AActor* Actor) con
 			Prop->SetObjectPropertyValue_InContainer(Actor, Widget);
 		}
 
+		// Perform binding
+		for ( const FDelegateRuntimeBinding& Binding : Bindings )
+		{
+			//TODO UMG Terrible performance, improve with Maps.
+			if ( Binding.ObjectName == VariableName )
+			{
+				UFunction* BoundFunction = WidgetActor->FindFunction(Binding.FunctionName);
+
+				FString DelegateName = Binding.PropertyName.ToString() + "Delegate";
+
+				for ( TFieldIterator<UProperty> It(Widget->GetClass()); It; ++It )
+				{
+					if ( UDelegateProperty* DelegateProp = Cast<UDelegateProperty>(*It) )
+					{
+						if ( DelegateProp->GetName() == DelegateName )
+						{
+							FScriptDelegate* ScriptDelegate = DelegateProp->GetPropertyValuePtr_InContainer(Widget);
+							ScriptDelegate->BindUFunction(WidgetActor, Binding.FunctionName);
+							break;
+						}
+					}
+				}
+			}
+		}
+
 		Widget->RegisterComponent();
 
 #if WITH_EDITOR
diff --git a/Engine/Source/Runtime/UMG/Public/Blueprint/UserWidget.h b/Engine/Source/Runtime/UMG/Public/Blueprint/UserWidget.h
--- a/Engine/Source/Runtime/UMG/Public/Blueprint/UserWidget.h
+++ b/Engine/Source/Runtime/UMG/Public/Blueprint/UserWidget.h
@@ -6,7 +6,7 @@
 
 DECLARE_DYNAMIC_MULTICAST_DELEGATE_OneParam(FOnVisibilityChangedEvent, ESlateVisibility::Type, Visibility);
 
-UCLASS(Abstract, hideCategories=(Object, Actor))
+UCLASS(Abstract, hideCategories=(Object, Actor, Replication, Rendering, Input))
 class UMG_API AUserWidget : public AActor
 {
 	GENERATED_UCLASS_BODY()
@@ -103,7 +103,6 @@ class UMG_API AUserWidget : public AActor
 	TSharedPtr<SWidget> GetWidgetFromName(const FString& Name) const;
 	USlateWrapperComponent* GetHandleFromName(const FString& Name) const;
 
-
 private:
 	TSharedPtr<SWidget> RootWidget;
 	TMap< TWeakPtr<SWidget>, USlateWrapperComponent* > WidgetToComponent;
diff --git a/Engine/Source/Runtime/UMG/Public/Blueprint/WidgetBlueprint.h b/Engine/Source/Runtime/UMG/Public/Blueprint/WidgetBlueprint.h
--- a/Engine/Source/Runtime/UMG/Public/Blueprint/WidgetBlueprint.h
+++ b/Engine/Source/Runtime/UMG/Public/Blueprint/WidgetBlueprint.h
@@ -2,9 +2,38 @@
 
 #pragma once
 
+#include "WidgetBlueprintGeneratedClass.h"
+
 #include "WidgetBlueprint.generated.h"
 
-UCLASS(dependson=(UBlueprint, AUserWidget), BlueprintType)
+USTRUCT()
+struct UMG_API FDelegateEditorBinding
+{
+	GENERATED_USTRUCT_BODY()
+
+	UPROPERTY()
+	FString ObjectName;
+
+	UPROPERTY()
+	FName PropertyName;
+
+	UPROPERTY()
+	FName FunctionName;
+
+	UPROPERTY()
+	FGuid MemberGuid;
+
+	bool operator==( const FDelegateEditorBinding& Other ) const
+	{
+		// NOTE: We intentionally only compare object name and property name, the function is irrelevant since
+		// you're only allowed to bind a property on an object to a single function.
+		return ObjectName == Other.ObjectName && PropertyName == Other.PropertyName;
+	}
+
+	FDelegateRuntimeBinding ToRuntimeBinding(class UWidgetBlueprint* Blueprint) const;
+};
+
+UCLASS(dependson=(UBlueprint, AUserWidget, UWidgetBlueprintGeneratedClass), BlueprintType)
 class UMG_API UWidgetBlueprint : public UBlueprint
 {
 	GENERATED_UCLASS_BODY()
@@ -14,6 +43,9 @@ public:
 	UPROPERTY()
 	class UWidgetTree* WidgetTree;
 
+	UPROPERTY()
+	TArray< FDelegateEditorBinding > Bindings;
+
 	virtual void PostLoad() OVERRIDE;
 	virtual void PostLoadSubobjects(FObjectInstancingGraph* OuterInstanceGraph) OVERRIDE;
 	
diff --git a/Engine/Source/Runtime/UMG/Public/Blueprint/WidgetBlueprintGeneratedClass.h b/Engine/Source/Runtime/UMG/Public/Blueprint/WidgetBlueprintGeneratedClass.h
--- a/Engine/Source/Runtime/UMG/Public/Blueprint/WidgetBlueprintGeneratedClass.h
+++ b/Engine/Source/Runtime/UMG/Public/Blueprint/WidgetBlueprintGeneratedClass.h
@@ -4,7 +4,29 @@
 
 #include "WidgetBlueprintGeneratedClass.generated.h"
 
-UCLASS(dependson=(UWidgetBlueprint))
+USTRUCT()
+struct FDelegateRuntimeBinding
+{
+	GENERATED_USTRUCT_BODY()
+
+	UPROPERTY()
+	FString ObjectName;
+
+	UPROPERTY()
+	FName PropertyName;
+
+	UPROPERTY()
+	FName FunctionName;
+	
+	/*UFunction* UK2Node_CallFunction::GetTargetFunction() const
+	{
+	UFunction* Function = FunctionReference.ResolveMember<UFunction>(this);
+	return Function;
+	}*/
+};
+
+
+UCLASS()
 class UMG_API UWidgetBlueprintGeneratedClass : public UBlueprintGeneratedClass
 {
 	GENERATED_UCLASS_BODY()
@@ -13,5 +35,13 @@ class UMG_API UWidgetBlueprintGeneratedClass : public UBlueprintGeneratedClass
 	UPROPERTY()
 	class UWidgetTree* WidgetTree;
 
+	UPROPERTY()
+	TArray< FDelegateRuntimeBinding > Bindings;
+
+	/** This is transient data calculated at link time. */
+	TArray<UStructProperty*> WidgetNodeProperties;
+
 	virtual void CreateComponentsForActor(AActor* Actor) const OVERRIDE;
+
+	virtual void Link(FArchive& Ar, bool bRelinkExistingProperties) OVERRIDE;
 };
\ No newline at end of file
diff --git a/Engine/Source/Runtime/UMG/Public/Blueprint/WidgetNode_Base.h b/Engine/Source/Runtime/UMG/Public/Blueprint/WidgetNode_Base.h
new file mode 100644
--- /dev/null
+++ b/Engine/Source/Runtime/UMG/Public/Blueprint/WidgetNode_Base.h
@@ -0,0 +1,48 @@
+// Copyright 1998-2014 Epic Games, Inc. All Rights Reserved.
+
+#pragma once
+
+#include "WidgetNode_Base.generated.h"
+
+// An exposed value updater
+USTRUCT()
+struct FExposedWidgetValueHandler
+{
+	GENERATED_USTRUCT_BODY()
+
+		// The function to call to update associated properties (can be NULL)
+		UPROPERTY()
+		FName BoundFunction;
+
+	void Execute(AUserWidget* Widget) const
+	{
+		if ( BoundFunction != NAME_None )
+		{
+			//@TODO: Should be able to be Checked, or at least produce a warning when it fails
+			if ( UFunction* Function = Widget->FindFunction(BoundFunction) )
+			{
+				Widget->ProcessEvent(Function, NULL);
+			}
+		}
+	}
+
+	UFunction* GetFunction(AUserWidget* Widget)
+	{
+		if ( BoundFunction != NAME_None )
+		{
+			return Widget->FindFunction(BoundFunction);
+		}
+
+		return NULL;
+	}
+};
+
+USTRUCT()
+struct UMG_API FWidgetNode_Base
+{
+	GENERATED_USTRUCT_BODY()
+
+	// The default handler for graph-exposed inputs
+	UPROPERTY()
+	FExposedWidgetValueHandler EvaluateGraphExposedInputs;
+};
diff --git a/Engine/Source/Runtime/UMG/Public/Blueprint/WidgetNode_Text.h b/Engine/Source/Runtime/UMG/Public/Blueprint/WidgetNode_Text.h
new file mode 100644
--- /dev/null
+++ b/Engine/Source/Runtime/UMG/Public/Blueprint/WidgetNode_Text.h
@@ -0,0 +1,22 @@
+// Copyright 1998-2014 Epic Games, Inc. All Rights Reserved.
+
+#pragma once
+
+#include "WidgetNode_Text.generated.h"
+
+USTRUCT()
+struct UMG_API FWidgetNode_Text : public FWidgetNode_Base
+{
+	GENERATED_USTRUCT_BODY()
+public:
+	// Which input should be connected to the output?
+	UPROPERTY(EditAnywhere, BlueprintReadWrite, Category=Runtime, meta=(AlwaysAsPin))
+	FText Text;
+
+	UPROPERTY(EditAnywhere, BlueprintReadWrite, Category=Runtime, meta=( AlwaysAsPin ))
+	FLinearColor ForegroundColor;
+
+	
+
+	//etc..
+};
\ No newline at end of file
diff --git a/Engine/Source/Runtime/UMG/Public/Components/TextBlockComponent.h b/Engine/Source/Runtime/UMG/Public/Components/TextBlockComponent.h
--- a/Engine/Source/Runtime/UMG/Public/Components/TextBlockComponent.h
+++ b/Engine/Source/Runtime/UMG/Public/Components/TextBlockComponent.h
@@ -9,10 +9,15 @@ class UMG_API UTextBlockComponent : public USlateLeafWidgetComponent
 {
 	GENERATED_UCLASS_BODY()
 
+	DECLARE_DYNAMIC_DELEGATE_RetVal(FText, FTextDelegate);
+
 public:
-	UPROPERTY(EditAnywhere, Category=Content)
+	UPROPERTY(EditDefaultsOnly, Category=Content)
 	FText Text;
 
+	UPROPERTY()
+	FTextDelegate TextDelegate;
+	
 	UPROPERTY(EditAnywhere, Category=Appearance)
 	FSlateFontInfo Font;
 
@@ -34,7 +39,7 @@ protected:
 	virtual TSharedRef<SWidget> RebuildWidget() OVERRIDE;
 	// End of USlateWrapperComponent interface
 
-	FString GetText() const;
+	FText GetText() const;
 	FSlateColor GetColorAndOpacity() const;
 	FLinearColor GetShadowColorAndOpacity() const;
 };
diff --git a/Engine/Source/Runtime/UMG/Public/UMG.h b/Engine/Source/Runtime/UMG/Public/UMG.h
--- a/Engine/Source/Runtime/UMG/Public/UMG.h
+++ b/Engine/Source/Runtime/UMG/Public/UMG.h
@@ -65,6 +65,10 @@
 
 #include "WidgetTree.h"
 #include "UserWidget.h"
+
+#include "WidgetNode_Base.h"
+#include "WidgetNode_Text.h"
+
 #include "WidgetBlueprint.h"
 #include "WidgetBlueprintGeneratedClass.h"
  
