Fixing non-unity build

diff --git a/Engine/Source/Editor/UMGEditor/Private/DetailCustomizations.cpp b/Engine/Source/Editor/UMGEditor/Private/DetailCustomizations.cpp
--- a/Engine/Source/Editor/UMGEditor/Private/DetailCustomizations.cpp
+++ b/Engine/Source/Editor/UMGEditor/Private/DetailCustomizations.cpp
@@ -8,6 +8,7 @@
 #include "ObjectEditorUtils.h"
 #include "WidgetGraphSchema.h"
 #include "ScopedTransaction.h"
+#include "BlueprintEditorUtils.h"
 
 #define LOCTEXT_NAMESPACE "UMG"
 
