UMG - Now getting the correct property outter class to be used when getting property handles for the binding customization.

diff --git a/Engine/Source/Editor/UMGEditor/Private/DetailCustomizations.cpp b/Engine/Source/Editor/UMGEditor/Private/DetailCustomizations.cpp
--- a/Engine/Source/Editor/UMGEditor/Private/DetailCustomizations.cpp
+++ b/Engine/Source/Editor/UMGEditor/Private/DetailCustomizations.cpp
@@ -34,7 +34,7 @@ void FBlueprintWidgetCustomization::CustomizeDetails( IDetailLayoutBuilder& Deta
 				ConstPropertyName.RemoveFromEnd(TEXT("Delegate"));
 
 				//TSharedRef<IPropertyHandle> DelegatePropertyHandle = DetailLayout.GetProperty(Property->GetFName(), PropertyClass);
-				TSharedRef<IPropertyHandle> ConstPropertyHandle = DetailLayout.GetProperty(FName(*ConstPropertyName), PropertyClass);
+				TSharedRef<IPropertyHandle> ConstPropertyHandle = DetailLayout.GetProperty(FName(*ConstPropertyName), CastChecked<UClass>(Property->GetOuter()));
 
 				const bool bHasValidHandle = ConstPropertyHandle->IsValidHandle();
 
