UMG - Prototype support for single handler event delegates.  Fixing a bug with generated binding functions not having pins created as user pins.

diff --git a/Engine/Source/Editor/BlueprintGraph/Classes/K2Node_EditablePinBase.h b/Engine/Source/Editor/BlueprintGraph/Classes/K2Node_EditablePinBase.h
--- a/Engine/Source/Editor/BlueprintGraph/Classes/K2Node_EditablePinBase.h
+++ b/Engine/Source/Editor/BlueprintGraph/Classes/K2Node_EditablePinBase.h
@@ -128,6 +128,11 @@ class UK2Node_EditablePinBase : public UK2Node
 	// Modifies the default value of an existing pin on the node.
 	BLUEPRINTGRAPH_API virtual bool ModifyUserDefinedPinDefaultValue(TSharedPtr<FUserPinInfo> PinInfo, const FString& NewDefaultValue);
 
+	/**
+	 * Creates function pins that are user defined based on a function signature.
+	 */
+	BLUEPRINTGRAPH_API virtual bool CreateUserDefinedPinsForFunctionEntryExit(const UFunction* Function, bool bForFunctionEntry);
+
 	/**
 	 * Can this node have execution wires added or removed?
 	 */
diff --git a/Engine/Source/Editor/BlueprintGraph/Private/EdGraphSchema_K2.cpp b/Engine/Source/Editor/BlueprintGraph/Private/EdGraphSchema_K2.cpp
--- a/Engine/Source/Editor/BlueprintGraph/Private/EdGraphSchema_K2.cpp
+++ b/Engine/Source/Editor/BlueprintGraph/Private/EdGraphSchema_K2.cpp
@@ -2190,7 +2190,10 @@ void UEdGraphSchema_K2::CreateFunctionGraphTerminators(UEdGraph& Graph, UFunctio
 	EntryNode->SignatureName = GraphName;
 	FunctionEntryCreator.Finalize();
 
-	EntryNode->CreatePinsForFunctionEntryExit(FunctionSignature, /*bIsFunctionEntry=*/ true);
+	// We don't have a signature class to base this on permanently, because it's not an override function.
+	// so we need to define the pins as user defined so that they are serialized.
+
+	EntryNode->CreateUserDefinedPinsForFunctionEntryExit(FunctionSignature, /*bIsFunctionEntry=*/ true);
 
 	// See if any function params are marked as out
 	bool bHasOutParam =  false;
@@ -2213,7 +2216,7 @@ void UEdGraphSchema_K2::CreateFunctionGraphTerminators(UEdGraph& Graph, UFunctio
 		ReturnNode->NodePosY = EntryNode->NodePosY;
 		NodeCreator.Finalize();
 
-		ReturnNode->CreatePinsForFunctionEntryExit(FunctionSignature, /*bIsFunctionEntry=*/ false);
+		ReturnNode->CreateUserDefinedPinsForFunctionEntryExit(FunctionSignature, /*bIsFunctionEntry=*/ false);
 
 		// Auto-connect the pins for entry and exit, so that by default the signature is properly generated
 		UEdGraphPin* EntryNodeExec = FindExecutionPin(*EntryNode, EGPD_Output);
diff --git a/Engine/Source/Editor/BlueprintGraph/Private/K2Node_EditablePinBase.cpp b/Engine/Source/Editor/BlueprintGraph/Private/K2Node_EditablePinBase.cpp
--- a/Engine/Source/Editor/BlueprintGraph/Private/K2Node_EditablePinBase.cpp
+++ b/Engine/Source/Editor/BlueprintGraph/Private/K2Node_EditablePinBase.cpp
@@ -197,3 +197,31 @@ bool UK2Node_EditablePinBase::ModifyUserDefinedPinDefaultValue(TSharedPtr<FUserP
 	PinInfo->PinDefaultValue = NewDefaultValue;
 	return true;
 }
+
+bool UK2Node_EditablePinBase::CreateUserDefinedPinsForFunctionEntryExit(const UFunction* Function, bool bForFunctionEntry)
+{
+	const UEdGraphSchema_K2* K2Schema = GetDefault<UEdGraphSchema_K2>();
+
+	// Create the inputs and outputs
+	bool bAllPinsGood = true;
+	for ( TFieldIterator<UProperty> PropIt(Function); PropIt && ( PropIt->PropertyFlags & CPF_Parm ); ++PropIt )
+	{
+		UProperty* Param = *PropIt;
+
+		const bool bIsFunctionInput = !Param->HasAnyPropertyFlags(CPF_OutParm) || Param->HasAnyPropertyFlags(CPF_ReferenceParm);
+
+		if ( bIsFunctionInput == bForFunctionEntry )
+		{
+			const EEdGraphPinDirection Direction = bForFunctionEntry ? EGPD_Output : EGPD_Input;
+
+			FEdGraphPinType PinType;
+			K2Schema->ConvertPropertyToPinType(Param, /*out*/ PinType);
+
+			const bool bPinGood = CreateUserDefinedPin(Param->GetName(), PinType) != NULL;
+
+			bAllPinsGood = bAllPinsGood && bPinGood;
+		}
+	}
+
+	return bAllPinsGood;
+}
\ No newline at end of file
diff --git a/Engine/Source/Editor/UMGEditor/Private/DetailCustomizations.cpp b/Engine/Source/Editor/UMGEditor/Private/DetailCustomizations.cpp
--- a/Engine/Source/Editor/UMGEditor/Private/DetailCustomizations.cpp
+++ b/Engine/Source/Editor/UMGEditor/Private/DetailCustomizations.cpp
@@ -101,6 +101,69 @@ void FBlueprintWidgetCustomization::CustomizeDetails( IDetailLayoutBuilder& Deta
 				];
 			}
 		}
+
+		//--------------------------------------------
+
+		for ( TFieldIterator<UProperty> PropertyIt(PropertyClass, EFieldIteratorFlags::IncludeSuper); PropertyIt; ++PropertyIt )
+		{
+			if ( UDelegateProperty* Property = Cast<UDelegateProperty>(*PropertyIt) )
+			{
+				if ( !Property->GetName().EndsWith(TEXT("Event")) )
+				{
+					continue;
+				}
+
+				TSharedRef<IPropertyHandle> DelegatePropertyHandle = DetailLayout.GetProperty(Property->GetFName(), CastChecked<UClass>(Property->GetOuter()));
+
+				const bool bHasValidHandle = DelegatePropertyHandle->IsValidHandle();
+
+				if ( !bHasValidHandle )
+				{
+					continue;
+				}
+
+				IDetailCategoryBuilder& PropertyCategory = DetailLayout.EditCategory(FObjectEditorUtils::GetCategoryFName(Property));
+
+				IDetailPropertyRow& PropertyRow = PropertyCategory.AddProperty(DelegatePropertyHandle);
+
+				TSharedPtr<SWidget> NameWidget;
+				TSharedPtr<SWidget> ValueWidget;
+				FDetailWidgetRow Row;
+				PropertyRow.GetDefaultWidgets(NameWidget, ValueWidget, Row);
+
+				TSharedRef<SComboButton> DelegateButton = SNew(SComboButton)
+				.OnGetMenuContent(this, &FBlueprintWidgetCustomization::OnGenerateDelegateMenu, DelegatePropertyHandle, Property->SignatureFunction)
+				.ContentPadding(0)
+				.ButtonContent()
+				[
+					SNew(STextBlock)
+					.Text(this, &FBlueprintWidgetCustomization::GetCurrentBindingText, DelegatePropertyHandle)
+					.Font(IDetailLayoutBuilder::GetDetailFont())
+				];
+
+				const bool bShowChildren = true;
+				PropertyRow.CustomWidget(bShowChildren)
+				.NameContent()
+				.MinDesiredWidth(Row.NameWidget.MinWidth)
+				.MaxDesiredWidth(Row.NameWidget.MaxWidth)
+				[
+					NameWidget.ToSharedRef()
+				]
+				.ValueContent()
+				.MinDesiredWidth(Row.ValueWidget.MinWidth)
+				.MaxDesiredWidth(Row.ValueWidget.MaxWidth)
+				[
+					SNew(SHorizontalBox)
+					
+					+ SHorizontalBox::Slot()
+					.FillWidth(1)
+					[
+						DelegateButton
+					]
+				];
+			}
+		}
+
 	}
 }
 
@@ -207,7 +270,7 @@ TSharedRef<SWidget> FBlueprintWidgetCustomization::OnGenerateDelegateMenu(TShare
 		.MaxHeight(DisplayMetrics.PrimaryDisplayHeight * 0.5)
 		[
 			MenuBuilder.MakeWidget()
-		];;
+		];
 }
 
 FText FBlueprintWidgetCustomization::GetCurrentBindingText(TSharedRef<IPropertyHandle> PropertyHandle) const
diff --git a/Engine/Source/Runtime/UMG/Private/Components/ButtonComponent.cpp b/Engine/Source/Runtime/UMG/Private/Components/ButtonComponent.cpp
--- a/Engine/Source/Runtime/UMG/Private/Components/ButtonComponent.cpp
+++ b/Engine/Source/Runtime/UMG/Private/Components/ButtonComponent.cpp
@@ -107,8 +107,12 @@ void UButtonComponent::SetContentPadding(FMargin InContentPadding)
 
 FReply UButtonComponent::HandleOnClicked()
 {
-	OnClicked.Broadcast();
-	return FReply::Handled();
+	if ( OnClickedEvent.IsBound() )
+	{
+		return OnClickedEvent.Execute().ToReply();
+	}
+	
+	return FReply::Unhandled();
 }
 
 
diff --git a/Engine/Source/Runtime/UMG/Private/WidgetBlueprintGeneratedClass.cpp b/Engine/Source/Runtime/UMG/Private/WidgetBlueprintGeneratedClass.cpp
--- a/Engine/Source/Runtime/UMG/Private/WidgetBlueprintGeneratedClass.cpp
+++ b/Engine/Source/Runtime/UMG/Private/WidgetBlueprintGeneratedClass.cpp
@@ -85,7 +85,7 @@ void UWidgetBlueprintGeneratedClass::CreateComponentsForActor(AActor* Actor) con
 				{
 					if ( UDelegateProperty* DelegateProp = Cast<UDelegateProperty>(*It) )
 					{
-						if ( DelegateProp->GetName() == DelegateName )
+						if ( DelegateProp->GetName() == DelegateName || DelegateProp->GetFName() == Binding.PropertyName )
 						{
 							FScriptDelegate* ScriptDelegate = DelegateProp->GetPropertyValuePtr_InContainer(Widget);
 							ScriptDelegate->BindUFunction(WidgetActor, Binding.FunctionName);
diff --git a/Engine/Source/Runtime/UMG/Public/Components/ButtonComponent.h b/Engine/Source/Runtime/UMG/Public/Components/ButtonComponent.h
--- a/Engine/Source/Runtime/UMG/Public/Components/ButtonComponent.h
+++ b/Engine/Source/Runtime/UMG/Public/Components/ButtonComponent.h
@@ -50,11 +50,11 @@ public:
 
 	UPROPERTY(EditDefaultsOnly, Category=Sound)
 	FSlateSound HoveredSound;
-	
-	/** Called when the button is clicked */
-	UPROPERTY(BlueprintAssignable)
-	FOnButtonClickedEvent OnClicked;
 
+	/** Called when the button is clicked */
+	UPROPERTY(EditDefaultsOnly, Category=Events)
+	FOnReply OnClickedEvent;
+	
 	/*  */
 	UFUNCTION(BlueprintPure, Category="Appearance")
 	FLinearColor GetButtonColorAndOpacity();
diff --git a/Engine/Source/Runtime/UMG/Public/Components/SlateWrapperComponent.h b/Engine/Source/Runtime/UMG/Public/Components/SlateWrapperComponent.h
--- a/Engine/Source/Runtime/UMG/Public/Components/SlateWrapperComponent.h
+++ b/Engine/Source/Runtime/UMG/Public/Components/SlateWrapperComponent.h
@@ -14,12 +14,16 @@ class UMG_API USlateWrapperComponent : public UActorComponent
 {
 	GENERATED_UCLASS_BODY()
 
+	// Common Bindings
 	DECLARE_DYNAMIC_DELEGATE_RetVal(bool, FGetBool);
 	DECLARE_DYNAMIC_DELEGATE_RetVal(FText, FGetText);
 	DECLARE_DYNAMIC_DELEGATE_RetVal(FLinearColor, FGetSlateColor);
 	DECLARE_DYNAMIC_DELEGATE_RetVal(FLinearColor, FGetLinearColor);
 	DECLARE_DYNAMIC_DELEGATE_RetVal(ESlateVisibility::Type, FGetSlateVisibility);
 
+	// Events
+	DECLARE_DYNAMIC_DELEGATE_RetVal(FSReply, FOnReply);
+
 	UPROPERTY(EditAnywhere, Category=Variable)
 	FString Name;
 
diff --git a/Engine/Source/Runtime/UMG/Public/Components/SlateWrapperTypes.h b/Engine/Source/Runtime/UMG/Public/Components/SlateWrapperTypes.h
--- a/Engine/Source/Runtime/UMG/Public/Components/SlateWrapperTypes.h
+++ b/Engine/Source/Runtime/UMG/Public/Components/SlateWrapperTypes.h
@@ -42,6 +42,27 @@ namespace ESlateSizeRule
 	};
 }
 
+USTRUCT()
+struct FSReply
+{
+	GENERATED_USTRUCT_BODY()
+
+	UPROPERTY(EditAnywhere, Category=User)
+		uint32 bIsHandled : 1;
+
+	FReply ToReply() const
+	{
+		FReply Reply = FReply::Unhandled();
+
+		if ( bIsHandled )
+		{
+			Reply = FReply::Handled();
+		}
+
+		return Reply;
+	}
+};
+
 USTRUCT()
 struct FSlateChildSize
 {
