UMG - Can now dynamically create widget bp's at runtime.  Successfully able to create a main menu.  Workflow improvements to wrap existing UI elements.  Delete fixed.

diff --git a/Engine/Source/Editor/UMGEditor/Private/AssetTypeActions_WidgetBlueprint.cpp b/Engine/Source/Editor/UMGEditor/Private/AssetTypeActions_WidgetBlueprint.cpp
--- a/Engine/Source/Editor/UMGEditor/Private/AssetTypeActions_WidgetBlueprint.cpp
+++ b/Engine/Source/Editor/UMGEditor/Private/AssetTypeActions_WidgetBlueprint.cpp
@@ -12,6 +12,7 @@
 #include "MessageLog.h"
 
 #include "AssetTypeActions_WidgetBlueprint.h"
+#include "WidgetBlueprintEditor.h"
 
 #define LOCTEXT_NAMESPACE "AssetTypeActions"
 
@@ -24,12 +25,18 @@ void FAssetTypeActions_WidgetBlueprint::OpenAssetEditor( const TArray<UObject*>&
 		auto Blueprint = Cast<UBlueprint>(*ObjIt);
 		if (Blueprint && Blueprint->SkeletonGeneratedClass && Blueprint->GeneratedClass )
 		{
-			FBlueprintEditorModule& BlueprintEditorModule = FModuleManager::LoadModuleChecked<FBlueprintEditorModule>( "Kismet" );
-			TSharedRef< IBlueprintEditor > NewKismetEditor = BlueprintEditorModule.CreateBlueprintEditor( Mode, EditWithinLevelEditor, Blueprint, ShouldUseDataOnlyEditor(Blueprint) );
+			TSharedRef< FWidgetBlueprintEditor > NewBlueprintEditor(new FWidgetBlueprintEditor());
+
+			TArray<UBlueprint*> Blueprints;
+			Blueprints.Add(Blueprint);
+			NewBlueprintEditor->InitWidgetBlueprintEditor(Mode, EditWithinLevelEditor, Blueprints, ShouldUseDataOnlyEditor(Blueprint));
+
+			//FBlueprintEditorModule& BlueprintEditorModule = FModuleManager::LoadModuleChecked<FBlueprintEditorModule>( "Kismet" );
+			//TSharedRef< IBlueprintEditor > NewKismetEditor = BlueprintEditorModule.CreateBlueprintEditor( Mode, EditWithinLevelEditor, Blueprint, ShouldUseDataOnlyEditor(Blueprint) );
 		}
 		else
 		{
-			FMessageDialog::Open( EAppMsgType::Ok, LOCTEXT("FailedToLoadBlueprint", "Blueprint could not be loaded because it derives from an invalid class.  Check to make sure the parent class for this blueprint hasn't been removed!"));
+			FMessageDialog::Open( EAppMsgType::Ok, LOCTEXT("FailedToLoadBlueprint", "Widget Blueprint could not be loaded because it derives from an invalid class.  Check to make sure the parent class for this blueprint hasn't been removed!"));
 		}
 	}
 }
diff --git a/Engine/Source/Editor/UMGEditor/Private/SUMGEditorTree.cpp b/Engine/Source/Editor/UMGEditor/Private/SUMGEditorTree.cpp
--- a/Engine/Source/Editor/UMGEditor/Private/SUMGEditorTree.cpp
+++ b/Engine/Source/Editor/UMGEditor/Private/SUMGEditorTree.cpp
@@ -12,6 +12,9 @@
 #include "BlueprintEditor.h"
 #include "SKismetInspector.h"
 #include "BlueprintEditorUtils.h"
+#include "WidgetTemplateClass.h"
+
+#define LOCTEXT_NAMESPACE "UMG"
 
 void SUMGEditorTree::Construct(const FArguments& InArgs, TSharedPtr<FBlueprintEditor> InBlueprintEditor, USimpleConstructionScript* InSCS)
 {
@@ -37,17 +40,6 @@ void SUMGEditorTree::Construct(const FArguments& InArgs, TSharedPtr<FBlueprintEd
 			]
 		]
 
-		+ SVerticalBox::Slot()
-		.AutoHeight()
-		[
-			SNew(SButton)
-			.OnClicked(this, &SUMGEditorTree::DeleteSelected)
-			[
-				SNew(STextBlock)
-				.Text(NSLOCTEXT("SUMGEditorTree", "DeleteSelected", "Delete Selected"))
-			]
-		]
-
 		+ SVerticalBox::Slot()
 		.FillHeight(1.0f)
 		[
@@ -57,6 +49,7 @@ void SUMGEditorTree::Construct(const FArguments& InArgs, TSharedPtr<FBlueprintEd
 			.OnGetChildren(this, &SUMGEditorTree::WidgetHierarchy_OnGetChildren)
 			.OnGenerateRow(this, &SUMGEditorTree::WidgetHierarchy_OnGenerateRow)
 			.OnSelectionChanged(this, &SUMGEditorTree::WidgetHierarchy_OnSelectionChanged)
+			.OnContextMenuOpening(this, &SUMGEditorTree::WidgetHierarchy_OnContextMenuOpening)
 			.TreeItemsSource(&RootWidgets)
 		]
 	];
@@ -124,6 +117,59 @@ void SUMGEditorTree::ShowDetailsForObjects(TArray<USlateWrapperComponent*> Widge
 	BlueprintEditor.Pin()->GetInspector()->ShowDetailsForObjects(InspectorObjects, Options);
 }
 
+void SUMGEditorTree::BuildWrapWithMenu(FMenuBuilder& Menu)
+{
+	Menu.BeginSection("WrapWith", LOCTEXT("WidgetTree_WrapWith", "Wrap With..."));
+	{
+		for ( TObjectIterator<UClass> ClassIt; ClassIt; ++ClassIt )
+		{
+			UClass* WidgetClass = *ClassIt;
+			if ( WidgetClass->IsChildOf(USlateNonLeafWidgetComponent::StaticClass()) && WidgetClass->HasAnyClassFlags(CLASS_Abstract) == false )
+			{
+				Menu.AddMenuEntry(
+					WidgetClass->GetDisplayNameText(),
+					FText::GetEmpty(),
+					FSlateIcon(),
+					FUIAction(
+						FExecuteAction::CreateSP(this, &SUMGEditorTree::WrapSelectedWidgets, WidgetClass),
+						FCanExecuteAction()
+					)
+				);
+			}
+		}
+	}
+	Menu.EndSection();
+}
+
+TSharedPtr<SWidget> SUMGEditorTree::WidgetHierarchy_OnContextMenuOpening()
+{
+	FMenuBuilder MenuBuilder(true, NULL);
+
+	MenuBuilder.BeginSection("Actions");
+	{
+		const TArray<USlateWrapperComponent*> SelectionList = WidgetTreeView->GetSelectedItems();
+
+		MenuBuilder.AddMenuEntry(
+			LOCTEXT("WidgetTree_Delete", "Delete"),
+			LOCTEXT("WidgetTree_DeleteToolTip", "Deletes the current widget"),
+			FSlateIcon(),
+			FUIAction(
+			FExecuteAction::CreateSP(this, &SUMGEditorTree::DeleteSelected),
+			FCanExecuteAction()
+			)
+		);
+
+		MenuBuilder.AddSubMenu(
+			LOCTEXT("WidgetTree_WrapWith", "Wrap With..."),
+			LOCTEXT("WidgetTree_WrapWithToolTip", "Wraps the currently selected widgets inside of another container widget"),
+			FNewMenuDelegate::CreateSP(this, &SUMGEditorTree::BuildWrapWithMenu)
+		);
+	}
+	MenuBuilder.EndSection();
+
+	return MenuBuilder.MakeWidget();
+}
+
 void SUMGEditorTree::WidgetHierarchy_OnGetChildren(USlateWrapperComponent* InParent, TArray< USlateWrapperComponent* >& OutChildren)
 {
 	USlateNonLeafWidgetComponent* Widget = Cast<USlateNonLeafWidgetComponent>(InParent);
@@ -161,6 +207,27 @@ void SUMGEditorTree::WidgetHierarchy_OnSelectionChanged(USlateWrapperComponent*
 	}
 }
 
+void SUMGEditorTree::WrapSelectedWidgets(UClass* WidgetClass)
+{
+	const TArray<USlateWrapperComponent*> SelectionList = WidgetTreeView->GetSelectedItems();
+
+	TSharedPtr<FWidgetTemplateClass> Template = MakeShareable(new FWidgetTemplateClass(WidgetClass));
+
+	UWidgetBlueprint* BP = GetBlueprint();
+	USlateNonLeafWidgetComponent* NewWrapperWidget = CastChecked<USlateNonLeafWidgetComponent>( Template->Create(BP->WidgetTree) );
+
+	int32 OutIndex;
+	USlateNonLeafWidgetComponent* CurrentParent = BP->WidgetTree->FindWidgetParent(SelectionList[0], OutIndex);
+	if ( CurrentParent )
+	{
+		CurrentParent->ReplaceChildAt(OutIndex, NewWrapperWidget);
+
+		NewWrapperWidget->AddChild(SelectionList[0], FVector2D());
+	}
+
+	FBlueprintEditorUtils::MarkBlueprintAsStructurallyModified(BP);
+}
+
 FReply SUMGEditorTree::CreateTestUI()
 {
 	UWidgetBlueprint* BP = GetBlueprint();
@@ -207,7 +274,14 @@ FReply SUMGEditorTree::CreateTestUI()
 	return FReply::Handled();
 }
 
-FReply SUMGEditorTree::DeleteSelected()
+FReply SUMGEditorTree::HandleDeleteSelected()
+{
+	DeleteSelected();
+
+	return FReply::Handled();
+}
+
+void SUMGEditorTree::DeleteSelected()
 {
 	TArray<USlateWrapperComponent*> SelectedItems = WidgetTreeView->GetSelectedItems();
 	if ( SelectedItems.Num() > 0 )
@@ -225,8 +299,6 @@ FReply SUMGEditorTree::DeleteSelected()
 			FBlueprintEditorUtils::MarkBlueprintAsStructurallyModified(BP);
 		}
 	}
-
-	return FReply::Handled();
 }
 
 void SUMGEditorTree::RefreshTree()
@@ -245,4 +317,6 @@ void SUMGEditorTree::RefreshTree()
 }
 
 
-//@TODO UMG Drop widgets onto the tree, when nothing is present, if there is a root node present, what happens then, let the root node attempt to place it?
\ No newline at end of file
+//@TODO UMG Drop widgets onto the tree, when nothing is present, if there is a root node present, what happens then, let the root node attempt to place it?
+
+#undef LOCTEXT_NAMESPACE
diff --git a/Engine/Source/Editor/UMGEditor/Private/SUMGEditorTree.h b/Engine/Source/Editor/UMGEditor/Private/SUMGEditorTree.h
--- a/Engine/Source/Editor/UMGEditor/Private/SUMGEditorTree.h
+++ b/Engine/Source/Editor/UMGEditor/Private/SUMGEditorTree.h
@@ -23,6 +23,8 @@ private:
 
 	void ShowDetailsForObjects(TArray<USlateWrapperComponent*> Widgets);
 
+	void BuildWrapWithMenu(FMenuBuilder& Menu);
+	TSharedPtr<SWidget> WidgetHierarchy_OnContextMenuOpening();
 	void WidgetHierarchy_OnGetChildren(USlateWrapperComponent* InParent, TArray< USlateWrapperComponent* >& OutChildren);
 	TSharedRef< ITableRow > WidgetHierarchy_OnGenerateRow(USlateWrapperComponent* InItem, const TSharedRef<STableViewBase>& OwnerTable);
 	void WidgetHierarchy_OnSelectionChanged(USlateWrapperComponent* SelectedItem, ESelectInfo::Type SelectInfo);
@@ -30,7 +32,10 @@ private:
 	FReply CreateTestUI();
 	void RefreshTree();
 
-	FReply DeleteSelected();
+	FReply HandleDeleteSelected();
+	void DeleteSelected();
+
+	void WrapSelectedWidgets(UClass* WidgetClass);
 	
 private:
 
diff --git a/Engine/Source/Editor/UMGEditor/Private/WidgetBlueprintCompiler.cpp b/Engine/Source/Editor/UMGEditor/Private/WidgetBlueprintCompiler.cpp
--- a/Engine/Source/Editor/UMGEditor/Private/WidgetBlueprintCompiler.cpp
+++ b/Engine/Source/Editor/UMGEditor/Private/WidgetBlueprintCompiler.cpp
@@ -571,6 +571,10 @@ void FWidgetBlueprintCompiler::CleanAndSanitizeClass(UBlueprintGeneratedClass* C
 {
 	Super::CleanAndSanitizeClass(ClassToClean, OldCDO);
 
+	// Make sure our typed pointer is set
+	check(ClassToClean == NewClass);
+	NewWidgetBlueprintClass = CastChecked<UWidgetBlueprintGeneratedClass>((UObject*)NewClass);
+
 	// Purge all subobjects (properties, functions, params) of the class, as they will be regenerated
 	TArray<UObject*> ClassSubObjects;
 	GetObjectsWithOuter(ClassToClean, ClassSubObjects, true);
diff --git a/Engine/Source/Editor/UMGEditor/Private/WidgetBlueprintEditor.cpp b/Engine/Source/Editor/UMGEditor/Private/WidgetBlueprintEditor.cpp
--- a/Engine/Source/Editor/UMGEditor/Private/WidgetBlueprintEditor.cpp
+++ b/Engine/Source/Editor/UMGEditor/Private/WidgetBlueprintEditor.cpp
@@ -14,4 +14,9 @@ FWidgetBlueprintEditor::~FWidgetBlueprintEditor()
 {
 }
 
+void FWidgetBlueprintEditor::InitWidgetBlueprintEditor(const EToolkitMode::Type Mode, const TSharedPtr< IToolkitHost >& InitToolkitHost, const TArray<UBlueprint*>& InBlueprints, bool bShouldOpenInDefaultsMode)
+{
+	InitBlueprintEditor(Mode, InitToolkitHost, InBlueprints, bShouldOpenInDefaultsMode);
+}
+
 #undef LOCTEXT_NAMESPACE
diff --git a/Engine/Source/Editor/UMGEditor/Private/WidgetBlueprintEditor.h b/Engine/Source/Editor/UMGEditor/Private/WidgetBlueprintEditor.h
--- a/Engine/Source/Editor/UMGEditor/Private/WidgetBlueprintEditor.h
+++ b/Engine/Source/Editor/UMGEditor/Private/WidgetBlueprintEditor.h
@@ -13,4 +13,6 @@ public:
 	FWidgetBlueprintEditor();
 
 	virtual ~FWidgetBlueprintEditor();
+
+	void InitWidgetBlueprintEditor(const EToolkitMode::Type Mode, const TSharedPtr< IToolkitHost >& InitToolkitHost, const TArray<UBlueprint*>& InBlueprints, bool bShouldOpenInDefaultsMode);
 };
diff --git a/Engine/Source/Editor/UnrealEd/Private/Kismet2/BlueprintEditorUtils.cpp b/Engine/Source/Editor/UnrealEd/Private/Kismet2/BlueprintEditorUtils.cpp
--- a/Engine/Source/Editor/UnrealEd/Private/Kismet2/BlueprintEditorUtils.cpp
+++ b/Engine/Source/Editor/UnrealEd/Private/Kismet2/BlueprintEditorUtils.cpp
@@ -1287,7 +1287,8 @@ void FBlueprintEditorUtils::PostDuplicateBlueprint(UBlueprint* Blueprint)
 				FName NewSkelClassName, NewGenClassName;
 				Blueprint->GetBlueprintClassNames(NewGenClassName, NewSkelClassName);
 				UBlueprintGeneratedClass* NewClass = ConstructObject<UBlueprintGeneratedClass>(
-					UBlueprintGeneratedClass::StaticClass(), Blueprint->GetOutermost(), NewGenClassName, RF_Public|RF_Transactional);
+					Blueprint->GetBlueprintClass(), Blueprint->GetOutermost(), NewGenClassName, RF_Public | RF_Transactional);
+
 				Blueprint->GeneratedClass = NewClass;
 				NewClass->ClassGeneratedBy = Blueprint;
 				NewClass->SetSuperStruct(Blueprint->ParentClass);
diff --git a/Engine/Source/Runtime/Engine/Classes/Animation/AnimBlueprint.h b/Engine/Source/Runtime/Engine/Classes/Animation/AnimBlueprint.h
--- a/Engine/Source/Runtime/Engine/Classes/Animation/AnimBlueprint.h
+++ b/Engine/Source/Runtime/Engine/Classes/Animation/AnimBlueprint.h
@@ -73,6 +73,8 @@ class ENGINE_API UAnimBlueprint : public UBlueprint
 
 #if WITH_EDITOR
 
+	virtual UClass* GetBlueprintClass() const OVERRIDE;
+
 	// Inspects the hierarchy and looks for an override for the requested node GUID
 	// @param NodeGuid - Guid of the node to search for
 	// @param bIgnoreSelf - Ignore this blueprint and only search parents, handy for finding parent overrides
diff --git a/Engine/Source/Runtime/Engine/Classes/Engine/Blueprint.h b/Engine/Source/Runtime/Engine/Classes/Engine/Blueprint.h
--- a/Engine/Source/Runtime/Engine/Classes/Engine/Blueprint.h
+++ b/Engine/Source/Runtime/Engine/Classes/Engine/Blueprint.h
@@ -443,6 +443,9 @@ public:
 		SkeletonClassName = FName(*SkeletonClassNameString);
 	}
 
+	/** Gets the class generated when this blueprint is compiled. */
+	virtual UClass* GetBlueprintClass() const;
+
 	// Should the generic blueprint factory work for this blueprint?
 	virtual bool SupportedByDefaultBlueprintFactory() const
 	{
diff --git a/Engine/Source/Runtime/Engine/Private/AnimBlueprint.cpp b/Engine/Source/Runtime/Engine/Private/AnimBlueprint.cpp
--- a/Engine/Source/Runtime/Engine/Private/AnimBlueprint.cpp
+++ b/Engine/Source/Runtime/Engine/Private/AnimBlueprint.cpp
@@ -26,6 +26,11 @@ UAnimBlueprintGeneratedClass* UAnimBlueprint::GetAnimBlueprintSkeletonClass() co
 
 #if WITH_EDITOR
 
+UClass* UAnimBlueprint::GetBlueprintClass() const
+{
+	return UAnimBlueprintGeneratedClass::StaticClass();
+}
+
 int32 UAnimBlueprint::FindOrAddGroup(FName GroupName)
 {
 	if (GroupName == NAME_None)
diff --git a/Engine/Source/Runtime/Engine/Private/Blueprint.cpp b/Engine/Source/Runtime/Engine/Private/Blueprint.cpp
--- a/Engine/Source/Runtime/Engine/Private/Blueprint.cpp
+++ b/Engine/Source/Runtime/Engine/Private/Blueprint.cpp
@@ -300,7 +300,7 @@ void UBlueprint::PostLoad()
 	FBlueprintEditorUtils::PurgeNullGraphs(this);
 
 	// Remove old AutoConstructionScript graph
- 	for (int32 i = 0; i < FunctionGraphs.Num(); ++i)
+	for (int32 i = 0; i < FunctionGraphs.Num(); ++i)
 	{
 		UEdGraph* FuncGraph = FunctionGraphs[i];
 		if ((FuncGraph != NULL) && (FuncGraph->GetName() == TEXT("AutoConstructionScript")))
@@ -419,6 +419,11 @@ void UBlueprint::DebuggingWorldRegistrationHelper(UObject* ObjectProvidingWorld,
 	}
 }
 
+UClass* UBlueprint::GetBlueprintClass() const
+{
+	return UBlueprintGeneratedClass::StaticClass();
+}
+
 void UBlueprint::SetObjectBeingDebugged(UObject* NewObject)
 {
 	// Unregister the old object
diff --git a/Engine/Source/Runtime/UMG/Private/Components/ContentWidget.cpp b/Engine/Source/Runtime/UMG/Private/Components/ContentWidget.cpp
--- a/Engine/Source/Runtime/UMG/Private/Components/ContentWidget.cpp
+++ b/Engine/Source/Runtime/UMG/Private/Components/ContentWidget.cpp
@@ -58,6 +58,12 @@ bool UContentWidget::RemoveChild(USlateWrapperComponent* Child)
 	return false;
 }
 
+void UContentWidget::ReplaceChildAt(int32 Index, USlateWrapperComponent* Child)
+{
+	check(Index == 0);
+	Content = Child;
+}
+
 /////////////////////////////////////////////////////
 
 #undef LOCTEXT_NAMESPACE
diff --git a/Engine/Source/Runtime/UMG/Private/Components/HorizontalBoxComponent.cpp b/Engine/Source/Runtime/UMG/Private/Components/HorizontalBoxComponent.cpp
--- a/Engine/Source/Runtime/UMG/Private/Components/HorizontalBoxComponent.cpp
+++ b/Engine/Source/Runtime/UMG/Private/Components/HorizontalBoxComponent.cpp
@@ -43,6 +43,16 @@ bool UHorizontalBoxComponent::RemoveChild(USlateWrapperComponent* Child)
 	return false;
 }
 
+void UHorizontalBoxComponent::ReplaceChildAt(int32 Index, USlateWrapperComponent* Content)
+{
+	UHorizontalBoxSlot* Slot = Slots[Index];
+	Slot->Content = Content;
+
+#if WITH_EDITOR
+	Content->Slot = Slot;
+#endif
+}
+
 TSharedRef<SWidget> UHorizontalBoxComponent::RebuildWidget()
 {
 	TSharedRef<SHorizontalBox> NewCanvas = SNew(SHorizontalBox);
diff --git a/Engine/Source/Runtime/UMG/Private/Components/UniformGridPanel.cpp b/Engine/Source/Runtime/UMG/Private/Components/UniformGridPanel.cpp
new file mode 100644
--- /dev/null
+++ b/Engine/Source/Runtime/UMG/Private/Components/UniformGridPanel.cpp
@@ -0,0 +1,119 @@
+// Copyright 1998-2014 Epic Games, Inc. All Rights Reserved.
+
+#include "UMGPrivatePCH.h"
+
+/////////////////////////////////////////////////////
+// UUniformGridPanel
+
+UUniformGridPanel::UUniformGridPanel(const FPostConstructInitializeProperties& PCIP)
+	: Super(PCIP)
+{
+	bIsVariable = false;
+}
+
+int32 UUniformGridPanel::GetChildrenCount() const
+{
+	return Slots.Num();
+}
+
+USlateWrapperComponent* UUniformGridPanel::GetChildAt(int32 Index) const
+{
+	return Slots[Index]->Content;
+}
+
+bool UUniformGridPanel::AddChild(USlateWrapperComponent* Child, FVector2D Position)
+{
+	AddSlot(Child);
+	return true;
+}
+
+bool UUniformGridPanel::RemoveChild(USlateWrapperComponent* Child)
+{
+	for ( int32 SlotIndex = 0; SlotIndex < Slots.Num(); ++SlotIndex )
+	{
+		UUniformGridSlot* Slot = Slots[SlotIndex];
+
+		if ( Slot->Content == Child )
+		{
+			Slots.RemoveAt(SlotIndex);
+			return true;
+		}
+	}
+
+	return false;
+}
+
+void UUniformGridPanel::ReplaceChildAt(int32 Index, USlateWrapperComponent* Content)
+{
+	UUniformGridSlot* Slot = Slots[Index];
+	Slot->Content = Content;
+
+#if WITH_EDITOR
+	Content->Slot = Slot;
+#endif
+}
+
+TSharedRef<SWidget> UUniformGridPanel::RebuildWidget()
+{
+	TSharedRef<SUniformGridPanel> NewPanel = SNew(SUniformGridPanel);
+	MyUniformGridPanel = NewPanel;
+
+	for ( int32 SlotIndex = 0; SlotIndex < Slots.Num(); ++SlotIndex )
+	{
+		UUniformGridSlot* Slot = Slots[SlotIndex];
+		if ( Slot == NULL )
+		{
+			Slots[SlotIndex] = Slot = ConstructObject<UUniformGridSlot>(UUniformGridSlot::StaticClass(), this);
+		}
+
+		auto& NewSlot = NewPanel->AddSlot(Slot->Column, Slot->Row)
+			.HAlign(Slot->HorizontalAlignment)
+			.VAlign(Slot->VerticalAlignment)
+			[
+				Slot->Content == NULL ? SNullWidget::NullWidget : Slot->Content->GetWidget()
+			];
+	}
+
+	return NewPanel;
+}
+
+UUniformGridSlot* UUniformGridPanel::AddSlot(USlateWrapperComponent* Content)
+{
+	UUniformGridSlot* Slot = ConstructObject<UUniformGridSlot>(UUniformGridSlot::StaticClass(), this);
+	Slot->Content = Content;
+
+#if WITH_EDITOR
+	Content->Slot = Slot;
+#endif
+	
+	Slots.Add(Slot);
+
+	return Slot;
+}
+
+#if WITH_EDITOR
+
+void UUniformGridPanel::ConnectEditorData()
+{
+	for ( UUniformGridSlot* Slot : Slots )
+	{
+		Slot->Content->Slot = Slot;
+	}
+}
+
+void UUniformGridPanel::PostEditChangeProperty(struct FPropertyChangedEvent& PropertyChangedEvent)
+{
+	// Ensure the slots have unique names
+	int32 SlotNumbering = 1;
+
+	TSet<FName> UniqueSlotNames;
+	for (int32 SlotIndex = 0; SlotIndex < Slots.Num(); ++SlotIndex)
+	{
+		if ( Slots[SlotIndex] == NULL )
+		{
+			Slots[SlotIndex] = ConstructObject<UUniformGridSlot>(UUniformGridSlot::StaticClass(), this);
+		}
+	}
+}
+
+#endif
diff --git a/Engine/Source/Runtime/UMG/Private/Components/UniformGridSlot.cpp b/Engine/Source/Runtime/UMG/Private/Components/UniformGridSlot.cpp
new file mode 100644
--- /dev/null
+++ b/Engine/Source/Runtime/UMG/Private/Components/UniformGridSlot.cpp
@@ -0,0 +1,13 @@
+// Copyright 1998-2014 Epic Games, Inc. All Rights Reserved.
+
+#include "UMGPrivatePCH.h"
+
+/////////////////////////////////////////////////////
+// UUniformGridSlot
+
+UUniformGridSlot::UUniformGridSlot(const FPostConstructInitializeProperties& PCIP)
+	: Super(PCIP)
+{
+	HorizontalAlignment = HAlign_Left;
+	VerticalAlignment = VAlign_Top;
+}
diff --git a/Engine/Source/Runtime/UMG/Private/Components/VerticalBoxComponent.cpp b/Engine/Source/Runtime/UMG/Private/Components/VerticalBoxComponent.cpp
--- a/Engine/Source/Runtime/UMG/Private/Components/VerticalBoxComponent.cpp
+++ b/Engine/Source/Runtime/UMG/Private/Components/VerticalBoxComponent.cpp
@@ -43,6 +43,16 @@ bool UVerticalBoxComponent::RemoveChild(USlateWrapperComponent* Child)
 	return false;
 }
 
+void UVerticalBoxComponent::ReplaceChildAt(int32 Index, USlateWrapperComponent* Content)
+{
+	UVerticalBoxSlot* Slot = Slots[Index];
+	Slot->Content = Content;
+
+#if WITH_EDITOR
+	Content->Slot = Slot;
+#endif
+}
+
 TSharedRef<SWidget> UVerticalBoxComponent::RebuildWidget()
 {
 	TSharedRef<SVerticalBox> NewCanvas = SNew(SVerticalBox);
diff --git a/Engine/Source/Runtime/UMG/Private/WidgetBlueprint.cpp b/Engine/Source/Runtime/UMG/Private/WidgetBlueprint.cpp
--- a/Engine/Source/Runtime/UMG/Private/WidgetBlueprint.cpp
+++ b/Engine/Source/Runtime/UMG/Private/WidgetBlueprint.cpp
@@ -35,6 +35,13 @@ void UWidgetBlueprint::PostLoad()
 #endif
 }
 
+#if WITH_EDITOR
+UClass* UWidgetBlueprint::GetBlueprintClass() const
+{
+	return UWidgetBlueprintGeneratedClass::StaticClass();
+}
+#endif
+
 bool UWidgetBlueprint::ValidateGeneratedClass(const UClass* InClass)
 {
 	bool Result = Super::ValidateGeneratedClass(InClass);
diff --git a/Engine/Source/Runtime/UMG/Private/WidgetTree.cpp b/Engine/Source/Runtime/UMG/Private/WidgetTree.cpp
--- a/Engine/Source/Runtime/UMG/Private/WidgetTree.cpp
+++ b/Engine/Source/Runtime/UMG/Private/WidgetTree.cpp
@@ -32,18 +32,50 @@ USlateWrapperComponent* UWidgetTree::FindWidget(FString& Name) const
 	return NULL;
 }
 
-bool UWidgetTree::RemoveWidget(USlateWrapperComponent* Widget)
+USlateNonLeafWidgetComponent* UWidgetTree::FindWidgetParent(USlateWrapperComponent* Widget, int32& OutChildIndex)
+{
+	for ( USlateWrapperComponent* Template : WidgetTemplates )
+	{
+		USlateNonLeafWidgetComponent* NonLeafTemplate = Cast<USlateNonLeafWidgetComponent>(Template);
+		if ( NonLeafTemplate )
+		{
+			for ( int32 ChildIndex = 0; ChildIndex < NonLeafTemplate->GetChildrenCount(); ChildIndex++ )
+			{
+				if ( NonLeafTemplate->GetChildAt(ChildIndex) == Widget )
+				{
+					OutChildIndex = ChildIndex;
+					return NonLeafTemplate;
+				}
+			}
+		}
+	}
+
+	OutChildIndex = -1;
+	return NULL;
+}
+
+bool UWidgetTree::RemoveWidget(USlateWrapperComponent* InRemovedWidget)
 {
 	USlateWrapperComponent* Parent = NULL;
 
 	bool bRemoved = false;
 
+	//TODO UMG Make the Widget Tree actually a tree, instead of a list, it makes things like removal difficult.
+
+	if ( USlateNonLeafWidgetComponent* InNonLeafRemovedWidget = Cast<USlateNonLeafWidgetComponent>(InRemovedWidget) )
+	{
+		while ( InNonLeafRemovedWidget->GetChildrenCount() > 0 )
+		{
+			RemoveWidget(InNonLeafRemovedWidget->GetChildAt(0));
+		}
+	}
+
 	for ( USlateWrapperComponent* Template : WidgetTemplates )
 	{
 		USlateNonLeafWidgetComponent* NonLeafTemplate = Cast<USlateNonLeafWidgetComponent>(Template);
 		if ( NonLeafTemplate )
 		{
-			if ( NonLeafTemplate->RemoveChild(Widget) )
+			if ( NonLeafTemplate->RemoveChild(InRemovedWidget) )
 			{
 				bRemoved = true;
 				break;
@@ -53,7 +85,7 @@ bool UWidgetTree::RemoveWidget(USlateWrapperComponent* Widget)
 
 	if ( bRemoved )
 	{
-		WidgetTemplates.Remove(Widget);
+		WidgetTemplates.Remove(InRemovedWidget);
 	}
 
 	return bRemoved;
diff --git a/Engine/Source/Runtime/UMG/Public/Blueprint/WidgetBlueprint.h b/Engine/Source/Runtime/UMG/Public/Blueprint/WidgetBlueprint.h
--- a/Engine/Source/Runtime/UMG/Public/Blueprint/WidgetBlueprint.h
+++ b/Engine/Source/Runtime/UMG/Public/Blueprint/WidgetBlueprint.h
@@ -51,6 +51,8 @@ public:
 	
 #if WITH_EDITOR
 	// UBlueprint interface
+	virtual UClass* GetBlueprintClass() const OVERRIDE;
+
 	virtual bool SupportedByDefaultBlueprintFactory() const OVERRIDE
 	{
 		return false;
diff --git a/Engine/Source/Runtime/UMG/Public/Blueprint/WidgetTree.h b/Engine/Source/Runtime/UMG/Public/Blueprint/WidgetTree.h
--- a/Engine/Source/Runtime/UMG/Public/Blueprint/WidgetTree.h
+++ b/Engine/Source/Runtime/UMG/Public/Blueprint/WidgetTree.h
@@ -19,6 +19,8 @@ class UMG_API UWidgetTree : public UObject
 
 	bool RemoveWidget(class USlateWrapperComponent* Widget);
 
+	class USlateNonLeafWidgetComponent* FindWidgetParent(class USlateWrapperComponent* Widget, int32& OutChildIndex);
+
 	template< class T >
 	T* ConstructWidget(TSubclassOf<class USlateWrapperComponent> WidgetType)
 	{
diff --git a/Engine/Source/Runtime/UMG/Public/Components/ContentWidget.h b/Engine/Source/Runtime/UMG/Public/Components/ContentWidget.h
--- a/Engine/Source/Runtime/UMG/Public/Components/ContentWidget.h
+++ b/Engine/Source/Runtime/UMG/Public/Components/ContentWidget.h
@@ -20,6 +20,7 @@ class UMG_API UContentWidget : public USlateNonLeafWidgetComponent
 	virtual USlateWrapperComponent* GetChildAt(int32 Index) const OVERRIDE;
 	virtual bool AddChild(USlateWrapperComponent* Child, FVector2D Position) OVERRIDE;
 	virtual bool RemoveChild(USlateWrapperComponent* Child) OVERRIDE;
+	virtual void ReplaceChildAt(int32 Index, USlateWrapperComponent* Child) OVERRIDE;
 	// End USlateNonLeafWidgetComponent
 
 protected:
diff --git a/Engine/Source/Runtime/UMG/Public/Components/HorizontalBoxComponent.h b/Engine/Source/Runtime/UMG/Public/Components/HorizontalBoxComponent.h
--- a/Engine/Source/Runtime/UMG/Public/Components/HorizontalBoxComponent.h
+++ b/Engine/Source/Runtime/UMG/Public/Components/HorizontalBoxComponent.h
@@ -20,6 +20,7 @@ class UMG_API UHorizontalBoxComponent : public USlateNonLeafWidgetComponent
 	virtual USlateWrapperComponent* GetChildAt(int32 Index) const OVERRIDE;
 	virtual bool AddChild(USlateWrapperComponent* Child, FVector2D Position) OVERRIDE;
 	virtual bool RemoveChild(USlateWrapperComponent* Child) OVERRIDE;
+	virtual void ReplaceChildAt(int32 Index, USlateWrapperComponent* Child) OVERRIDE;
 	// End USlateNonLeafWidgetComponent
 
 #if WITH_EDITOR
diff --git a/Engine/Source/Runtime/UMG/Public/Components/SlateNonLeafWidgetComponent.h b/Engine/Source/Runtime/UMG/Public/Components/SlateNonLeafWidgetComponent.h
--- a/Engine/Source/Runtime/UMG/Public/Components/SlateNonLeafWidgetComponent.h
+++ b/Engine/Source/Runtime/UMG/Public/Components/SlateNonLeafWidgetComponent.h
@@ -12,6 +12,10 @@ class UMG_API USlateNonLeafWidgetComponent : public USlateWrapperComponent
 	virtual int32 GetChildrenCount() const { return 0; }
 	virtual USlateWrapperComponent* GetChildAt(int32 Index) const { return NULL; }
 
+	virtual void ReplaceChildAt(int32 Index, USlateWrapperComponent* Child)
+	{
+	}
+
 	virtual bool AddChild(USlateWrapperComponent* Child, FVector2D Position)
 	{
 		return false;
diff --git a/Engine/Source/Runtime/UMG/Public/Components/UniformGridPanel.h b/Engine/Source/Runtime/UMG/Public/Components/UniformGridPanel.h
new file mode 100644
--- /dev/null
+++ b/Engine/Source/Runtime/UMG/Public/Components/UniformGridPanel.h
@@ -0,0 +1,44 @@
+// Copyright 1998-2014 Epic Games, Inc. All Rights Reserved.
+
+#pragma once
+
+#include "UniformGridPanel.generated.h"
+
+UCLASS(meta=(BlueprintSpawnableComponent), ClassGroup=UserInterface)
+class UMG_API UUniformGridPanel : public USlateNonLeafWidgetComponent
+{
+	GENERATED_UCLASS_BODY()
+
+	/** The items placed on the canvas */
+	UPROPERTY(EditAnywhere, EditInline, Category=Slots)
+	TArray<UUniformGridSlot*> Slots;
+
+	UUniformGridSlot* AddSlot(USlateWrapperComponent* Content);
+
+	// USlateNonLeafWidgetComponent
+	virtual int32 GetChildrenCount() const OVERRIDE;
+	virtual USlateWrapperComponent* GetChildAt(int32 Index) const OVERRIDE;
+	virtual bool AddChild(USlateWrapperComponent* Child, FVector2D Position) OVERRIDE;
+	virtual bool RemoveChild(USlateWrapperComponent* Child) OVERRIDE;
+	virtual void ReplaceChildAt(int32 Index, USlateWrapperComponent* Child) OVERRIDE;
+	// End USlateNonLeafWidgetComponent
+
+#if WITH_EDITOR
+	// UObject interface
+	virtual void PostEditChangeProperty(struct FPropertyChangedEvent& PropertyChangedEvent) OVERRIDE;
+	// End of UObject interface
+
+	// USlateWrapperComponent interface
+	virtual void ConnectEditorData() OVERRIDE;
+	// End USlateWrapperComponent interface
+#endif
+
+protected:
+
+	TWeakPtr<class SUniformGridPanel> MyUniformGridPanel;
+
+protected:
+	// USlateWrapperComponent interface
+	virtual TSharedRef<SWidget> RebuildWidget() OVERRIDE;
+	// End of USlateWrapperComponent interface
+};
diff --git a/Engine/Source/Runtime/UMG/Public/Components/UniformGridSlot.h b/Engine/Source/Runtime/UMG/Public/Components/UniformGridSlot.h
new file mode 100644
--- /dev/null
+++ b/Engine/Source/Runtime/UMG/Public/Components/UniformGridSlot.h
@@ -0,0 +1,25 @@
+// Copyright 1998-2014 Epic Games, Inc. All Rights Reserved.
+
+#pragma once
+
+#include "SlateWrapperTypes.h"
+
+#include "UniformGridSlot.generated.h"
+
+UCLASS()
+class UMG_API UUniformGridSlot : public UPanelSlot
+{
+	GENERATED_UCLASS_BODY()
+
+	UPROPERTY(EditAnywhere, Category=Appearance)
+	TEnumAsByte<EHorizontalAlignment> HorizontalAlignment;
+
+	UPROPERTY(EditAnywhere, Category=Appearance)
+	TEnumAsByte<EVerticalAlignment> VerticalAlignment;
+	
+	UPROPERTY(EditAnywhere, Category=Appearance)
+	int32 Row;
+	
+	UPROPERTY(EditAnywhere, Category=Appearance)
+	int32 Column;
+};
\ No newline at end of file
diff --git a/Engine/Source/Runtime/UMG/Public/Components/VerticalBoxComponent.h b/Engine/Source/Runtime/UMG/Public/Components/VerticalBoxComponent.h
--- a/Engine/Source/Runtime/UMG/Public/Components/VerticalBoxComponent.h
+++ b/Engine/Source/Runtime/UMG/Public/Components/VerticalBoxComponent.h
@@ -20,6 +20,7 @@ class UMG_API UVerticalBoxComponent : public USlateNonLeafWidgetComponent
 	virtual USlateWrapperComponent* GetChildAt(int32 Index) const OVERRIDE;
 	virtual bool AddChild(USlateWrapperComponent* Child, FVector2D Position) OVERRIDE;
 	virtual bool RemoveChild(USlateWrapperComponent* Child) OVERRIDE;
+	virtual void ReplaceChildAt(int32 Index, USlateWrapperComponent* Child) OVERRIDE;
 	// End USlateNonLeafWidgetComponent
 
 #if WITH_EDITOR
diff --git a/Engine/Source/Runtime/UMG/Public/UMG.h b/Engine/Source/Runtime/UMG/Public/UMG.h
--- a/Engine/Source/Runtime/UMG/Public/UMG.h
+++ b/Engine/Source/Runtime/UMG/Public/UMG.h
@@ -43,6 +43,9 @@
 #include "VerticalBoxSlot.h"
 #include "VerticalBoxComponent.h"
 
+#include "UniformGridSlot.h"
+#include "UniformGridPanel.h"
+
 #include "TextBlockComponent.h"
 #include "BorderComponent.h"
 #include "ButtonComponent.h"
