UMG - Fixing non-unity build.

diff --git a/Engine/Source/Editor/UMGEditor/Private/UMGBlueprintEditorExtensionHook.cpp b/Engine/Source/Editor/UMGEditor/Private/UMGBlueprintEditorExtensionHook.cpp
--- a/Engine/Source/Editor/UMGEditor/Private/UMGBlueprintEditorExtensionHook.cpp
+++ b/Engine/Source/Editor/UMGEditor/Private/UMGBlueprintEditorExtensionHook.cpp
@@ -19,6 +19,8 @@
 
 #include "DetailCustomizations.h"
 
+#include "WidgetBlueprintEditor.h"
+
 #define LOCTEXT_NAMESPACE "UMG_EXTENSION"
 
 /////////////////////////////////////////////////////
