UMG - Renaming several base classes.

diff --git a/Engine/Source/Editor/UMGEditor/Private/SUMGDesigner.cpp b/Engine/Source/Editor/UMGEditor/Private/SUMGDesigner.cpp
--- a/Engine/Source/Editor/UMGEditor/Private/SUMGDesigner.cpp
+++ b/Engine/Source/Editor/UMGEditor/Private/SUMGDesigner.cpp
@@ -180,7 +180,7 @@ void SUMGDesigner::OnObjectPropertyChanged(UObject* ObjectBeingModified, FProper
 	//UpdatePreview(InBlueprint);
 }
 
-void SUMGDesigner::ShowDetailsForObjects(TArray<USlateWrapperComponent*> Widgets)
+void SUMGDesigner::ShowDetailsForObjects(TArray<UWidget*> Widgets)
 {
 	// @TODO COde duplication
 
@@ -188,7 +188,7 @@ void SUMGDesigner::ShowDetailsForObjects(TArray<USlateWrapperComponent*> Widgets
 	FString InspectorTitle;
 	TArray<UObject*> InspectorObjects;
 	InspectorObjects.Empty(Widgets.Num());
-	for ( USlateWrapperComponent* Widget : Widgets )
+	for ( UWidget* Widget : Widgets )
 	{
 		//if ( NodePtr->CanEditDefaults() )
 		{
@@ -204,7 +204,7 @@ void SUMGDesigner::ShowDetailsForObjects(TArray<USlateWrapperComponent*> Widgets
 	BlueprintEditor.Pin()->GetInspector()->ShowDetailsForObjects(InspectorObjects, Options);
 }
 
-USlateWrapperComponent* SUMGDesigner::GetTemplateAtCursor(const FGeometry& MyGeometry, const FPointerEvent& MouseEvent, FArrangedWidget& ArrangedWidget)
+UWidget* SUMGDesigner::GetTemplateAtCursor(const FGeometry& MyGeometry, const FPointerEvent& MouseEvent, FArrangedWidget& ArrangedWidget)
 {
 	//@TODO UMG Make it so you can request dropable widgets only, to find the first parentable.
 
@@ -219,7 +219,7 @@ USlateWrapperComponent* SUMGDesigner::GetTemplateAtCursor(const FGeometry& MyGeo
 	UUserWidget* WidgetActor = BlueprintEditor.Pin()->GetPreview();
 	if ( WidgetActor )
 	{
-		USlateWrapperComponent* PreviewHandle = NULL;
+		UWidget* PreviewHandle = NULL;
 		for ( int32 ChildIndex = Children.Num() - 1; ChildIndex >= 0; ChildIndex-- )
 		{
 			FArrangedWidget& Child = Children.GetInternalArray()[ChildIndex];
@@ -236,7 +236,7 @@ USlateWrapperComponent* SUMGDesigner::GetTemplateAtCursor(const FGeometry& MyGeo
 		if ( PreviewHandle )
 		{
 			FString Name = PreviewHandle->GetName();
-			USlateWrapperComponent* Template = Blueprint->WidgetTree->FindWidget(Name);
+			UWidget* Template = Blueprint->WidgetTree->FindWidget(Name);
 			return Template;
 		}
 	}
@@ -257,7 +257,7 @@ FReply SUMGDesigner::OnMouseButtonDown(const FGeometry& MyGeometry, const FPoint
 		{
 			//@TODO UMG primary FBlueprintEditor needs to be inherited and selection control needs to be centralized.
 			// Set the template as selected in the details panel
-			TArray<USlateWrapperComponent*> SelectedTemplates;
+			TArray<UWidget*> SelectedTemplates;
 			SelectedTemplates.Add(SelectedTemplate);
 			ShowDetailsForObjects(SelectedTemplates);
 		}
@@ -578,15 +578,15 @@ FReply SUMGDesigner::OnDrop(const FGeometry& MyGeometry, const FDragDropEvent& D
 	if ( DragDropOp.IsValid() )
 	{
 		FArrangedWidget ArrangedWidget(SNullWidget::NullWidget, FGeometry());
-		USlateWrapperComponent* Template = GetTemplateAtCursor(MyGeometry, DragDropEvent, ArrangedWidget);
+		UWidget* Template = GetTemplateAtCursor(MyGeometry, DragDropEvent, ArrangedWidget);
 
 		UWidgetBlueprint* BP = GetBlueprint();
 
-		if ( Template && Template->IsA(USlateNonLeafWidgetComponent::StaticClass()) )
+		if ( Template && Template->IsA(UPanelWidget::StaticClass()) )
 		{
-			USlateNonLeafWidgetComponent* Parent = Cast<USlateNonLeafWidgetComponent>(Template);
+			UPanelWidget* Parent = Cast<UPanelWidget>(Template);
 
-			USlateWrapperComponent* Widget = DragDropOp->Template->Create(BP->WidgetTree);
+			UWidget* Widget = DragDropOp->Template->Create(BP->WidgetTree);
 
 			FVector2D LocalPosition = ArrangedWidget.Geometry.AbsoluteToLocal(DragDropEvent.GetScreenSpacePosition());
 			Parent->AddChild(Widget, LocalPosition);
@@ -603,7 +603,7 @@ FReply SUMGDesigner::OnDrop(const FGeometry& MyGeometry, const FDragDropEvent& D
 		else if ( BP->WidgetTree->WidgetTemplates.Num() == 0 )
 		{
 			// No existing templates so just create it and make it the root widget.
-			USlateWrapperComponent* Widget = DragDropOp->Template->Create(BP->WidgetTree);
+			UWidget* Widget = DragDropOp->Template->Create(BP->WidgetTree);
 			SelectedTemplate = Widget;
 
 			FBlueprintEditorUtils::MarkBlueprintAsStructurallyModified(BP);
diff --git a/Engine/Source/Editor/UMGEditor/Private/SUMGDesigner.h b/Engine/Source/Editor/UMGEditor/Private/SUMGDesigner.h
--- a/Engine/Source/Editor/UMGEditor/Private/SUMGDesigner.h
+++ b/Engine/Source/Editor/UMGEditor/Private/SUMGDesigner.h
@@ -33,12 +33,12 @@ public:
 	// End of SWidget interface
 
 private:
-	USlateWrapperComponent* GetTemplateAtCursor(const FGeometry& MyGeometry, const FPointerEvent& MouseEvent, FArrangedWidget& ArrangedWidget);
+	UWidget* GetTemplateAtCursor(const FGeometry& MyGeometry, const FPointerEvent& MouseEvent, FArrangedWidget& ArrangedWidget);
 
 	UWidgetBlueprint* GetBlueprint() const;
 	void OnBlueprintChanged(UBlueprint* InBlueprint);
 	void OnObjectPropertyChanged(UObject* ObjectBeingModified, FPropertyChangedEvent& PropertyChangedEvent);
-	void ShowDetailsForObjects(TArray<USlateWrapperComponent*> Widgets);
+	void ShowDetailsForObjects(TArray<UWidget*> Widgets);
 
 	bool GetArrangedWidget(TSharedRef<SWidget> Widget, FArrangedWidget& ArrangedWidget) const;
 	bool GetArrangedWidgetRelativeToWindow(TSharedRef<SWidget> Widget, FArrangedWidget& ArrangedWidget) const;
@@ -74,7 +74,7 @@ private:
 	TWeakPtr<SWidget> PreviewWidget;
 	UUserWidget* PreviewWidgetActor;
 
-	USlateWrapperComponent* SelectedTemplate;
+	UWidget* SelectedTemplate;
 	TWeakPtr<SWidget> SelectedWidget;
 
 	TSharedPtr<SBorder> PreviewSurface;
diff --git a/Engine/Source/Editor/UMGEditor/Private/SUMGEditorTree.cpp b/Engine/Source/Editor/UMGEditor/Private/SUMGEditorTree.cpp
--- a/Engine/Source/Editor/UMGEditor/Private/SUMGEditorTree.cpp
+++ b/Engine/Source/Editor/UMGEditor/Private/SUMGEditorTree.cpp
@@ -43,7 +43,7 @@ void SUMGEditorTree::Construct(const FArguments& InArgs, TSharedPtr<FBlueprintEd
 		+ SVerticalBox::Slot()
 		.FillHeight(1.0f)
 		[
-			SAssignNew(WidgetTreeView, STreeView< USlateWrapperComponent* >)
+			SAssignNew(WidgetTreeView, STreeView< UWidget* >)
 			.ItemHeight(20.0f)
 			.SelectionMode(ESelectionMode::Single)
 			.OnGetChildren(this, &SUMGEditorTree::WidgetHierarchy_OnGetChildren)
@@ -95,13 +95,13 @@ void SUMGEditorTree::OnObjectPropertyChanged(UObject* ObjectBeingModified, FProp
 	}
 }
 
-void SUMGEditorTree::ShowDetailsForObjects(TArray<USlateWrapperComponent*> Widgets)
+void SUMGEditorTree::ShowDetailsForObjects(TArray<UWidget*> Widgets)
 {
 	// Convert the selection set to an array of UObject* pointers
 	FString InspectorTitle;
 	TArray<UObject*> InspectorObjects;
 	InspectorObjects.Empty(Widgets.Num());
-	for ( USlateWrapperComponent* Widget : Widgets )
+	for ( UWidget* Widget : Widgets )
 	{
 		//if ( NodePtr->CanEditDefaults() )
 		{
@@ -124,7 +124,7 @@ void SUMGEditorTree::BuildWrapWithMenu(FMenuBuilder& Menu)
 		for ( TObjectIterator<UClass> ClassIt; ClassIt; ++ClassIt )
 		{
 			UClass* WidgetClass = *ClassIt;
-			if ( WidgetClass->IsChildOf(USlateNonLeafWidgetComponent::StaticClass()) && WidgetClass->HasAnyClassFlags(CLASS_Abstract) == false )
+			if ( WidgetClass->IsChildOf(UPanelWidget::StaticClass()) && WidgetClass->HasAnyClassFlags(CLASS_Abstract) == false )
 			{
 				Menu.AddMenuEntry(
 					WidgetClass->GetDisplayNameText(),
@@ -147,7 +147,7 @@ TSharedPtr<SWidget> SUMGEditorTree::WidgetHierarchy_OnContextMenuOpening()
 
 	MenuBuilder.BeginSection("Actions");
 	{
-		const TArray<USlateWrapperComponent*> SelectionList = WidgetTreeView->GetSelectedItems();
+		const TArray<UWidget*> SelectionList = WidgetTreeView->GetSelectedItems();
 
 		MenuBuilder.AddMenuEntry(
 			LOCTEXT("WidgetTree_Delete", "Delete"),
@@ -170,14 +170,14 @@ TSharedPtr<SWidget> SUMGEditorTree::WidgetHierarchy_OnContextMenuOpening()
 	return MenuBuilder.MakeWidget();
 }
 
-void SUMGEditorTree::WidgetHierarchy_OnGetChildren(USlateWrapperComponent* InParent, TArray< USlateWrapperComponent* >& OutChildren)
+void SUMGEditorTree::WidgetHierarchy_OnGetChildren(UWidget* InParent, TArray< UWidget* >& OutChildren)
 {
-	USlateNonLeafWidgetComponent* Widget = Cast<USlateNonLeafWidgetComponent>(InParent);
+	UPanelWidget* Widget = Cast<UPanelWidget>(InParent);
 	if ( Widget )
 	{
 		for ( int32 i = 0; i < Widget->GetChildrenCount(); i++ )
 		{
-			USlateWrapperComponent* Child = Widget->GetChildAt(i);
+			UWidget* Child = Widget->GetChildAt(i);
 			if ( Child )
 			{
 				OutChildren.Add(Child);
@@ -186,10 +186,10 @@ void SUMGEditorTree::WidgetHierarchy_OnGetChildren(USlateWrapperComponent* InPar
 	}
 }
 
-TSharedRef< ITableRow > SUMGEditorTree::WidgetHierarchy_OnGenerateRow(USlateWrapperComponent* InItem, const TSharedRef<STableViewBase>& OwnerTable)
+TSharedRef< ITableRow > SUMGEditorTree::WidgetHierarchy_OnGenerateRow(UWidget* InItem, const TSharedRef<STableViewBase>& OwnerTable)
 {
 	return
-		SNew(STableRow< USlateWrapperComponent* >, OwnerTable)
+		SNew(STableRow< UWidget* >, OwnerTable)
 		.Padding(2.0f)
 //		.OnDragDetected(this, &SUMGEditorWidgetTemplates::OnDraggingWidgetTemplateItem)
 		[
@@ -197,11 +197,11 @@ TSharedRef< ITableRow > SUMGEditorTree::WidgetHierarchy_OnGenerateRow(USlateWrap
 		];
 }
 
-void SUMGEditorTree::WidgetHierarchy_OnSelectionChanged(USlateWrapperComponent* SelectedItem, ESelectInfo::Type SelectInfo)
+void SUMGEditorTree::WidgetHierarchy_OnSelectionChanged(UWidget* SelectedItem, ESelectInfo::Type SelectInfo)
 {
 	if ( SelectInfo != ESelectInfo::Direct )
 	{
-		TArray<USlateWrapperComponent*> Items;
+		TArray<UWidget*> Items;
 		Items.Add(SelectedItem);
 		ShowDetailsForObjects(Items);
 	}
@@ -209,15 +209,15 @@ void SUMGEditorTree::WidgetHierarchy_OnSelectionChanged(USlateWrapperComponent*
 
 void SUMGEditorTree::WrapSelectedWidgets(UClass* WidgetClass)
 {
-	const TArray<USlateWrapperComponent*> SelectionList = WidgetTreeView->GetSelectedItems();
+	const TArray<UWidget*> SelectionList = WidgetTreeView->GetSelectedItems();
 
 	TSharedPtr<FWidgetTemplateClass> Template = MakeShareable(new FWidgetTemplateClass(WidgetClass));
 
 	UWidgetBlueprint* BP = GetBlueprint();
-	USlateNonLeafWidgetComponent* NewWrapperWidget = CastChecked<USlateNonLeafWidgetComponent>( Template->Create(BP->WidgetTree) );
+	UPanelWidget* NewWrapperWidget = CastChecked<UPanelWidget>( Template->Create(BP->WidgetTree) );
 
 	int32 OutIndex;
-	USlateNonLeafWidgetComponent* CurrentParent = BP->WidgetTree->FindWidgetParent(SelectionList[0], OutIndex);
+	UPanelWidget* CurrentParent = BP->WidgetTree->FindWidgetParent(SelectionList[0], OutIndex);
 	if ( CurrentParent )
 	{
 		CurrentParent->ReplaceChildAt(OutIndex, NewWrapperWidget);
@@ -231,7 +231,7 @@ void SUMGEditorTree::WrapSelectedWidgets(UClass* WidgetClass)
 FReply SUMGEditorTree::CreateTestUI()
 {
 	UWidgetBlueprint* BP = GetBlueprint();
-	TArray<USlateWrapperComponent*>& WidgetTemplates = BP->WidgetTree->WidgetTemplates;
+	TArray<UWidget*>& WidgetTemplates = BP->WidgetTree->WidgetTemplates;
 
 	if ( WidgetTemplates.Num() > 0 )
 	{
@@ -283,13 +283,13 @@ FReply SUMGEditorTree::HandleDeleteSelected()
 
 void SUMGEditorTree::DeleteSelected()
 {
-	TArray<USlateWrapperComponent*> SelectedItems = WidgetTreeView->GetSelectedItems();
+	TArray<UWidget*> SelectedItems = WidgetTreeView->GetSelectedItems();
 	if ( SelectedItems.Num() > 0 )
 	{
 		UWidgetBlueprint* BP = GetBlueprint();
 
 		bool bRemoved = false;
-		for ( USlateWrapperComponent* Item : SelectedItems )
+		for ( UWidget* Item : SelectedItems )
 		{
 			bRemoved = BP->WidgetTree->RemoveWidget(Item);
 		}
@@ -306,7 +306,7 @@ void SUMGEditorTree::RefreshTree()
 	RootWidgets.Reset();
 
 	UWidgetBlueprint* Blueprint = GetBlueprint();
-	TArray<USlateWrapperComponent*>& WidgetTemplates = Blueprint->WidgetTree->WidgetTemplates;
+	TArray<UWidget*>& WidgetTemplates = Blueprint->WidgetTree->WidgetTemplates;
 
 	if ( WidgetTemplates.Num() > 0 )
 	{
diff --git a/Engine/Source/Editor/UMGEditor/Private/SUMGEditorTree.h b/Engine/Source/Editor/UMGEditor/Private/SUMGEditorTree.h
--- a/Engine/Source/Editor/UMGEditor/Private/SUMGEditorTree.h
+++ b/Engine/Source/Editor/UMGEditor/Private/SUMGEditorTree.h
@@ -21,13 +21,13 @@ private:
 	void OnBlueprintChanged(UBlueprint* InBlueprint);
 	void OnObjectPropertyChanged(UObject* ObjectBeingModified, FPropertyChangedEvent& PropertyChangedEvent);
 
-	void ShowDetailsForObjects(TArray<USlateWrapperComponent*> Widgets);
+	void ShowDetailsForObjects(TArray<UWidget*> Widgets);
 
 	void BuildWrapWithMenu(FMenuBuilder& Menu);
 	TSharedPtr<SWidget> WidgetHierarchy_OnContextMenuOpening();
-	void WidgetHierarchy_OnGetChildren(USlateWrapperComponent* InParent, TArray< USlateWrapperComponent* >& OutChildren);
-	TSharedRef< ITableRow > WidgetHierarchy_OnGenerateRow(USlateWrapperComponent* InItem, const TSharedRef<STableViewBase>& OwnerTable);
-	void WidgetHierarchy_OnSelectionChanged(USlateWrapperComponent* SelectedItem, ESelectInfo::Type SelectInfo);
+	void WidgetHierarchy_OnGetChildren(UWidget* InParent, TArray< UWidget* >& OutChildren);
+	TSharedRef< ITableRow > WidgetHierarchy_OnGenerateRow(UWidget* InItem, const TSharedRef<STableViewBase>& OwnerTable);
+	void WidgetHierarchy_OnSelectionChanged(UWidget* SelectedItem, ESelectInfo::Type SelectInfo);
 
 	FReply CreateTestUI();
 	void RefreshTree();
@@ -41,7 +41,7 @@ private:
 
 	TWeakPtr<class FBlueprintEditor> BlueprintEditor;
 
-	TArray< USlateWrapperComponent* > RootWidgets;
+	TArray< UWidget* > RootWidgets;
 
-	TSharedPtr< STreeView< USlateWrapperComponent* > > WidgetTreeView;
+	TSharedPtr< STreeView< UWidget* > > WidgetTreeView;
 };
diff --git a/Engine/Source/Editor/UMGEditor/Private/SUMGEditorTreeItem.cpp b/Engine/Source/Editor/UMGEditor/Private/SUMGEditorTreeItem.cpp
--- a/Engine/Source/Editor/UMGEditor/Private/SUMGEditorTreeItem.cpp
+++ b/Engine/Source/Editor/UMGEditor/Private/SUMGEditorTreeItem.cpp
@@ -13,7 +13,7 @@
 #include "SKismetInspector.h"
 #include "BlueprintEditorUtils.h"
 
-void SUMGEditorTreeItem::Construct(const FArguments& InArgs, TSharedPtr<FBlueprintEditor> InBlueprintEditor, USlateWrapperComponent* InItem)
+void SUMGEditorTreeItem::Construct(const FArguments& InArgs, TSharedPtr<FBlueprintEditor> InBlueprintEditor, UWidget* InItem)
 {
 	BlueprintEditor = InBlueprintEditor;
 	Item = InItem;
@@ -76,12 +76,12 @@ FReply SUMGEditorTreeItem::OnDrop(const FGeometry& MyGeometry, const FDragDropEv
 	TSharedPtr<FWidgetTemplateDragDropOp> DragDropOp = DragDropEvent.GetOperationAs<FWidgetTemplateDragDropOp>();
 	if ( DragDropOp.IsValid() )
 	{
-		if ( Item->IsA(USlateNonLeafWidgetComponent::StaticClass()) )
+		if ( Item->IsA(UPanelWidget::StaticClass()) )
 		{
 			UWidgetBlueprint* BP = CastChecked<UWidgetBlueprint>(BlueprintEditor.Pin()->GetBlueprintObj());			
-			USlateNonLeafWidgetComponent* Parent = Cast<USlateNonLeafWidgetComponent>(Item);
+			UPanelWidget* Parent = Cast<UPanelWidget>(Item);
 
-			USlateWrapperComponent* Widget = DragDropOp->Template->Create(BP->WidgetTree);
+			UWidget* Widget = DragDropOp->Template->Create(BP->WidgetTree);
 
 			Parent->AddChild(Widget, FVector2D(0,0));
 
diff --git a/Engine/Source/Editor/UMGEditor/Private/SUMGEditorTreeItem.h b/Engine/Source/Editor/UMGEditor/Private/SUMGEditorTreeItem.h
--- a/Engine/Source/Editor/UMGEditor/Private/SUMGEditorTreeItem.h
+++ b/Engine/Source/Editor/UMGEditor/Private/SUMGEditorTreeItem.h
@@ -12,7 +12,7 @@ public:
 	SLATE_BEGIN_ARGS( SUMGEditorTreeItem ){}
 	SLATE_END_ARGS()
 
-	void Construct(const FArguments& InArgs, TSharedPtr<FBlueprintEditor> InBlueprintEditor, USlateWrapperComponent* InItem);
+	void Construct(const FArguments& InArgs, TSharedPtr<FBlueprintEditor> InBlueprintEditor, UWidget* InItem);
 
 	virtual void OnDragEnter(const FGeometry& MyGeometry, const FDragDropEvent& DragDropEvent) OVERRIDE;
 	virtual void OnDragLeave(const FDragDropEvent& DragDropEvent) OVERRIDE;
@@ -24,5 +24,5 @@ private:
 	FString GetItemTooltipText() const;
 
 	TWeakPtr<class FBlueprintEditor> BlueprintEditor;
-	USlateWrapperComponent* Item;
+	UWidget* Item;
 };
diff --git a/Engine/Source/Editor/UMGEditor/Private/SUMGEditorWidgetTemplates.cpp b/Engine/Source/Editor/UMGEditor/Private/SUMGEditorWidgetTemplates.cpp
--- a/Engine/Source/Editor/UMGEditor/Private/SUMGEditorWidgetTemplates.cpp
+++ b/Engine/Source/Editor/UMGEditor/Private/SUMGEditorWidgetTemplates.cpp
@@ -123,7 +123,7 @@ void SUMGEditorWidgetTemplates::BuildClassWidgetList()
 	for ( TObjectIterator<UClass> ClassIt; ClassIt; ++ClassIt )
 	{
 		UClass* WidgetClass = *ClassIt;
-		if ( WidgetClass->IsChildOf(USlateWrapperComponent::StaticClass()) && WidgetClass->HasAnyClassFlags(CLASS_Abstract) == false )
+		if ( WidgetClass->IsChildOf(UWidget::StaticClass()) && WidgetClass->HasAnyClassFlags(CLASS_Abstract) == false )
 		{
 			TSharedPtr<FWidgetTemplateClass> Template = MakeShareable(new FWidgetTemplateClass(WidgetClass));
 			Template->Name = WidgetClass->GetDisplayNameText();
diff --git a/Engine/Source/Editor/UMGEditor/Private/Templates/WidgetTemplate.h b/Engine/Source/Editor/UMGEditor/Private/Templates/WidgetTemplate.h
--- a/Engine/Source/Editor/UMGEditor/Private/Templates/WidgetTemplate.h
+++ b/Engine/Source/Editor/UMGEditor/Private/Templates/WidgetTemplate.h
@@ -9,7 +9,7 @@ public:
 
 	virtual FText GetCategory() = 0;
 
-	virtual USlateWrapperComponent* Create(class UWidgetTree* Tree) = 0;
+	virtual UWidget* Create(class UWidgetTree* Tree) = 0;
 
 public:
 	FText Name;
diff --git a/Engine/Source/Editor/UMGEditor/Private/Templates/WidgetTemplateButton.cpp b/Engine/Source/Editor/UMGEditor/Private/Templates/WidgetTemplateButton.cpp
--- a/Engine/Source/Editor/UMGEditor/Private/Templates/WidgetTemplateButton.cpp
+++ b/Engine/Source/Editor/UMGEditor/Private/Templates/WidgetTemplateButton.cpp
@@ -18,7 +18,7 @@ FText FWidgetTemplateButton::GetCategory()
 	return LOCTEXT("Common", "Common");
 }
 
-USlateWrapperComponent* FWidgetTemplateButton::Create(UWidgetTree* Tree)
+UWidget* FWidgetTemplateButton::Create(UWidgetTree* Tree)
 {
 	UButtonComponent* Button = Tree->ConstructWidget<UButtonComponent>(UButtonComponent::StaticClass());
 	UTextBlockComponent* Text = Tree->ConstructWidget<UTextBlockComponent>(UTextBlockComponent::StaticClass());
diff --git a/Engine/Source/Editor/UMGEditor/Private/Templates/WidgetTemplateButton.h b/Engine/Source/Editor/UMGEditor/Private/Templates/WidgetTemplateButton.h
--- a/Engine/Source/Editor/UMGEditor/Private/Templates/WidgetTemplateButton.h
+++ b/Engine/Source/Editor/UMGEditor/Private/Templates/WidgetTemplateButton.h
@@ -11,5 +11,5 @@ public:
 
 	virtual FText GetCategory() OVERRIDE;
 
-	virtual USlateWrapperComponent* Create(UWidgetTree* Tree) OVERRIDE;
+	virtual UWidget* Create(UWidgetTree* Tree) OVERRIDE;
 };
diff --git a/Engine/Source/Editor/UMGEditor/Private/Templates/WidgetTemplateCheckBox.cpp b/Engine/Source/Editor/UMGEditor/Private/Templates/WidgetTemplateCheckBox.cpp
--- a/Engine/Source/Editor/UMGEditor/Private/Templates/WidgetTemplateCheckBox.cpp
+++ b/Engine/Source/Editor/UMGEditor/Private/Templates/WidgetTemplateCheckBox.cpp
@@ -18,7 +18,7 @@ FText FWidgetTemplateCheckBox::GetCategory()
 	return LOCTEXT("Common", "Common");
 }
 
-USlateWrapperComponent* FWidgetTemplateCheckBox::Create(UWidgetTree* Tree)
+UWidget* FWidgetTemplateCheckBox::Create(UWidgetTree* Tree)
 {
 	UHorizontalBoxComponent* Horizontal = Tree->ConstructWidget<UHorizontalBoxComponent>(UHorizontalBoxComponent::StaticClass());
 	UCheckBoxComponent* Checkbox = Tree->ConstructWidget<UCheckBoxComponent>(UCheckBoxComponent::StaticClass());
diff --git a/Engine/Source/Editor/UMGEditor/Private/Templates/WidgetTemplateCheckBox.h b/Engine/Source/Editor/UMGEditor/Private/Templates/WidgetTemplateCheckBox.h
--- a/Engine/Source/Editor/UMGEditor/Private/Templates/WidgetTemplateCheckBox.h
+++ b/Engine/Source/Editor/UMGEditor/Private/Templates/WidgetTemplateCheckBox.h
@@ -11,5 +11,5 @@ public:
 
 	virtual FText GetCategory() OVERRIDE;
 
-	virtual USlateWrapperComponent* Create(UWidgetTree* Tree) OVERRIDE;
+	virtual UWidget* Create(UWidgetTree* Tree) OVERRIDE;
 };
diff --git a/Engine/Source/Editor/UMGEditor/Private/Templates/WidgetTemplateClass.cpp b/Engine/Source/Editor/UMGEditor/Private/Templates/WidgetTemplateClass.cpp
--- a/Engine/Source/Editor/UMGEditor/Private/Templates/WidgetTemplateClass.cpp
+++ b/Engine/Source/Editor/UMGEditor/Private/Templates/WidgetTemplateClass.cpp
@@ -6,7 +6,7 @@
 
 #define LOCTEXT_NAMESPACE "UMGEditor"
 
-FWidgetTemplateClass::FWidgetTemplateClass(TSubclassOf<USlateWrapperComponent> InWidgetClass)
+FWidgetTemplateClass::FWidgetTemplateClass(TSubclassOf<UWidget> InWidgetClass)
 	: WidgetClass(InWidgetClass)
 {
 }
@@ -17,9 +17,9 @@ FText FWidgetTemplateClass::GetCategory()
 	return LOCTEXT("Class", "Class");
 }
 
-USlateWrapperComponent* FWidgetTemplateClass::Create(UWidgetTree* Tree)
+UWidget* FWidgetTemplateClass::Create(UWidgetTree* Tree)
 {
-	return Tree->ConstructWidget<USlateWrapperComponent>(WidgetClass);
+	return Tree->ConstructWidget<UWidget>(WidgetClass);
 }
 
 #undef LOCTEXT_NAMESPACE
diff --git a/Engine/Source/Editor/UMGEditor/Private/Templates/WidgetTemplateClass.h b/Engine/Source/Editor/UMGEditor/Private/Templates/WidgetTemplateClass.h
--- a/Engine/Source/Editor/UMGEditor/Private/Templates/WidgetTemplateClass.h
+++ b/Engine/Source/Editor/UMGEditor/Private/Templates/WidgetTemplateClass.h
@@ -7,12 +7,12 @@
 class FWidgetTemplateClass : public FWidgetTemplate
 {
 public:
-	FWidgetTemplateClass(TSubclassOf<USlateWrapperComponent> InWidgetClass);
+	FWidgetTemplateClass(TSubclassOf<UWidget> InWidgetClass);
 
 	virtual FText GetCategory() OVERRIDE;
 
-	virtual USlateWrapperComponent* Create(UWidgetTree* Tree) OVERRIDE;
+	virtual UWidget* Create(UWidgetTree* Tree) OVERRIDE;
 
 protected:
-	TSubclassOf<USlateWrapperComponent> WidgetClass;
+	TSubclassOf<UWidget> WidgetClass;
 };
diff --git a/Engine/Source/Editor/UMGEditor/Private/UMGBlueprintEditorExtensionHook.cpp b/Engine/Source/Editor/UMGEditor/Private/UMGBlueprintEditorExtensionHook.cpp
--- a/Engine/Source/Editor/UMGEditor/Private/UMGBlueprintEditorExtensionHook.cpp
+++ b/Engine/Source/Editor/UMGEditor/Private/UMGBlueprintEditorExtensionHook.cpp
@@ -224,7 +224,7 @@ public:
 
 		TSharedRef<class SKismetInspector> Inspector = MyBlueprintEditor.Pin()->GetInspector();
 		FOnGetDetailCustomizationInstance LayoutDelegateDetails = FOnGetDetailCustomizationInstance::CreateStatic(&FBlueprintWidgetCustomization::MakeInstance, MyBlueprintEditor.Pin()->GetBlueprintObj());
-		Inspector->GetPropertyView()->RegisterInstancedCustomPropertyLayout(USlateWrapperComponent::StaticClass(), LayoutDelegateDetails);
+		Inspector->GetPropertyView()->RegisterInstancedCustomPropertyLayout(UWidget::StaticClass(), LayoutDelegateDetails);
 		//Inspector->GetPropertyView()->RegisterInstancedCustomPropertyLayout(UDelegateProperty::StaticClass(), LayoutDelegateDetails);
 	}
 
diff --git a/Engine/Source/Editor/UMGEditor/Private/WidgetBlueprintCompiler.cpp b/Engine/Source/Editor/UMGEditor/Private/WidgetBlueprintCompiler.cpp
--- a/Engine/Source/Editor/UMGEditor/Private/WidgetBlueprintCompiler.cpp
+++ b/Engine/Source/Editor/UMGEditor/Private/WidgetBlueprintCompiler.cpp
@@ -536,7 +536,7 @@ void FWidgetBlueprintCompiler::ValidateWidgetNames()
 		}
 	}
 
-	for ( USlateWrapperComponent* Widget : Blueprint->WidgetTree->WidgetTemplates )
+	for ( UWidget* Widget : Blueprint->WidgetTree->WidgetTemplates )
 	{
 		if ( ParentBPNameValidator.IsValid() && ParentBPNameValidator->IsValid(Widget->GetName()) != EValidatorResult::Ok )
 		{
@@ -583,7 +583,7 @@ void FWidgetBlueprintCompiler::CleanAndSanitizeClass(UBlueprintGeneratedClass* C
 
 	if ( 0 != Blueprint->WidgetTree->WidgetTemplates.Num() )
 	{
-		ClassSubObjects.RemoveAllSwap(FCullTemplateObjectsHelper<USlateWrapperComponent>(Blueprint->WidgetTree->WidgetTemplates));
+		ClassSubObjects.RemoveAllSwap(FCullTemplateObjectsHelper<UWidget>(Blueprint->WidgetTree->WidgetTemplates));
 	}
 }
 
@@ -595,7 +595,7 @@ void FWidgetBlueprintCompiler::CreateClassVariablesFromBlueprint()
 
 	ValidateWidgetNames();
 
-	for ( USlateWrapperComponent* Widget : Blueprint->WidgetTree->WidgetTemplates )
+	for ( UWidget* Widget : Blueprint->WidgetTree->WidgetTemplates )
 	{
 		// Skip non-variable widgets
 		if ( !Widget->bIsVariable )
diff --git a/Engine/Source/Editor/UMGEditor/Private/WidgetBlueprintCompiler.h b/Engine/Source/Editor/UMGEditor/Private/WidgetBlueprintCompiler.h
--- a/Engine/Source/Editor/UMGEditor/Private/WidgetBlueprintCompiler.h
+++ b/Engine/Source/Editor/UMGEditor/Private/WidgetBlueprintCompiler.h
@@ -119,7 +119,7 @@ protected:
 	UWidgetBlueprintGeneratedClass* NewWidgetBlueprintClass;
 
 	// Map of properties created for widgets; to aid in debug data generation
-	TMap<class USlateWrapperComponent*, class UProperty*> WidgetToMemberVariableMap;
+	TMap<class UWidget*, class UProperty*> WidgetToMemberVariableMap;
 
 	///----------------------------------------------------------------
 	// Callback system properties.
diff --git a/Engine/Source/Runtime/UMG/Private/Components/BorderComponent.cpp b/Engine/Source/Runtime/UMG/Private/Components/BorderComponent.cpp
--- a/Engine/Source/Runtime/UMG/Private/Components/BorderComponent.cpp
+++ b/Engine/Source/Runtime/UMG/Private/Components/BorderComponent.cpp
@@ -50,7 +50,7 @@ TSharedRef<SWidget> UBorderComponent::RebuildWidget()
 	return NewBorder;
 }
 
-void UBorderComponent::SetContent(USlateWrapperComponent* Content)
+void UBorderComponent::SetContent(UWidget* Content)
 {
 	Super::SetContent(Content);
 
diff --git a/Engine/Source/Runtime/UMG/Private/Components/ButtonComponent.cpp b/Engine/Source/Runtime/UMG/Private/Components/ButtonComponent.cpp
--- a/Engine/Source/Runtime/UMG/Private/Components/ButtonComponent.cpp
+++ b/Engine/Source/Runtime/UMG/Private/Components/ButtonComponent.cpp
@@ -65,7 +65,7 @@ TSharedRef<SWidget> UButtonComponent::RebuildWidget()
 	return NewButton;
 }
 
-void UButtonComponent::SetContent(USlateWrapperComponent* InContent)
+void UButtonComponent::SetContent(UWidget* InContent)
 {
 	Super::SetContent(InContent);
 
diff --git a/Engine/Source/Runtime/UMG/Private/Components/CanvasPanelComponent.cpp b/Engine/Source/Runtime/UMG/Private/Components/CanvasPanelComponent.cpp
--- a/Engine/Source/Runtime/UMG/Private/Components/CanvasPanelComponent.cpp
+++ b/Engine/Source/Runtime/UMG/Private/Components/CanvasPanelComponent.cpp
@@ -46,7 +46,7 @@ int32 UCanvasPanelComponent::GetChildrenCount() const
 	return Slots.Num();
 }
 
-USlateWrapperComponent* UCanvasPanelComponent::GetChildAt(int32 Index) const
+UWidget* UCanvasPanelComponent::GetChildAt(int32 Index) const
 {
 	return Slots[Index]->Content;
 }
@@ -83,7 +83,7 @@ TSharedRef<SWidget> UCanvasPanelComponent::RebuildWidget()
 	return NewCanvas;
 }
 
-UCanvasPanelSlot* UCanvasPanelComponent::AddSlot(USlateWrapperComponent* Content)
+UCanvasPanelSlot* UCanvasPanelComponent::AddSlot(UWidget* Content)
 {
 	UCanvasPanelSlot* Slot = ConstructObject<UCanvasPanelSlot>(UCanvasPanelSlot::StaticClass(), this);
 	Slot->Content = Content;
@@ -97,7 +97,7 @@ UCanvasPanelSlot* UCanvasPanelComponent::AddSlot(USlateWrapperComponent* Content
 	return Slot;
 }
 
-bool UCanvasPanelComponent::AddChild(USlateWrapperComponent* Child, FVector2D Position)
+bool UCanvasPanelComponent::AddChild(UWidget* Child, FVector2D Position)
 {
 	UCanvasPanelSlot* Slot = AddSlot(Child);
 	Slot->Position = Position;
diff --git a/Engine/Source/Runtime/UMG/Private/Components/ContentWidget.cpp b/Engine/Source/Runtime/UMG/Private/Components/ContentWidget.cpp
--- a/Engine/Source/Runtime/UMG/Private/Components/ContentWidget.cpp
+++ b/Engine/Source/Runtime/UMG/Private/Components/ContentWidget.cpp
@@ -21,12 +21,12 @@ UContentWidget::UContentWidget(const FPostConstructInitializeProperties& PCIP)
 	//bTickInEditor = true;
 }
 
-USlateWrapperComponent* UContentWidget::GetContent()
+UWidget* UContentWidget::GetContent()
 {
 	return Content;
 }
 
-void UContentWidget::SetContent(USlateWrapperComponent* InContent)
+void UContentWidget::SetContent(UWidget* InContent)
 {
 	Content = InContent;
 }
@@ -36,18 +36,18 @@ int32 UContentWidget::GetChildrenCount() const
 	return Content != NULL ? 1 : 0;
 }
 
-USlateWrapperComponent* UContentWidget::GetChildAt(int32 Index) const
+UWidget* UContentWidget::GetChildAt(int32 Index) const
 {
 	return Content;
 }
 
-bool UContentWidget::AddChild(USlateWrapperComponent* InContent, FVector2D Position)
+bool UContentWidget::AddChild(UWidget* InContent, FVector2D Position)
 {
 	Content = InContent;
 	return true;
 }
 
-bool UContentWidget::RemoveChild(USlateWrapperComponent* Child)
+bool UContentWidget::RemoveChild(UWidget* Child)
 {
 	if ( Content == Child )
 	{
@@ -58,7 +58,7 @@ bool UContentWidget::RemoveChild(USlateWrapperComponent* Child)
 	return false;
 }
 
-void UContentWidget::ReplaceChildAt(int32 Index, USlateWrapperComponent* Child)
+void UContentWidget::ReplaceChildAt(int32 Index, UWidget* Child)
 {
 	check(Index == 0);
 	Content = Child;
diff --git a/Engine/Source/Runtime/UMG/Private/Components/GridPanelComponent.cpp b/Engine/Source/Runtime/UMG/Private/Components/GridPanelComponent.cpp
--- a/Engine/Source/Runtime/UMG/Private/Components/GridPanelComponent.cpp
+++ b/Engine/Source/Runtime/UMG/Private/Components/GridPanelComponent.cpp
@@ -30,12 +30,12 @@ TSharedRef<SWidget> UGridPanelComponent::RebuildWidget()
 	//	}
 	//	
 	//	// Make sure all children are assigned to a slot
-	//	TArray<USlateWrapperComponent*> WrappedComponentsInSlotOrder;
+	//	TArray<UWidget*> WrappedComponentsInSlotOrder;
 	//	WrappedComponentsInSlotOrder.AddZeroed(Slots.Num());
 
 	//	for (int32 ComponentIndex = 0; ComponentIndex < AttachChildren.Num(); ++ComponentIndex)
 	//	{
-	//		if (USlateWrapperComponent* Wrapper = Cast<USlateWrapperComponent>(AttachChildren[ComponentIndex]))
+	//		if (UWidget* Wrapper = Cast<UWidget>(AttachChildren[ComponentIndex]))
 	//		{
 	//			if (Wrapper->IsRegistered())
 	//			{
diff --git a/Engine/Source/Runtime/UMG/Private/Components/HorizontalBoxComponent.cpp b/Engine/Source/Runtime/UMG/Private/Components/HorizontalBoxComponent.cpp
--- a/Engine/Source/Runtime/UMG/Private/Components/HorizontalBoxComponent.cpp
+++ b/Engine/Source/Runtime/UMG/Private/Components/HorizontalBoxComponent.cpp
@@ -16,18 +16,18 @@ int32 UHorizontalBoxComponent::GetChildrenCount() const
 	return Slots.Num();
 }
 
-USlateWrapperComponent* UHorizontalBoxComponent::GetChildAt(int32 Index) const
+UWidget* UHorizontalBoxComponent::GetChildAt(int32 Index) const
 {
 	return Slots[Index]->Content;
 }
 
-bool UHorizontalBoxComponent::AddChild(USlateWrapperComponent* Child, FVector2D Position)
+bool UHorizontalBoxComponent::AddChild(UWidget* Child, FVector2D Position)
 {
 	AddSlot(Child);
 	return true;
 }
 
-bool UHorizontalBoxComponent::RemoveChild(USlateWrapperComponent* Child)
+bool UHorizontalBoxComponent::RemoveChild(UWidget* Child)
 {
 	for ( int32 SlotIndex = 0; SlotIndex < Slots.Num(); ++SlotIndex )
 	{
@@ -43,7 +43,7 @@ bool UHorizontalBoxComponent::RemoveChild(USlateWrapperComponent* Child)
 	return false;
 }
 
-void UHorizontalBoxComponent::ReplaceChildAt(int32 Index, USlateWrapperComponent* Content)
+void UHorizontalBoxComponent::ReplaceChildAt(int32 Index, UWidget* Content)
 {
 	UHorizontalBoxSlot* Slot = Slots[Index];
 	Slot->Content = Content;
@@ -79,14 +79,14 @@ TSharedRef<SWidget> UHorizontalBoxComponent::RebuildWidget()
 					Slot->Content == NULL ? SNullWidget::NullWidget : Slot->Content->GetWidget()
 				];
 
-			NewSlot.SizeParam = USlateWrapperComponent::ConvertSerializedSizeParamToRuntime(Slot->Size);
+			NewSlot.SizeParam = UWidget::ConvertSerializedSizeParamToRuntime(Slot->Size);
 		}
 	}
 
 	return NewCanvas;
 }
 
-UHorizontalBoxSlot* UHorizontalBoxComponent::AddSlot(USlateWrapperComponent* Content)
+UHorizontalBoxSlot* UHorizontalBoxComponent::AddSlot(UWidget* Content)
 {
 	UHorizontalBoxSlot* Slot = ConstructObject<UHorizontalBoxSlot>(UHorizontalBoxSlot::StaticClass(), this);
 	Slot->Content = Content;
diff --git a/Engine/Source/Runtime/UMG/Private/Components/PanelWidget.cpp b/Engine/Source/Runtime/UMG/Private/Components/PanelWidget.cpp
new file mode 100644
--- /dev/null
+++ b/Engine/Source/Runtime/UMG/Private/Components/PanelWidget.cpp
@@ -0,0 +1,26 @@
+// Copyright 1998-2014 Epic Games, Inc. All Rights Reserved.
+
+#include "UMGPrivatePCH.h"
+
+#if WITH_EDITOR
+#include "MessageLog.h"
+#include "UObjectToken.h"
+#endif
+
+#define LOCTEXT_NAMESPACE "UMG"
+
+/////////////////////////////////////////////////////
+// UPanelWidget
+
+UPanelWidget::UPanelWidget(const FPostConstructInitializeProperties& PCIP)
+	: Super(PCIP)
+{
+//	PrimaryComponentTick.bCanEverTick = true;
+
+//	bAutoActivate = true;
+//	bTickInEditor = true;
+}
+
+/////////////////////////////////////////////////////
+
+#undef LOCTEXT_NAMESPACE
diff --git a/Engine/Source/Runtime/UMG/Private/Components/SlateLeafWidgetComponent.cpp b/Engine/Source/Runtime/UMG/Private/Components/SlateLeafWidgetComponent.cpp
deleted file mode 100644
diff --git a/Engine/Source/Runtime/UMG/Private/Components/SlateNonLeafWidgetComponent.cpp b/Engine/Source/Runtime/UMG/Private/Components/SlateNonLeafWidgetComponent.cpp
deleted file mode 100644
diff --git a/Engine/Source/Runtime/UMG/Private/Components/SlateWrapperComponent.cpp b/Engine/Source/Runtime/UMG/Private/Components/SlateWrapperComponent.cpp
deleted file mode 100644
diff --git a/Engine/Source/Runtime/UMG/Private/Components/UniformGridPanel.cpp b/Engine/Source/Runtime/UMG/Private/Components/UniformGridPanel.cpp
--- a/Engine/Source/Runtime/UMG/Private/Components/UniformGridPanel.cpp
+++ b/Engine/Source/Runtime/UMG/Private/Components/UniformGridPanel.cpp
@@ -16,18 +16,18 @@ int32 UUniformGridPanel::GetChildrenCount() const
 	return Slots.Num();
 }
 
-USlateWrapperComponent* UUniformGridPanel::GetChildAt(int32 Index) const
+UWidget* UUniformGridPanel::GetChildAt(int32 Index) const
 {
 	return Slots[Index]->Content;
 }
 
-bool UUniformGridPanel::AddChild(USlateWrapperComponent* Child, FVector2D Position)
+bool UUniformGridPanel::AddChild(UWidget* Child, FVector2D Position)
 {
 	AddSlot(Child);
 	return true;
 }
 
-bool UUniformGridPanel::RemoveChild(USlateWrapperComponent* Child)
+bool UUniformGridPanel::RemoveChild(UWidget* Child)
 {
 	for ( int32 SlotIndex = 0; SlotIndex < Slots.Num(); ++SlotIndex )
 	{
@@ -43,7 +43,7 @@ bool UUniformGridPanel::RemoveChild(USlateWrapperComponent* Child)
 	return false;
 }
 
-void UUniformGridPanel::ReplaceChildAt(int32 Index, USlateWrapperComponent* Content)
+void UUniformGridPanel::ReplaceChildAt(int32 Index, UWidget* Content)
 {
 	UUniformGridSlot* Slot = Slots[Index];
 	Slot->Content = Content;
@@ -77,7 +77,7 @@ TSharedRef<SWidget> UUniformGridPanel::RebuildWidget()
 	return NewPanel;
 }
 
-UUniformGridSlot* UUniformGridPanel::AddSlot(USlateWrapperComponent* Content)
+UUniformGridSlot* UUniformGridPanel::AddSlot(UWidget* Content)
 {
 	UUniformGridSlot* Slot = ConstructObject<UUniformGridSlot>(UUniformGridSlot::StaticClass(), this);
 	Slot->Content = Content;
diff --git a/Engine/Source/Runtime/UMG/Private/Components/VerticalBoxComponent.cpp b/Engine/Source/Runtime/UMG/Private/Components/VerticalBoxComponent.cpp
--- a/Engine/Source/Runtime/UMG/Private/Components/VerticalBoxComponent.cpp
+++ b/Engine/Source/Runtime/UMG/Private/Components/VerticalBoxComponent.cpp
@@ -16,18 +16,18 @@ int32 UVerticalBoxComponent::GetChildrenCount() const
 	return Slots.Num();
 }
 
-USlateWrapperComponent* UVerticalBoxComponent::GetChildAt(int32 Index) const
+UWidget* UVerticalBoxComponent::GetChildAt(int32 Index) const
 {
 	return Slots[Index]->Content;
 }
 
-bool UVerticalBoxComponent::AddChild(USlateWrapperComponent* Child, FVector2D Position)
+bool UVerticalBoxComponent::AddChild(UWidget* Child, FVector2D Position)
 {
 	AddSlot(Child);
 	return true;
 }
 
-bool UVerticalBoxComponent::RemoveChild(USlateWrapperComponent* Child)
+bool UVerticalBoxComponent::RemoveChild(UWidget* Child)
 {
 	for ( int32 SlotIndex = 0; SlotIndex < Slots.Num(); ++SlotIndex )
 	{
@@ -43,7 +43,7 @@ bool UVerticalBoxComponent::RemoveChild(USlateWrapperComponent* Child)
 	return false;
 }
 
-void UVerticalBoxComponent::ReplaceChildAt(int32 Index, USlateWrapperComponent* Content)
+void UVerticalBoxComponent::ReplaceChildAt(int32 Index, UWidget* Content)
 {
 	UVerticalBoxSlot* Slot = Slots[Index];
 	Slot->Content = Content;
@@ -79,14 +79,14 @@ TSharedRef<SWidget> UVerticalBoxComponent::RebuildWidget()
 					Slot->Content == NULL ? SNullWidget::NullWidget : Slot->Content->GetWidget()
 				];
 
-			NewSlot.SizeParam = USlateWrapperComponent::ConvertSerializedSizeParamToRuntime(Slot->Size);
+			NewSlot.SizeParam = UWidget::ConvertSerializedSizeParamToRuntime(Slot->Size);
 		}
 	}
 
 	return NewCanvas;
 }
 
-UVerticalBoxSlot* UVerticalBoxComponent::AddSlot(USlateWrapperComponent* Content)
+UVerticalBoxSlot* UVerticalBoxComponent::AddSlot(UWidget* Content)
 {
 	UVerticalBoxSlot* Slot = ConstructObject<UVerticalBoxSlot>(UVerticalBoxSlot::StaticClass(), this);
 	Slot->Content = Content;
diff --git a/Engine/Source/Runtime/UMG/Private/Components/Widget.cpp b/Engine/Source/Runtime/UMG/Private/Components/Widget.cpp
new file mode 100644
--- /dev/null
+++ b/Engine/Source/Runtime/UMG/Private/Components/Widget.cpp
@@ -0,0 +1,138 @@
+// Copyright 1998-2014 Epic Games, Inc. All Rights Reserved.
+
+#include "UMGPrivatePCH.h"
+
+/////////////////////////////////////////////////////
+// UWidget
+
+UWidget::UWidget(const FPostConstructInitializeProperties& PCIP)
+	: Super(PCIP)
+{
+	bIsEnabled = true;
+	bIsVariable = true;
+	Visiblity = ESlateVisibility::Visible;
+}
+
+bool UWidget::GetIsEnabled() const
+{
+	return MyWidget.IsValid() ? MyWidget->IsEnabled() : bIsEnabled;
+}
+
+void UWidget::SetIsEnabled(bool bInIsEnabled)
+{
+	bIsEnabled = bInIsEnabled;
+	if ( MyWidget.IsValid() )
+	{
+		MyWidget->SetEnabled(bInIsEnabled);
+	}
+}
+
+TEnumAsByte<ESlateVisibility::Type> UWidget::GetVisibility()
+{
+	return UWidget::ConvertRuntimeToSerializedVisiblity(MyWidget->GetVisibility());
+}
+
+void UWidget::SetVisibility(TEnumAsByte<ESlateVisibility::Type> InVisibility)
+{
+	MyWidget->SetVisibility(UWidget::ConvertSerializedVisibilityToRuntime(InVisibility));
+}
+
+void UWidget::SetToolTipText(const FText& InToolTipText)
+{
+	MyWidget->SetToolTipText(ToolTipText);
+}
+
+void UWidget::OnUnregister()
+{
+	MyWidget = NULL;
+//	Super::OnUnregister();
+}
+
+void UWidget::PostEditChangeProperty(struct FPropertyChangedEvent& PropertyChangedEvent)
+{
+	// TODO Rename?
+}
+
+TSharedRef<SWidget> UWidget::GetWidget()
+{
+	if ( !MyWidget.IsValid() )
+	{
+		MyWidget = RebuildWidget();
+		MyWidget->SetEnabled(OPTIONAL_BINDING(bool, bIsEnabled));
+		MyWidget->SetVisibility(OPTIONAL_BINDING_CONVERT(ESlateVisibility::Type, Visiblity, EVisibility, ConvertVisibility));
+
+		if ( !ToolTipText.IsEmpty() )
+		{
+			MyWidget->SetToolTipText(OPTIONAL_BINDING(FText, ToolTipText));
+		}
+	}
+
+	return MyWidget.ToSharedRef();
+}
+
+TSharedRef<SWidget> UWidget::RebuildWidget()
+{
+	ensureMsg(false, TEXT("You must implement RebuildWidget() in your child class"));
+	return SNew(SSpacer);
+}
+
+EVisibility UWidget::ConvertSerializedVisibilityToRuntime(TEnumAsByte<ESlateVisibility::Type> Input)
+{
+	switch ( Input )
+	{
+	case ESlateVisibility::Visible:
+		return EVisibility::Visible;
+	case ESlateVisibility::Collapsed:
+		return EVisibility::Collapsed;
+	case ESlateVisibility::Hidden:
+		return EVisibility::Hidden;
+	case ESlateVisibility::HitTestInvisible:
+		return EVisibility::HitTestInvisible;
+	case ESlateVisibility::SelfHitTestInvisible:
+		return EVisibility::SelfHitTestInvisible;
+	default:
+		check(false);
+		return EVisibility::Visible;
+	}
+}
+
+TEnumAsByte<ESlateVisibility::Type> UWidget::ConvertRuntimeToSerializedVisiblity(const EVisibility& Input)
+{
+	if ( Input == EVisibility::Visible )
+	{
+		return ESlateVisibility::Visible;
+	}
+	else if ( Input == EVisibility::Collapsed )
+	{
+		return ESlateVisibility::Collapsed;
+	}
+	else if ( Input == EVisibility::Hidden )
+	{
+		return ESlateVisibility::Hidden;
+	}
+	else if ( Input == EVisibility::HitTestInvisible )
+	{
+		return ESlateVisibility::HitTestInvisible;
+	}
+	else if ( Input == EVisibility::SelfHitTestInvisible )
+	{
+		return ESlateVisibility::SelfHitTestInvisible;
+	}
+	else
+	{
+		check(false);
+		return ESlateVisibility::Visible;
+	}
+}
+
+FSizeParam UWidget::ConvertSerializedSizeParamToRuntime(const FSlateChildSize& Input)
+{
+	switch ( Input.SizeRule )
+	{
+	default:
+	case ESlateSizeRule::Automatic:
+		return FAuto();
+	case ESlateSizeRule::Fill:
+		return FStretch(Input.Value);
+	}
+}
\ No newline at end of file
diff --git a/Engine/Source/Runtime/UMG/Private/UserWidget.cpp b/Engine/Source/Runtime/UMG/Private/UserWidget.cpp
--- a/Engine/Source/Runtime/UMG/Private/UserWidget.cpp
+++ b/Engine/Source/Runtime/UMG/Private/UserWidget.cpp
@@ -152,7 +152,7 @@ UWorld* UUserWidget::GetWorld() const
 	return NULL;
 }
 
-USlateWrapperComponent* UUserWidget::GetWidgetHandle(TSharedRef<SWidget> InWidget)
+UWidget* UUserWidget::GetWidgetHandle(TSharedRef<SWidget> InWidget)
 {
 	return WidgetToComponent.FindRef(InWidget);
 }
@@ -174,7 +174,7 @@ void UUserWidget::RebuildWrapperWidget()
 	// Place all of our top-level children Slate wrapped components into the overlay
 	for ( int32 ComponentIndex = 0; ComponentIndex < Components.Num(); ++ComponentIndex )
 	{
-		USlateWrapperComponent* Handle = Components[ComponentIndex];
+		UWidget* Handle = Components[ComponentIndex];
 		TSharedRef<SWidget> Widget = Handle->GetWidget();
 
 		WidgetToComponent.Add(Widget, Handle);
@@ -208,7 +208,7 @@ TSharedPtr<SWidget> UUserWidget::GetWidgetFromName(const FString& Name) const
 	return TSharedPtr<SWidget>();
 }
 
-USlateWrapperComponent* UUserWidget::GetHandleFromName(const FString& Name) const
+UWidget* UUserWidget::GetHandleFromName(const FString& Name) const
 {
 	for ( auto& Entry : WidgetToComponent )
 	{
@@ -256,13 +256,13 @@ TSharedRef<SWidget> UUserWidget::MakeFullScreenWidget()
 				MakeWidget()
 			];
 
-		NewSlot.SizeParam = USlateWrapperComponent::ConvertSerializedSizeParamToRuntime(Size);
+		NewSlot.SizeParam = UWidget::ConvertSerializedSizeParamToRuntime(Size);
 
 		return VerticalBox;
 	}
 }
 
-USlateWrapperComponent* UUserWidget::GetRootWidgetComponent()
+UWidget* UUserWidget::GetRootWidgetComponent()
 {
 	if ( Components.Num() > 0 )
 	{
@@ -344,7 +344,7 @@ TEnumAsByte<ESlateVisibility::Type> UUserWidget::GetVisiblity()
 	{
 		TSharedPtr<SWidget> RootWidget = FullScreenWidget.Pin();
 
-		return USlateWrapperComponent::ConvertRuntimeToSerializedVisiblity(RootWidget->GetVisibility());
+		return UWidget::ConvertRuntimeToSerializedVisiblity(RootWidget->GetVisibility());
 	}
 
 	return ESlateVisibility::Collapsed;
diff --git a/Engine/Source/Runtime/UMG/Private/WidgetBlueprint.cpp b/Engine/Source/Runtime/UMG/Private/WidgetBlueprint.cpp
--- a/Engine/Source/Runtime/UMG/Private/WidgetBlueprint.cpp
+++ b/Engine/Source/Runtime/UMG/Private/WidgetBlueprint.cpp
@@ -28,7 +28,7 @@ void UWidgetBlueprint::PostLoad()
 	Super::PostLoad();
 
 #if WITH_EDITOR
-	for ( USlateWrapperComponent* Widget : WidgetTree->WidgetTemplates )
+	for ( UWidget* Widget : WidgetTree->WidgetTemplates )
 	{
 		Widget->ConnectEditorData();
 	}
diff --git a/Engine/Source/Runtime/UMG/Private/WidgetBlueprintGeneratedClass.cpp b/Engine/Source/Runtime/UMG/Private/WidgetBlueprintGeneratedClass.cpp
--- a/Engine/Source/Runtime/UMG/Private/WidgetBlueprintGeneratedClass.cpp
+++ b/Engine/Source/Runtime/UMG/Private/WidgetBlueprintGeneratedClass.cpp
@@ -46,7 +46,7 @@ void UWidgetBlueprintGeneratedClass::InitializeWidget(UUserWidget* Actor) const
 	UUserWidget* WidgetActor = CastChecked<UUserWidget>(Actor);
 	UClass* ActorClass = Actor->GetClass();
 
-	for ( USlateWrapperComponent* Widget : ClonedTree->WidgetTemplates )
+	for ( UWidget* Widget : ClonedTree->WidgetTemplates )
 	{
 		// Not fatal if NULL, but shouldn't happen
 		if ( !ensure(Widget != NULL) )
diff --git a/Engine/Source/Runtime/UMG/Private/WidgetTree.cpp b/Engine/Source/Runtime/UMG/Private/WidgetTree.cpp
--- a/Engine/Source/Runtime/UMG/Private/WidgetTree.cpp
+++ b/Engine/Source/Runtime/UMG/Private/WidgetTree.cpp
@@ -10,17 +10,17 @@ UWidgetTree::UWidgetTree(const FPostConstructInitializeProperties& PCIP)
 {
 }
 
-void UWidgetTree::RenameWidget(USlateWrapperComponent* Widget, FString& NewName)
+void UWidgetTree::RenameWidget(UWidget* Widget, FString& NewName)
 {
 	Widget->Rename(*NewName);
 
 	// TODO Update nodes in the blueprint!
 }
 
-USlateWrapperComponent* UWidgetTree::FindWidget(FString& Name) const
+UWidget* UWidgetTree::FindWidget(FString& Name) const
 {
 	FString ExistingName;
-	for ( USlateWrapperComponent* Widget : WidgetTemplates )
+	for ( UWidget* Widget : WidgetTemplates )
 	{
 		Widget->GetName(ExistingName);
 		if ( ExistingName.Equals(Name, ESearchCase::IgnoreCase) )
@@ -32,11 +32,11 @@ USlateWrapperComponent* UWidgetTree::FindWidget(FString& Name) const
 	return NULL;
 }
 
-USlateNonLeafWidgetComponent* UWidgetTree::FindWidgetParent(USlateWrapperComponent* Widget, int32& OutChildIndex)
+UPanelWidget* UWidgetTree::FindWidgetParent(UWidget* Widget, int32& OutChildIndex)
 {
-	for ( USlateWrapperComponent* Template : WidgetTemplates )
+	for ( UWidget* Template : WidgetTemplates )
 	{
-		USlateNonLeafWidgetComponent* NonLeafTemplate = Cast<USlateNonLeafWidgetComponent>(Template);
+		UPanelWidget* NonLeafTemplate = Cast<UPanelWidget>(Template);
 		if ( NonLeafTemplate )
 		{
 			for ( int32 ChildIndex = 0; ChildIndex < NonLeafTemplate->GetChildrenCount(); ChildIndex++ )
@@ -54,15 +54,15 @@ USlateNonLeafWidgetComponent* UWidgetTree::FindWidgetParent(USlateWrapperCompone
 	return NULL;
 }
 
-bool UWidgetTree::RemoveWidget(USlateWrapperComponent* InRemovedWidget)
+bool UWidgetTree::RemoveWidget(UWidget* InRemovedWidget)
 {
-	USlateWrapperComponent* Parent = NULL;
+	UWidget* Parent = NULL;
 
 	bool bRemoved = false;
 
 	//TODO UMG Make the Widget Tree actually a tree, instead of a list, it makes things like removal difficult.
 
-	if ( USlateNonLeafWidgetComponent* InNonLeafRemovedWidget = Cast<USlateNonLeafWidgetComponent>(InRemovedWidget) )
+	if ( UPanelWidget* InNonLeafRemovedWidget = Cast<UPanelWidget>(InRemovedWidget) )
 	{
 		while ( InNonLeafRemovedWidget->GetChildrenCount() > 0 )
 		{
@@ -70,9 +70,9 @@ bool UWidgetTree::RemoveWidget(USlateWrapperComponent* InRemovedWidget)
 		}
 	}
 
-	for ( USlateWrapperComponent* Template : WidgetTemplates )
+	for ( UWidget* Template : WidgetTemplates )
 	{
-		USlateNonLeafWidgetComponent* NonLeafTemplate = Cast<USlateNonLeafWidgetComponent>(Template);
+		UPanelWidget* NonLeafTemplate = Cast<UPanelWidget>(Template);
 		if ( NonLeafTemplate )
 		{
 			if ( NonLeafTemplate->RemoveChild(InRemovedWidget) )
diff --git a/Engine/Source/Runtime/UMG/Public/Blueprint/UserWidget.h b/Engine/Source/Runtime/UMG/Public/Blueprint/UserWidget.h
--- a/Engine/Source/Runtime/UMG/Public/Blueprint/UserWidget.h
+++ b/Engine/Source/Runtime/UMG/Public/Blueprint/UserWidget.h
@@ -80,7 +80,7 @@ class UMG_API UUserWidget : public UObject
 	TEnumAsByte<EVerticalAlignment> VerticalAlignment;
 
 	UPROPERTY()
-	TArray<USlateWrapperComponent*> Components;
+	TArray<UWidget*> Components;
 	
 	//UObject interface
 	virtual void PostInitProperties() OVERRIDE;
@@ -166,19 +166,19 @@ class UMG_API UUserWidget : public UObject
 	FSReply OnMotionDetected(FGeometry MyGeometry, FMotionEvent InMotionEvent);
 
 
-	USlateWrapperComponent* GetWidgetHandle(TSharedRef<SWidget> InWidget);
+	UWidget* GetWidgetHandle(TSharedRef<SWidget> InWidget);
 
 	TSharedRef<SWidget> MakeWidget();
 	TSharedRef<SWidget> MakeFullScreenWidget();
 
-	USlateWrapperComponent* GetRootWidgetComponent();
+	UWidget* GetRootWidgetComponent();
 
 	TSharedPtr<SWidget> GetWidgetFromName(const FString& Name) const;
-	USlateWrapperComponent* GetHandleFromName(const FString& Name) const;
+	UWidget* GetHandleFromName(const FString& Name) const;
 
 private:
 	TSharedPtr<SWidget> UserRootWidget;
-	TMap< TWeakPtr<SWidget>, USlateWrapperComponent* > WidgetToComponent;
+	TMap< TWeakPtr<SWidget>, UWidget* > WidgetToComponent;
 
 	TWeakPtr<SWidget> FullScreenWidget;
 
diff --git a/Engine/Source/Runtime/UMG/Public/Blueprint/WidgetTree.h b/Engine/Source/Runtime/UMG/Public/Blueprint/WidgetTree.h
--- a/Engine/Source/Runtime/UMG/Public/Blueprint/WidgetTree.h
+++ b/Engine/Source/Runtime/UMG/Public/Blueprint/WidgetTree.h
@@ -11,20 +11,20 @@ class UMG_API UWidgetTree : public UObject
 
 	/** Array of templates for widgets */
 	UPROPERTY()
-	TArray<class USlateWrapperComponent*> WidgetTemplates;
+	TArray<class UWidget*> WidgetTemplates;
 
-	void RenameWidget(USlateWrapperComponent* Widget, FString& NewName);
+	void RenameWidget(UWidget* Widget, FString& NewName);
 
-	class USlateWrapperComponent* FindWidget(FString& Name) const;
+	class UWidget* FindWidget(FString& Name) const;
 
-	bool RemoveWidget(class USlateWrapperComponent* Widget);
+	bool RemoveWidget(class UWidget* Widget);
 
-	class USlateNonLeafWidgetComponent* FindWidgetParent(class USlateWrapperComponent* Widget, int32& OutChildIndex);
+	class UPanelWidget* FindWidgetParent(class UWidget* Widget, int32& OutChildIndex);
 
 	template< class T >
-	T* ConstructWidget(TSubclassOf<class USlateWrapperComponent> WidgetType)
+	T* ConstructWidget(TSubclassOf<class UWidget> WidgetType)
 	{
-		USlateWrapperComponent* Widget = (USlateWrapperComponent*)ConstructObject<class USlateWrapperComponent>(WidgetType, this);
+		UWidget* Widget = (UWidget*)ConstructObject<class UWidget>(WidgetType, this);
 		WidgetTemplates.Add(Widget);
 
 		return (T*)Widget;
diff --git a/Engine/Source/Runtime/UMG/Public/Components/BorderComponent.h b/Engine/Source/Runtime/UMG/Public/Components/BorderComponent.h
--- a/Engine/Source/Runtime/UMG/Public/Components/BorderComponent.h
+++ b/Engine/Source/Runtime/UMG/Public/Components/BorderComponent.h
@@ -58,16 +58,16 @@ class UMG_API UBorderComponent : public UContentWidget
 	bool bShowEffectWhenDisabled;
 
 	// UContentWidget interface
-	virtual void SetContent(USlateWrapperComponent* Content) OVERRIDE;
+	virtual void SetContent(UWidget* Content) OVERRIDE;
 	// End UContentWidget interface
 
 protected:
 	TWeakPtr<class SBorder> MyBorder;
 
 protected:
-	// USlateWrapperComponent interface
+	// UWidget interface
 	virtual TSharedRef<SWidget> RebuildWidget() OVERRIDE;
-	// End of USlateWrapperComponent interface
+	// End of UWidget interface
 
 	FMargin GetContentPadding() const;
 	FLinearColor GetContentColor() const;
diff --git a/Engine/Source/Runtime/UMG/Public/Components/BoxPanelComponent.h b/Engine/Source/Runtime/UMG/Public/Components/BoxPanelComponent.h
--- a/Engine/Source/Runtime/UMG/Public/Components/BoxPanelComponent.h
+++ b/Engine/Source/Runtime/UMG/Public/Components/BoxPanelComponent.h
@@ -51,7 +51,7 @@ struct UMG_API FBoxPanelSlot
 
 
 UCLASS(meta=(BlueprintSpawnableComponent), ClassGroup=UserInterface)
-class UMG_API UBoxPanelComponent : public USlateNonLeafWidgetComponent
+class UMG_API UBoxPanelComponent : public UPanelWidget
 {
 	GENERATED_UCLASS_BODY()
 
@@ -71,7 +71,7 @@ protected:
 	TWeakPtr<class SConcreteBoxPanel> MyBox;
 
 protected:
-	// USlateWrapperComponent interface
+	// UWidget interface
 	virtual TSharedRef<SWidget> RebuildWidget() OVERRIDE;
-	// End of USlateWrapperComponent interface
+	// End of UWidget interface
 };
diff --git a/Engine/Source/Runtime/UMG/Public/Components/ButtonComponent.h b/Engine/Source/Runtime/UMG/Public/Components/ButtonComponent.h
--- a/Engine/Source/Runtime/UMG/Public/Components/ButtonComponent.h
+++ b/Engine/Source/Runtime/UMG/Public/Components/ButtonComponent.h
@@ -81,7 +81,7 @@ public:
 
 
 	// UContentWidget interface
-	virtual void SetContent(USlateWrapperComponent* Content) OVERRIDE;
+	virtual void SetContent(UWidget* Content) OVERRIDE;
 	// End UContentWidget interface
 
 protected:
@@ -96,9 +96,9 @@ protected:
 	}
 
 protected:
-	// USlateWrapperComponent interface
+	// UWidget interface
 	virtual TSharedRef<SWidget> RebuildWidget() OVERRIDE;
-	// End of USlateWrapperComponent interface
+	// End of UWidget interface
 
 	FMargin GetContentPadding() const;
 
diff --git a/Engine/Source/Runtime/UMG/Public/Components/CanvasPanelComponent.h b/Engine/Source/Runtime/UMG/Public/Components/CanvasPanelComponent.h
--- a/Engine/Source/Runtime/UMG/Public/Components/CanvasPanelComponent.h
+++ b/Engine/Source/Runtime/UMG/Public/Components/CanvasPanelComponent.h
@@ -5,7 +5,7 @@
 #include "CanvasPanelComponent.generated.h"
 
 UCLASS(meta=(BlueprintSpawnableComponent), ClassGroup=UserInterface)
-class UMG_API UCanvasPanelComponent : public USlateNonLeafWidgetComponent
+class UMG_API UCanvasPanelComponent : public UPanelWidget
 {
 	GENERATED_UCLASS_BODY()
 
@@ -17,22 +17,22 @@ class UMG_API UCanvasPanelComponent : public USlateNonLeafWidgetComponent
 	UPROPERTY(EditAnywhere, Category=Appearance)
 	FVector2D DesiredCanvasSize;
 
-	UCanvasPanelSlot* AddSlot(USlateWrapperComponent* Content);
-	virtual bool AddChild(USlateWrapperComponent* Child, FVector2D Position) OVERRIDE;
+	UCanvasPanelSlot* AddSlot(UWidget* Content);
+	virtual bool AddChild(UWidget* Child, FVector2D Position) OVERRIDE;
 
-	// USlateNonLeafWidgetComponent
+	// UPanelWidget
 	virtual int32 GetChildrenCount() const OVERRIDE;
-	virtual USlateWrapperComponent* GetChildAt(int32 Index) const OVERRIDE;
-	// End USlateNonLeafWidgetComponent
+	virtual UWidget* GetChildAt(int32 Index) const OVERRIDE;
+	// End UPanelWidget
 
 #if WITH_EDITOR
 	// UObject interface
 	virtual void PostEditChangeProperty(struct FPropertyChangedEvent& PropertyChangedEvent) OVERRIDE;
 	// End of UObject interface
 
-	// USlateWrapperComponent interface
+	// UWidget interface
 	virtual void ConnectEditorData() OVERRIDE;
-	// End USlateWrapperComponent interface
+	// End UWidget interface
 #endif
 
 protected:
@@ -40,7 +40,7 @@ protected:
 	TWeakPtr<class SFixedSizeCanvas> MyCanvas;
 
 protected:
-	// USlateWrapperComponent interface
+	// UWidget interface
 	virtual TSharedRef<SWidget> RebuildWidget() OVERRIDE;
-	// End of USlateWrapperComponent interface
+	// End of UWidget interface
 };
diff --git a/Engine/Source/Runtime/UMG/Public/Components/CheckBoxComponent.h b/Engine/Source/Runtime/UMG/Public/Components/CheckBoxComponent.h
--- a/Engine/Source/Runtime/UMG/Public/Components/CheckBoxComponent.h
+++ b/Engine/Source/Runtime/UMG/Public/Components/CheckBoxComponent.h
@@ -8,7 +8,7 @@ DECLARE_DYNAMIC_MULTICAST_DELEGATE_OneParam( FOnCheckBoxComponentStateChanged, b
 
 /** Check box widget */
 UCLASS(meta=(BlueprintSpawnableComponent), ClassGroup=UserInterface)
-class UMG_API UCheckBoxComponent : public USlateWrapperComponent
+class UMG_API UCheckBoxComponent : public UWidget
 {
 	GENERATED_UCLASS_BODY()
 
@@ -56,9 +56,9 @@ protected:
 	FSlateSound HoveredSound;
 
 protected:
-	// USlateWrapperComponent interface
+	// UWidget interface
 	virtual TSharedRef<SWidget> RebuildWidget() OVERRIDE;
-	// End of USlateWrapperComponent interface
+	// End of UWidget interface
 
 	ESlateCheckBoxState::Type GetCheckState() const;
 	void SlateOnCheckStateChangedCallback(ESlateCheckBoxState::Type NewState);
diff --git a/Engine/Source/Runtime/UMG/Public/Components/CircularThrobberComponent.h b/Engine/Source/Runtime/UMG/Public/Components/CircularThrobberComponent.h
--- a/Engine/Source/Runtime/UMG/Public/Components/CircularThrobberComponent.h
+++ b/Engine/Source/Runtime/UMG/Public/Components/CircularThrobberComponent.h
@@ -29,7 +29,7 @@ protected:
 	USlateBrushAsset* PieceImage;
 
 protected:
-	// USlateWrapperComponent interface
+	// UWidget interface
 	virtual TSharedRef<SWidget> RebuildWidget() OVERRIDE;
-	// End of USlateWrapperComponent interface
+	// End of UWidget interface
 };
diff --git a/Engine/Source/Runtime/UMG/Public/Components/ContentWidget.h b/Engine/Source/Runtime/UMG/Public/Components/ContentWidget.h
--- a/Engine/Source/Runtime/UMG/Public/Components/ContentWidget.h
+++ b/Engine/Source/Runtime/UMG/Public/Components/ContentWidget.h
@@ -5,26 +5,26 @@
 #include "ContentWidget.generated.h"
 
 UCLASS(Abstract)
-class UMG_API UContentWidget : public USlateNonLeafWidgetComponent
+class UMG_API UContentWidget : public UPanelWidget
 {
 	GENERATED_UCLASS_BODY()
 
 	UFUNCTION()
-	virtual USlateWrapperComponent* GetContent();
+	virtual UWidget* GetContent();
 	
 	UFUNCTION()
-	virtual void SetContent(USlateWrapperComponent* InContent);
+	virtual void SetContent(UWidget* InContent);
 
-	// USlateNonLeafWidgetComponent
+	// UPanelWidget
 	virtual int32 GetChildrenCount() const OVERRIDE;
-	virtual USlateWrapperComponent* GetChildAt(int32 Index) const OVERRIDE;
-	virtual bool AddChild(USlateWrapperComponent* Child, FVector2D Position) OVERRIDE;
-	virtual bool RemoveChild(USlateWrapperComponent* Child) OVERRIDE;
-	virtual void ReplaceChildAt(int32 Index, USlateWrapperComponent* Child) OVERRIDE;
-	// End USlateNonLeafWidgetComponent
+	virtual UWidget* GetChildAt(int32 Index) const OVERRIDE;
+	virtual bool AddChild(UWidget* Child, FVector2D Position) OVERRIDE;
+	virtual bool RemoveChild(UWidget* Child) OVERRIDE;
+	virtual void ReplaceChildAt(int32 Index, UWidget* Child) OVERRIDE;
+	// End UPanelWidget
 
 protected:
 
 	UPROPERTY()
-	USlateWrapperComponent* Content;
+	UWidget* Content;
 };
diff --git a/Engine/Source/Runtime/UMG/Public/Components/EditableText.h b/Engine/Source/Runtime/UMG/Public/Components/EditableText.h
--- a/Engine/Source/Runtime/UMG/Public/Components/EditableText.h
+++ b/Engine/Source/Runtime/UMG/Public/Components/EditableText.h
@@ -9,7 +9,7 @@ DECLARE_DYNAMIC_MULTICAST_DELEGATE_TwoParams(FOnEditableTextCommittedEvent, cons
 
 /** Editable text box widget */
 UCLASS(meta=(BlueprintSpawnableComponent), ClassGroup=UserInterface)
-class UMG_API UEditableText : public USlateWrapperComponent
+class UMG_API UEditableText : public UWidget
 {
 	GENERATED_UCLASS_BODY()
 
@@ -88,9 +88,9 @@ protected:
 	FOnEditableTextCommittedEvent OnTextCommitted;
 
 protected:
-	// USlateWrapperComponent interface
+	// UWidget interface
 	virtual TSharedRef<SWidget> RebuildWidget() OVERRIDE;
-	// End of USlateWrapperComponent
+	// End of UWidget
 
 #if WITH_EDITOR
 	// UObject interface
diff --git a/Engine/Source/Runtime/UMG/Public/Components/EditableTextBlockComponent.h b/Engine/Source/Runtime/UMG/Public/Components/EditableTextBlockComponent.h
--- a/Engine/Source/Runtime/UMG/Public/Components/EditableTextBlockComponent.h
+++ b/Engine/Source/Runtime/UMG/Public/Components/EditableTextBlockComponent.h
@@ -9,7 +9,7 @@ DECLARE_DYNAMIC_MULTICAST_DELEGATE_OneParam(FOnEditableTextBlock_TextChangedEven
 
 /** Editable text box widget */
 UCLASS(meta=(BlueprintSpawnableComponent), ClassGroup=UserInterface)
-class UMG_API UEditableTextBlockComponent : public USlateWrapperComponent
+class UMG_API UEditableTextBlockComponent : public UWidget
 {
 	GENERATED_UCLASS_BODY()
 
@@ -87,9 +87,9 @@ protected:
 	//SLATE_ARGUMENT(TSharedPtr<class IErrorReportingWidget>, ErrorReporting)
 
 protected:
-	// USlateWrapperComponent interface
+	// UWidget interface
 	virtual TSharedRef<SWidget> RebuildWidget() OVERRIDE;
-	// End of USlateWrapperComponent
+	// End of UWidget
 
 	void SlateOnTextChanged(const FText& Text);
 };
diff --git a/Engine/Source/Runtime/UMG/Public/Components/GridPanelComponent.h b/Engine/Source/Runtime/UMG/Public/Components/GridPanelComponent.h
--- a/Engine/Source/Runtime/UMG/Public/Components/GridPanelComponent.h
+++ b/Engine/Source/Runtime/UMG/Public/Components/GridPanelComponent.h
@@ -62,7 +62,7 @@ struct UMG_API FGridPanelSlot
 
 
 UCLASS(meta=(BlueprintSpawnableComponent), ClassGroup=UserInterface)
-class UMG_API UGridPanelComponent : public USlateNonLeafWidgetComponent
+class UMG_API UGridPanelComponent : public UPanelWidget
 {
 	GENERATED_UCLASS_BODY()
 
@@ -81,9 +81,9 @@ protected:
 	TSet<FIntPoint> UniqueSlotCoordinates;
 
 protected:
-	// USlateWrapperComponent interface
+	// UWidget interface
 	virtual TSharedRef<SWidget> RebuildWidget() OVERRIDE;
-	// End of USlateWrapperComponent interface
+	// End of UWidget interface
 
 	void RebuildUniqueSlotTable();
 	void EnsureSlotIsUnique(FGridPanelSlot& SlotConfig);
diff --git a/Engine/Source/Runtime/UMG/Public/Components/HorizontalBoxComponent.h b/Engine/Source/Runtime/UMG/Public/Components/HorizontalBoxComponent.h
--- a/Engine/Source/Runtime/UMG/Public/Components/HorizontalBoxComponent.h
+++ b/Engine/Source/Runtime/UMG/Public/Components/HorizontalBoxComponent.h
@@ -5,7 +5,7 @@
 #include "HorizontalBoxComponent.generated.h"
 
 UCLASS(meta=(BlueprintSpawnableComponent), ClassGroup=UserInterface)
-class UMG_API UHorizontalBoxComponent : public USlateNonLeafWidgetComponent
+class UMG_API UHorizontalBoxComponent : public UPanelWidget
 {
 	GENERATED_UCLASS_BODY()
 
@@ -13,24 +13,24 @@ class UMG_API UHorizontalBoxComponent : public USlateNonLeafWidgetComponent
 	UPROPERTY(EditAnywhere, EditInline, Category=Slots)
 	TArray<UHorizontalBoxSlot*> Slots;
 
-	UHorizontalBoxSlot* AddSlot(USlateWrapperComponent* Content);
+	UHorizontalBoxSlot* AddSlot(UWidget* Content);
 
-	// USlateNonLeafWidgetComponent
+	// UPanelWidget
 	virtual int32 GetChildrenCount() const OVERRIDE;
-	virtual USlateWrapperComponent* GetChildAt(int32 Index) const OVERRIDE;
-	virtual bool AddChild(USlateWrapperComponent* Child, FVector2D Position) OVERRIDE;
-	virtual bool RemoveChild(USlateWrapperComponent* Child) OVERRIDE;
-	virtual void ReplaceChildAt(int32 Index, USlateWrapperComponent* Child) OVERRIDE;
-	// End USlateNonLeafWidgetComponent
+	virtual UWidget* GetChildAt(int32 Index) const OVERRIDE;
+	virtual bool AddChild(UWidget* Child, FVector2D Position) OVERRIDE;
+	virtual bool RemoveChild(UWidget* Child) OVERRIDE;
+	virtual void ReplaceChildAt(int32 Index, UWidget* Child) OVERRIDE;
+	// End UPanelWidget
 
 #if WITH_EDITOR
 	// UObject interface
 	virtual void PostEditChangeProperty(struct FPropertyChangedEvent& PropertyChangedEvent) OVERRIDE;
 	// End of UObject interface
 
-	// USlateWrapperComponent interface
+	// UWidget interface
 	virtual void ConnectEditorData() OVERRIDE;
-	// End USlateWrapperComponent interface
+	// End UWidget interface
 #endif
 
 protected:
@@ -38,7 +38,7 @@ protected:
 	TWeakPtr<class SHorizontalBox> MyHorizontalBox;
 
 protected:
-	// USlateWrapperComponent interface
+	// UWidget interface
 	virtual TSharedRef<SWidget> RebuildWidget() OVERRIDE;
-	// End of USlateWrapperComponent interface
+	// End of UWidget interface
 };
diff --git a/Engine/Source/Runtime/UMG/Public/Components/ImageComponent.h b/Engine/Source/Runtime/UMG/Public/Components/ImageComponent.h
--- a/Engine/Source/Runtime/UMG/Public/Components/ImageComponent.h
+++ b/Engine/Source/Runtime/UMG/Public/Components/ImageComponent.h
@@ -19,9 +19,9 @@ class UMG_API UImageComponent : public USlateLeafWidgetComponent
 	FSlateColor ColorAndOpacity;
 
 protected:
-	// USlateWrapperComponent interface
+	// UWidget interface
 	virtual TSharedRef<SWidget> RebuildWidget() OVERRIDE;
-	// End of USlateWrapperComponent interface
+	// End of UWidget interface
 
 	const FSlateBrush* GetImageBrush() const;
 	FSlateColor GetColorAndOpacity() const;
diff --git a/Engine/Source/Runtime/UMG/Public/Components/PanelSlot.h b/Engine/Source/Runtime/UMG/Public/Components/PanelSlot.h
--- a/Engine/Source/Runtime/UMG/Public/Components/PanelSlot.h
+++ b/Engine/Source/Runtime/UMG/Public/Components/PanelSlot.h
@@ -10,7 +10,7 @@ class UMG_API UPanelSlot : public UObject
 	GENERATED_UCLASS_BODY()
 
 	UPROPERTY()
-	class USlateWrapperComponent* Content;
+	class UWidget* Content;
 
 	virtual void Resize(const FVector2D& Direction, const FVector2D& Amount);
 
diff --git a/Engine/Source/Runtime/UMG/Public/Components/ScrollBoxComponent.h b/Engine/Source/Runtime/UMG/Public/Components/ScrollBoxComponent.h
--- a/Engine/Source/Runtime/UMG/Public/Components/ScrollBoxComponent.h
+++ b/Engine/Source/Runtime/UMG/Public/Components/ScrollBoxComponent.h
@@ -5,14 +5,14 @@
 #include "ScrollBoxComponent.generated.h"
 
 UCLASS(meta=(BlueprintSpawnableComponent), ClassGroup=UserInterface)
-class UMG_API UScrollBoxComponent : public USlateWrapperComponent
+class UMG_API UScrollBoxComponent : public UWidget
 {
 	GENERATED_UCLASS_BODY()
 
 protected:
 
 protected:
-	// USlateWrapperComponent interface
+	// UWidget interface
 	virtual TSharedRef<SWidget> RebuildWidget() OVERRIDE;
-	// End of USlateWrapperComponent interface
+	// End of UWidget interface
 };
diff --git a/Engine/Source/Runtime/UMG/Public/Components/SlateLeafWidgetComponent.h b/Engine/Source/Runtime/UMG/Public/Components/SlateLeafWidgetComponent.h
deleted file mode 100644
diff --git a/Engine/Source/Runtime/UMG/Public/Components/SlateNonLeafWidgetComponent.h b/Engine/Source/Runtime/UMG/Public/Components/SlateNonLeafWidgetComponent.h
--- a/Engine/Source/Runtime/UMG/Public/Components/SlateNonLeafWidgetComponent.h
+++ b/Engine/Source/Runtime/UMG/Public/Components/SlateNonLeafWidgetComponent.h
@@ -5,23 +5,23 @@
 #include "SlateNonLeafWidgetComponent.generated.h"
 
 UCLASS(Abstract)
-class UMG_API USlateNonLeafWidgetComponent : public USlateWrapperComponent
+class UMG_API UPanelWidget : public UWidget
 {
 	GENERATED_UCLASS_BODY()
 
 	virtual int32 GetChildrenCount() const { return 0; }
-	virtual USlateWrapperComponent* GetChildAt(int32 Index) const { return NULL; }
+	virtual UWidget* GetChildAt(int32 Index) const { return NULL; }
 
-	virtual void ReplaceChildAt(int32 Index, USlateWrapperComponent* Child)
+	virtual void ReplaceChildAt(int32 Index, UWidget* Child)
 	{
 	}
 
-	virtual bool AddChild(USlateWrapperComponent* Child, FVector2D Position)
+	virtual bool AddChild(UWidget* Child, FVector2D Position)
 	{
 		return false;
 	}
 
-	virtual bool RemoveChild(USlateWrapperComponent* Child)
+	virtual bool RemoveChild(UWidget* Child)
 	{
 		return false;
 	}
diff --git a/Engine/Source/Runtime/UMG/Public/Components/SlateWrapperComponent.h b/Engine/Source/Runtime/UMG/Public/Components/SlateWrapperComponent.h
--- a/Engine/Source/Runtime/UMG/Public/Components/SlateWrapperComponent.h
+++ b/Engine/Source/Runtime/UMG/Public/Components/SlateWrapperComponent.h
@@ -3,14 +3,14 @@
 #pragma once
 
 #include "SlateWrapperTypes.h"
-#include "SlateWrapperComponent.generated.h"
+#include "Widget.generated.h"
 
 #define OPTIONAL_BINDING(ReturnType, MemberName) MemberName ## Delegate.IsBound() ? TAttribute< ReturnType >::Create(MemberName ## Delegate.GetUObject(), MemberName ## Delegate.GetFunctionName()) : TAttribute< ReturnType >(MemberName)
 #define OPTIONAL_BINDING_CONVERT(ReturnType, MemberName, ConvertedType, ConversionFunction) MemberName ## Delegate.IsBound() ? TAttribute< ConvertedType >::Create(TAttribute< ConvertedType >::FGetter::CreateUObject(this, &ThisClass::ConversionFunction, TAttribute< ReturnType >::Create(MemberName ## Delegate.GetUObject(), MemberName ## Delegate.GetFunctionName()))) : ConversionFunction(TAttribute< ReturnType >(MemberName))
 
 /** This is the base class for all Slate widget wrapper components */
 UCLASS(Abstract, hideCategories=(Activation, "Components|Activation", ComponentReplication, LOD, Rendering))
-class UMG_API USlateWrapperComponent : public UObject
+class UMG_API UWidget : public UObject
 {
 	GENERATED_UCLASS_BODY()
 
diff --git a/Engine/Source/Runtime/UMG/Public/Components/Slider.h b/Engine/Source/Runtime/UMG/Public/Components/Slider.h
--- a/Engine/Source/Runtime/UMG/Public/Components/Slider.h
+++ b/Engine/Source/Runtime/UMG/Public/Components/Slider.h
@@ -70,9 +70,9 @@ protected:
 		return TSharedPtr<class SSlider>();
 	}
 
-	// USlateWrapperComponent interface
+	// UWidget interface
 	virtual TSharedRef<SWidget> RebuildWidget() OVERRIDE;
-	// End of USlateWrapperComponent interface
+	// End of UWidget interface
 
 	void HandleOnValueChanged(float InValue);
 	void HandleOnMouseCaptureBegin();
diff --git a/Engine/Source/Runtime/UMG/Public/Components/SpacerComponent.h b/Engine/Source/Runtime/UMG/Public/Components/SpacerComponent.h
--- a/Engine/Source/Runtime/UMG/Public/Components/SpacerComponent.h
+++ b/Engine/Source/Runtime/UMG/Public/Components/SpacerComponent.h
@@ -15,9 +15,9 @@ class UMG_API USpacerComponent : public USlateLeafWidgetComponent
 	FVector2D Size;
 
 protected:
-	// USlateWrapperComponent interface
+	// UWidget interface
 	virtual TSharedRef<SWidget> RebuildWidget() OVERRIDE;
-	// End of USlateWrapperComponent interface
+	// End of UWidget interface
 
 	FVector2D GetSpacerSize() const;
 };
diff --git a/Engine/Source/Runtime/UMG/Public/Components/SpinningImageComponent.h b/Engine/Source/Runtime/UMG/Public/Components/SpinningImageComponent.h
--- a/Engine/Source/Runtime/UMG/Public/Components/SpinningImageComponent.h
+++ b/Engine/Source/Runtime/UMG/Public/Components/SpinningImageComponent.h
@@ -15,7 +15,7 @@ class UMG_API USpinningImageComponent : public UImageComponent
 	float Period;
 
 protected:
-	// USlateWrapperComponent interface
+	// UWidget interface
 	virtual TSharedRef<SWidget> RebuildWidget() OVERRIDE;
-	// End of USlateWrapperComponent interface
+	// End of UWidget interface
 };
diff --git a/Engine/Source/Runtime/UMG/Public/Components/TextBlockComponent.h b/Engine/Source/Runtime/UMG/Public/Components/TextBlockComponent.h
--- a/Engine/Source/Runtime/UMG/Public/Components/TextBlockComponent.h
+++ b/Engine/Source/Runtime/UMG/Public/Components/TextBlockComponent.h
@@ -39,9 +39,9 @@ public:
 	USlateWidgetStyleAsset* Style;
 
 protected:
-	// USlateWrapperComponent interface
+	// UWidget interface
 	virtual TSharedRef<SWidget> RebuildWidget() OVERRIDE;
-	// End of USlateWrapperComponent interface
+	// End of UWidget interface
 
 	FSlateColor GetColorAndOpacity() const;
 	FLinearColor GetShadowColorAndOpacity() const;
diff --git a/Engine/Source/Runtime/UMG/Public/Components/ThrobberComponent.h b/Engine/Source/Runtime/UMG/Public/Components/ThrobberComponent.h
--- a/Engine/Source/Runtime/UMG/Public/Components/ThrobberComponent.h
+++ b/Engine/Source/Runtime/UMG/Public/Components/ThrobberComponent.h
@@ -31,7 +31,7 @@ class UMG_API UThrobberComponent : public USlateLeafWidgetComponent
 	USlateBrushAsset* PieceImage;
 
 protected:
-	// USlateWrapperComponent interface
+	// UWidget interface
 	virtual TSharedRef<SWidget> RebuildWidget() OVERRIDE;
-	// End of USlateWrapperComponent interface
+	// End of UWidget interface
 };
diff --git a/Engine/Source/Runtime/UMG/Public/Components/UniformGridPanel.h b/Engine/Source/Runtime/UMG/Public/Components/UniformGridPanel.h
--- a/Engine/Source/Runtime/UMG/Public/Components/UniformGridPanel.h
+++ b/Engine/Source/Runtime/UMG/Public/Components/UniformGridPanel.h
@@ -5,7 +5,7 @@
 #include "UniformGridPanel.generated.h"
 
 UCLASS(meta=(BlueprintSpawnableComponent), ClassGroup=UserInterface)
-class UMG_API UUniformGridPanel : public USlateNonLeafWidgetComponent
+class UMG_API UUniformGridPanel : public UPanelWidget
 {
 	GENERATED_UCLASS_BODY()
 
@@ -13,24 +13,24 @@ class UMG_API UUniformGridPanel : public USlateNonLeafWidgetComponent
 	UPROPERTY(EditAnywhere, EditInline, Category=Slots)
 	TArray<UUniformGridSlot*> Slots;
 
-	UUniformGridSlot* AddSlot(USlateWrapperComponent* Content);
+	UUniformGridSlot* AddSlot(UWidget* Content);
 
-	// USlateNonLeafWidgetComponent
+	// UPanelWidget
 	virtual int32 GetChildrenCount() const OVERRIDE;
-	virtual USlateWrapperComponent* GetChildAt(int32 Index) const OVERRIDE;
-	virtual bool AddChild(USlateWrapperComponent* Child, FVector2D Position) OVERRIDE;
-	virtual bool RemoveChild(USlateWrapperComponent* Child) OVERRIDE;
-	virtual void ReplaceChildAt(int32 Index, USlateWrapperComponent* Child) OVERRIDE;
-	// End USlateNonLeafWidgetComponent
+	virtual UWidget* GetChildAt(int32 Index) const OVERRIDE;
+	virtual bool AddChild(UWidget* Child, FVector2D Position) OVERRIDE;
+	virtual bool RemoveChild(UWidget* Child) OVERRIDE;
+	virtual void ReplaceChildAt(int32 Index, UWidget* Child) OVERRIDE;
+	// End UPanelWidget
 
 #if WITH_EDITOR
 	// UObject interface
 	virtual void PostEditChangeProperty(struct FPropertyChangedEvent& PropertyChangedEvent) OVERRIDE;
 	// End of UObject interface
 
-	// USlateWrapperComponent interface
+	// UWidget interface
 	virtual void ConnectEditorData() OVERRIDE;
-	// End USlateWrapperComponent interface
+	// End UWidget interface
 #endif
 
 protected:
@@ -38,7 +38,7 @@ protected:
 	TWeakPtr<class SUniformGridPanel> MyUniformGridPanel;
 
 protected:
-	// USlateWrapperComponent interface
+	// UWidget interface
 	virtual TSharedRef<SWidget> RebuildWidget() OVERRIDE;
-	// End of USlateWrapperComponent interface
+	// End of UWidget interface
 };
diff --git a/Engine/Source/Runtime/UMG/Public/Components/VerticalBoxComponent.h b/Engine/Source/Runtime/UMG/Public/Components/VerticalBoxComponent.h
--- a/Engine/Source/Runtime/UMG/Public/Components/VerticalBoxComponent.h
+++ b/Engine/Source/Runtime/UMG/Public/Components/VerticalBoxComponent.h
@@ -5,7 +5,7 @@
 #include "VerticalBoxComponent.generated.h"
 
 UCLASS(meta=(BlueprintSpawnableComponent), ClassGroup=UserInterface)
-class UMG_API UVerticalBoxComponent : public USlateNonLeafWidgetComponent
+class UMG_API UVerticalBoxComponent : public UPanelWidget
 {
 	GENERATED_UCLASS_BODY()
 
@@ -13,24 +13,24 @@ class UMG_API UVerticalBoxComponent : public USlateNonLeafWidgetComponent
 	UPROPERTY(EditAnywhere, EditInline, Category=Slots)
 	TArray<UVerticalBoxSlot*> Slots;
 
-	UVerticalBoxSlot* AddSlot(USlateWrapperComponent* Content);
+	UVerticalBoxSlot* AddSlot(UWidget* Content);
 
-	// USlateNonLeafWidgetComponent
+	// UPanelWidget
 	virtual int32 GetChildrenCount() const OVERRIDE;
-	virtual USlateWrapperComponent* GetChildAt(int32 Index) const OVERRIDE;
-	virtual bool AddChild(USlateWrapperComponent* Child, FVector2D Position) OVERRIDE;
-	virtual bool RemoveChild(USlateWrapperComponent* Child) OVERRIDE;
-	virtual void ReplaceChildAt(int32 Index, USlateWrapperComponent* Child) OVERRIDE;
-	// End USlateNonLeafWidgetComponent
+	virtual UWidget* GetChildAt(int32 Index) const OVERRIDE;
+	virtual bool AddChild(UWidget* Child, FVector2D Position) OVERRIDE;
+	virtual bool RemoveChild(UWidget* Child) OVERRIDE;
+	virtual void ReplaceChildAt(int32 Index, UWidget* Child) OVERRIDE;
+	// End UPanelWidget
 
 #if WITH_EDITOR
 	// UObject interface
 	virtual void PostEditChangeProperty(struct FPropertyChangedEvent& PropertyChangedEvent) OVERRIDE;
 	// End of UObject interface
 
-	// USlateWrapperComponent interface
+	// UWidget interface
 	virtual void ConnectEditorData() OVERRIDE;
-	// End USlateWrapperComponent interface
+	// End UWidget interface
 #endif
 
 protected:
@@ -38,7 +38,7 @@ protected:
 	TWeakPtr<class SVerticalBox> MyVerticalBox;
 
 protected:
-	// USlateWrapperComponent interface
+	// UWidget interface
 	virtual TSharedRef<SWidget> RebuildWidget() OVERRIDE;
-	// End of USlateWrapperComponent interface
+	// End of UWidget interface
 };
