UMG - Working on allowing reusable blueprint widgets to be used inside other blueprint widgets.

diff --git a/Engine/Source/Editor/UMGEditor/Private/SUMGEditorWidgetTemplates.cpp b/Engine/Source/Editor/UMGEditor/Private/SUMGEditorWidgetTemplates.cpp
--- a/Engine/Source/Editor/UMGEditor/Private/SUMGEditorWidgetTemplates.cpp
+++ b/Engine/Source/Editor/UMGEditor/Private/SUMGEditorWidgetTemplates.cpp
@@ -123,13 +123,21 @@ void SUMGEditorWidgetTemplates::BuildClassWidgetList()
 	for ( TObjectIterator<UClass> ClassIt; ClassIt; ++ClassIt )
 	{
 		UClass* WidgetClass = *ClassIt;
-		if ( WidgetClass->IsChildOf(UWidget::StaticClass()) && WidgetClass->HasAnyClassFlags(CLASS_Abstract) == false )
+		if ( WidgetClass->IsChildOf(UWidget::StaticClass()) )
 		{
-			TSharedPtr<FWidgetTemplateClass> Template = MakeShareable(new FWidgetTemplateClass(WidgetClass));
-			Template->Name = WidgetClass->GetDisplayNameText();
-			Template->ToolTip = WidgetClass->GetDisplayNameText();
-
-			AddWidgetTemplate(Template);
+			const bool bIsSkeletonClass = WidgetClass->HasAnyFlags(RF_Transient) && WidgetClass->HasAnyClassFlags(CLASS_CompiledFromBlueprint);
+			const bool bIsValidClass =
+				!( WidgetClass->HasAnyClassFlags(CLASS_Abstract | CLASS_Deprecated | CLASS_NewerVersionExists) || bIsSkeletonClass );
+			const bool bIsSameClass = WidgetClass->GetFName() == GetBlueprint()->GeneratedClass->GetFName();
+
+			if ( bIsValidClass && !bIsSameClass )
+			{
+				TSharedPtr<FWidgetTemplateClass> Template = MakeShareable(new FWidgetTemplateClass(WidgetClass));
+				Template->Name = WidgetClass->GetDisplayNameText();
+				Template->ToolTip = WidgetClass->GetDisplayNameText();
+
+				AddWidgetTemplate(Template);
+			}
 		}
 	}
 }
diff --git a/Engine/Source/Runtime/UMG/Private/UserWidget.cpp b/Engine/Source/Runtime/UMG/Private/UserWidget.cpp
--- a/Engine/Source/Runtime/UMG/Private/UserWidget.cpp
+++ b/Engine/Source/Runtime/UMG/Private/UserWidget.cpp
@@ -132,7 +132,7 @@ void UUserWidget::PostInitProperties()
 		BGClass->InitializeWidget(this);
 	}
 
-	RebuildWrapperWidget();
+	RebuildWidget();
 }
 
 UWorld* UUserWidget::GetWorld() const
@@ -157,7 +157,7 @@ UWidget* UUserWidget::GetWidgetHandle(TSharedRef<SWidget> InWidget)
 	return WidgetToComponent.FindRef(InWidget);
 }
 
-void UUserWidget::RebuildWrapperWidget()
+TSharedRef<SWidget> UUserWidget::RebuildWidget()
 {
 	WidgetToComponent.Reset();
 	
@@ -193,6 +193,8 @@ void UUserWidget::RebuildWrapperWidget()
 			Hide();
 		}
 	}
+
+	return MakeWidget();
 }
 
 TSharedPtr<SWidget> UUserWidget::GetWidgetFromName(const FString& Name) const
diff --git a/Engine/Source/Runtime/UMG/Public/Blueprint/UserWidget.h b/Engine/Source/Runtime/UMG/Public/Blueprint/UserWidget.h
--- a/Engine/Source/Runtime/UMG/Public/Blueprint/UserWidget.h
+++ b/Engine/Source/Runtime/UMG/Public/Blueprint/UserWidget.h
@@ -11,13 +11,10 @@
 DECLARE_DYNAMIC_MULTICAST_DELEGATE_OneParam(FOnVisibilityChangedEvent, ESlateVisibility::Type, Visibility);
 
 UCLASS(Abstract, editinlinenew, hideCategories=(Object, Actor, Replication, Rendering, Input))
-class UMG_API UUserWidget : public UObject
+class UMG_API UUserWidget : public UWidget
 {
 	GENERATED_UCLASS_BODY()
 
-	UPROPERTY(EditDefaultsOnly, Category=Behavior)
-	TEnumAsByte<ESlateVisibility::Type> Visiblity;
-
 	/** Called when the visibility changes. */
 	UPROPERTY(BlueprintAssignable)
 	FOnVisibilityChangedEvent OnVisibilityChanged;
@@ -176,12 +173,12 @@ class UMG_API UUserWidget : public UObject
 	TSharedPtr<SWidget> GetWidgetFromName(const FString& Name) const;
 	UWidget* GetHandleFromName(const FString& Name) const;
 
+protected:
+	virtual TSharedRef<SWidget> RebuildWidget() OVERRIDE;
+
 private:
 	TSharedPtr<SWidget> UserRootWidget;
 	TMap< TWeakPtr<SWidget>, UWidget* > WidgetToComponent;
 
 	TWeakPtr<SWidget> FullScreenWidget;
-
-private:
-	void RebuildWrapperWidget();
 };
