UMG - Experimenting with a method of doing designer extensions for widgets.  Parent pointers for umg panel slots, needs more adoption, and less micromanagement.  Widget blueprints are once again creatable.

diff --git a/Engine/Source/Editor/UMGEditor/Private/SUMGDesigner.cpp b/Engine/Source/Editor/UMGEditor/Private/SUMGDesigner.cpp
--- a/Engine/Source/Editor/UMGEditor/Private/SUMGDesigner.cpp
+++ b/Engine/Source/Editor/UMGEditor/Private/SUMGDesigner.cpp
@@ -1,4 +1,4 @@
-// Copyright 1998-2014 Epic Games, Inc. All Rights Reserved.
+﻿// Copyright 1998-2014 Epic Games, Inc. All Rights Reserved.
 
 #include "UMGEditorPrivatePCH.h"
 
@@ -13,6 +13,106 @@
 
 #define LOCTEXT_NAMESPACE "UMG"
 
+class FDesignerExtension : public TSharedFromThis<FDesignerExtension>
+{
+public:
+	virtual void Initialize(UWidgetBlueprint* InBlueprint)
+	{
+		Blueprint = InBlueprint;
+	}
+
+	virtual void BuildWidgetsForSelection(const TArray< UWidget* >& Selection, TArray< TSharedRef<SWidget> >& Widgets) = 0;
+
+	FName GetExtensionId() const { return ExtensionId; }
+
+protected:
+	FName ExtensionId;
+	UWidgetBlueprint* Blueprint;
+};
+
+class FVerticalSlotExtension : public FDesignerExtension
+{
+public:
+	FVerticalSlotExtension()
+	{
+		ExtensionId = FName(TEXT("VerticalSlot"));
+	}
+
+	bool IsActive(const TArray< UWidget* >& Selection)
+	{
+		for ( UWidget* Widget : Selection )
+		{
+			if ( !Widget->Slot || !Widget->Slot->IsA(UVerticalBoxSlot::StaticClass()) )
+			{
+				return false;
+			}
+		}
+
+		return Selection.Num() == 1;
+	}
+	
+	virtual void BuildWidgetsForSelection(const TArray< UWidget* >& Selection, TArray< TSharedRef<SWidget> >& Widgets) OVERRIDE
+	{
+		SelectionCache = Selection;
+
+		if ( !IsActive(Selection) )
+		{
+			return;
+		}
+
+		TSharedRef<SButton> UpButton =
+			SNew(SButton)
+			.Text(LOCTEXT("UpArrow", "↑"))
+			.OnClicked(this, &FVerticalSlotExtension::HandleUpPressed);
+
+		TSharedRef<SButton> DownButton =
+			SNew(SButton)
+			.Text(LOCTEXT("DownArrow", "↓"))
+			.OnClicked(this, &FVerticalSlotExtension::HandleDownPressed);
+
+		Widgets.Add(UpButton);
+		Widgets.Add(DownButton);
+	}
+
+private:
+
+	FReply HandleUpPressed()
+	{
+		for ( UWidget* Widget : SelectionCache )
+		{
+			UVerticalBoxSlot* VerticalSlot = Cast<UVerticalBoxSlot>(Widget->Slot);
+			UVerticalBoxComponent* Parent = Cast<UVerticalBoxComponent>(VerticalSlot->Parent);
+
+			int32 CurrentIndex = Parent->GetChildIndex(Widget);
+			Parent->Slots.RemoveAt(CurrentIndex);
+			Parent->Slots.Insert(VerticalSlot, FMath::Max(CurrentIndex - 1, 0));
+		}
+
+		FBlueprintEditorUtils::MarkBlueprintAsStructurallyModified(Blueprint);
+
+		return FReply::Handled();
+	}
+
+	FReply HandleDownPressed()
+	{
+		for ( UWidget* Widget : SelectionCache )
+		{
+			UVerticalBoxSlot* VerticalSlot = Cast<UVerticalBoxSlot>(Widget->Slot);
+			UVerticalBoxComponent* Parent = Cast<UVerticalBoxComponent>(VerticalSlot->Parent);
+
+			int32 CurrentIndex = Parent->GetChildIndex(Widget);
+			Parent->Slots.RemoveAt(CurrentIndex);
+			Parent->Slots.Insert(VerticalSlot, FMath::Min(CurrentIndex + 1, Parent->GetChildrenCount()));
+		}
+
+		FBlueprintEditorUtils::MarkBlueprintAsStructurallyModified(Blueprint);
+
+		return FReply::Handled();
+	}
+
+	TArray< UWidget* > SelectionCache;
+};
+
 //class SDesignerWidget : public SBorder
 //{
 //	SLATE_BEGIN_ARGS(SDesignerWidget) {}
@@ -116,6 +216,8 @@ void SUMGDesigner::Construct(const FArguments& InArgs, TSharedPtr<FWidgetBluepri
 	DragDirections[DH_BOTTOM_CENTER] = FVector2D(0, 1);
 	DragDirections[DH_BOTTOM_RIGHT] = FVector2D(1, 1);
 
+	Register(MakeShareable(new FVerticalSlotExtension()));
+
 	UWidgetBlueprint* Blueprint = GetBlueprint();
 	Blueprint->OnChanged().AddSP(this, &SUMGDesigner::OnBlueprintChanged);
 
@@ -136,7 +238,8 @@ void SUMGDesigner::Construct(const FArguments& InArgs, TSharedPtr<FWidgetBluepri
 			.HAlign(HAlign_Fill)
 			.VAlign(VAlign_Fill)
 			[
-				SNew(STextBlock)
+				SAssignNew(ExtensionWidgetCanvas, SCanvas)
+				.Visibility(EVisibility::SelfHitTestInvisible)
 			]
 		]
 	);
@@ -162,6 +265,12 @@ UWidgetBlueprint* SUMGDesigner::GetBlueprint() const
 	return NULL;
 }
 
+void SUMGDesigner::Register(TSharedRef<FDesignerExtension> Extension)
+{
+	Extension->Initialize(GetBlueprint());
+	DesignerExtensions.Add(Extension);
+}
+
 void SUMGDesigner::OnBlueprintChanged(UBlueprint* InBlueprint)
 {
 	if ( InBlueprint )
@@ -260,6 +369,36 @@ FReply SUMGDesigner::OnMouseButtonDown(const FGeometry& MyGeometry, const FPoint
 			TArray<UWidget*> SelectedTemplates;
 			SelectedTemplates.Add(SelectedTemplate);
 			ShowDetailsForObjects(SelectedTemplates);
+
+			// Remove all the current extension widgets
+			ExtensionWidgetCanvas->ClearChildren();
+
+			ExtensionWidgets.Reset();
+
+			// Build extension widgets for new selection
+			for ( TSharedRef<FDesignerExtension>& Ext : DesignerExtensions )
+			{
+				Ext->BuildWidgetsForSelection(SelectedTemplates, ExtensionWidgets);
+			}
+
+			TSharedRef<SHorizontalBox> ExtensionBox = SNew(SHorizontalBox);
+
+			// Add Widgets to designer surface
+			for ( TSharedRef<SWidget>& ExWidget : ExtensionWidgets )
+			{
+				ExtensionBox->AddSlot()
+					.AutoWidth()
+					[
+						ExWidget
+					];
+			}
+
+			ExtensionWidgetCanvas->AddSlot()
+				.Position(TAttribute<FVector2D>(this, &SUMGDesigner::GetSelectionDesignerWidgetsLocation))
+				.Size(FVector2D(100, 24))
+				[
+					ExtensionBox
+				];
 		}
 	}
 	else
@@ -337,6 +476,36 @@ bool SUMGDesigner::GetArrangedWidgetRelativeToWindow(TSharedRef<SWidget> Widget,
 	return false;
 }
 
+bool SUMGDesigner::GetArrangedWidgetRelativeToDesigner(TSharedRef<SWidget> Widget, FArrangedWidget& ArrangedWidget) const
+{
+	FWidgetPath WidgetPath;
+	if ( FSlateApplication::Get().GeneratePathToWidgetUnchecked(Widget, WidgetPath) )
+	{
+		FArrangedWidget ArrangedDesigner = WidgetPath.FindArrangedWidget(this->AsShared());
+
+		ArrangedWidget = WidgetPath.FindArrangedWidget(Widget);
+		ArrangedWidget.Geometry.AbsolutePosition -= ArrangedDesigner.Geometry.AbsolutePosition;
+		return true;
+	}
+
+	return false;
+}
+
+FVector2D SUMGDesigner::GetSelectionDesignerWidgetsLocation() const
+{
+	if ( SelectedWidget.IsValid() )
+	{
+		TSharedRef<SWidget> Widget = SelectedWidget.Pin().ToSharedRef();
+
+		FArrangedWidget ArrangedWidget(SNullWidget::NullWidget, FGeometry());
+		GetArrangedWidgetRelativeToDesigner(Widget, ArrangedWidget);
+
+		return ArrangedWidget.Geometry.AbsolutePosition - FVector2D(0, 25);
+	}
+
+	return FVector2D(0, 0);
+}
+
 int32 SUMGDesigner::OnPaint(const FGeometry& AllottedGeometry, const FSlateRect& MyClippingRect, FSlateWindowElementList& OutDrawElements, int32 LayerId, const FWidgetStyle& InWidgetStyle, bool bParentEnabled) const
 {
 	SDesignSurface::OnPaint(AllottedGeometry, MyClippingRect, OutDrawElements, LayerId, InWidgetStyle, bParentEnabled);
diff --git a/Engine/Source/Editor/UMGEditor/Private/SUMGDesigner.h b/Engine/Source/Editor/UMGEditor/Private/SUMGDesigner.h
--- a/Engine/Source/Editor/UMGEditor/Private/SUMGDesigner.h
+++ b/Engine/Source/Editor/UMGEditor/Private/SUMGDesigner.h
@@ -6,6 +6,8 @@
 #include "SNodePanel.h"
 #include "SDesignSurface.h"
 
+class FDesignerExtension;
+
 class SUMGDesigner : public SDesignSurface
 {
 public:
@@ -32,6 +34,8 @@ public:
 	virtual FReply OnDrop(const FGeometry& MyGeometry, const FDragDropEvent& DragDropEvent) OVERRIDE;
 	// End of SWidget interface
 
+	void Register(TSharedRef<FDesignerExtension> Extension);
+
 private:
 	UWidget* GetTemplateAtCursor(const FGeometry& MyGeometry, const FPointerEvent& MouseEvent, FArrangedWidget& ArrangedWidget);
 
@@ -42,6 +46,9 @@ private:
 
 	bool GetArrangedWidget(TSharedRef<SWidget> Widget, FArrangedWidget& ArrangedWidget) const;
 	bool GetArrangedWidgetRelativeToWindow(TSharedRef<SWidget> Widget, FArrangedWidget& ArrangedWidget) const;
+	bool GetArrangedWidgetRelativeToDesigner(TSharedRef<SWidget> Widget, FArrangedWidget& ArrangedWidget) const;
+
+	FVector2D GetSelectionDesignerWidgetsLocation() const;
 
 private:
 	enum DragHandle
@@ -64,6 +71,12 @@ private:
 
 	TArray< FVector2D > DragDirections;
 
+	/** Extensions for the designer to allow for custom widgets to be inserted onto the design surface as selection changes. */
+	TArray< TSharedRef<FDesignerExtension> > DesignerExtensions;
+
+	/** Temporary widgets that are destroyed every time selection changes and are built from the current selection by extensions. */
+	TArray< TSharedRef<SWidget> > ExtensionWidgets;
+
 private:
 	void DrawDragHandles(const FPaintGeometry& SelectionGeometry, const FSlateRect& MyClippingRect, FSlateWindowElementList& OutDrawElements, int32 LayerId, const FWidgetStyle& InWidgetStyle) const;
 	DragHandle HitTestDragHandles(const FGeometry& AllottedGeometry, const FPointerEvent& PointerEvent) const;
@@ -78,6 +91,7 @@ private:
 	TWeakPtr<SWidget> SelectedWidget;
 
 	TSharedPtr<SBorder> PreviewSurface;
+	TSharedPtr<SCanvas> ExtensionWidgetCanvas;
 
 	DragHandle CurrentHandle;
 };
diff --git a/Engine/Source/Editor/UMGEditor/Private/WidgetBlueprintFactory.cpp b/Engine/Source/Editor/UMGEditor/Private/WidgetBlueprintFactory.cpp
--- a/Engine/Source/Editor/UMGEditor/Private/WidgetBlueprintFactory.cpp
+++ b/Engine/Source/Editor/UMGEditor/Private/WidgetBlueprintFactory.cpp
@@ -41,7 +41,7 @@ UObject* UWidgetBlueprintFactory::FactoryCreateNew(UClass* Class, UObject* InPar
 	{
 		FFormatNamedArguments Args;
 		Args.Add( TEXT("ClassName"), (ParentClass != NULL) ? FText::FromString( ParentClass->GetName() ) : LOCTEXT("Null", "(null)") );
-		FMessageDialog::Open( EAppMsgType::Ok, FText::Format( LOCTEXT("CannotCreateAnimBlueprint", "Cannot create a Anim Blueprint based on the class '{0}'."), Args ) );
+		FMessageDialog::Open( EAppMsgType::Ok, FText::Format( LOCTEXT("CannotCreateWidgetBlueprint", "Cannot create a Widget Blueprint based on the class '{ClassName}'."), Args ) );
 		return NULL;
 	}
 	else
diff --git a/Engine/Source/Runtime/UMG/Private/Components/BorderComponent.cpp b/Engine/Source/Runtime/UMG/Private/Components/BorderComponent.cpp
--- a/Engine/Source/Runtime/UMG/Private/Components/BorderComponent.cpp
+++ b/Engine/Source/Runtime/UMG/Private/Components/BorderComponent.cpp
@@ -20,9 +20,9 @@ UBorderComponent::UBorderComponent(const FPostConstructInitializeProperties& PCI
 	//BorderImage = *BorderDefaults._BorderImage.Get();
 	DesiredSizeScale = BorderDefaults._DesiredSizeScale.Get();
 	ContentScale = BorderDefaults._ContentScale.Get();
-	ContentColorAndOpacity = BorderDefaults._ColorAndOpacity.Get();
-	BorderColorAndOpacity = BorderDefaults._BorderBackgroundColor.Get();
-	ForegroundColor = BorderDefaults._ForegroundColor.Get();
+	ContentColorAndOpacity = FLinearColor::White;
+	BorderColorAndOpacity = FLinearColor::White;
+	ForegroundColor = FLinearColor::Black;
 	bShowEffectWhenDisabled = BorderDefaults._ShowEffectWhenDisabled.Get();
 }
 
diff --git a/Engine/Source/Runtime/UMG/Private/Components/VerticalBoxComponent.cpp b/Engine/Source/Runtime/UMG/Private/Components/VerticalBoxComponent.cpp
--- a/Engine/Source/Runtime/UMG/Private/Components/VerticalBoxComponent.cpp
+++ b/Engine/Source/Runtime/UMG/Private/Components/VerticalBoxComponent.cpp
@@ -21,6 +21,21 @@ UWidget* UVerticalBoxComponent::GetChildAt(int32 Index) const
 	return Slots[Index]->Content;
 }
 
+int32 UVerticalBoxComponent::GetChildIndex(UWidget* Content) const
+{
+	for ( int32 SlotIndex = 0; SlotIndex < Slots.Num(); ++SlotIndex )
+	{
+		UVerticalBoxSlot* Slot = Slots[SlotIndex];
+
+		if ( Slot->Content == Content )
+		{
+			return SlotIndex;
+		}
+	}
+
+	return -1;
+}
+
 bool UVerticalBoxComponent::AddChild(UWidget* Child, FVector2D Position)
 {
 	AddSlot(Child);
@@ -29,15 +44,11 @@ bool UVerticalBoxComponent::AddChild(UWidget* Child, FVector2D Position)
 
 bool UVerticalBoxComponent::RemoveChild(UWidget* Child)
 {
-	for ( int32 SlotIndex = 0; SlotIndex < Slots.Num(); ++SlotIndex )
+	int32 SlotIndex = GetChildIndex(Child);
+	if ( SlotIndex != -1 )
 	{
-		UVerticalBoxSlot* Slot = Slots[SlotIndex];
-
-		if ( Slot->Content == Child )
-		{
-			Slots.RemoveAt(SlotIndex);
-			return true;
-		}
+		Slots.RemoveAt(SlotIndex);
+		return true;
 	}
 
 	return false;
@@ -53,6 +64,19 @@ void UVerticalBoxComponent::ReplaceChildAt(int32 Index, UWidget* Content)
 #endif
 }
 
+void UVerticalBoxComponent::InsertChildAt(int32 Index, UWidget* Content)
+{
+	UVerticalBoxSlot* Slot = ConstructObject<UVerticalBoxSlot>(UVerticalBoxSlot::StaticClass(), this);
+	Slot->Content = Content;
+	Slot->Parent = this;
+
+#if WITH_EDITOR
+	Content->Slot = Slot;
+#endif
+
+	Slots.Insert(Slot, Index);
+}
+
 TSharedRef<SWidget> UVerticalBoxComponent::RebuildWidget()
 {
 	TSharedRef<SVerticalBox> NewCanvas = SNew(SVerticalBox);
@@ -71,6 +95,8 @@ TSharedRef<SWidget> UVerticalBoxComponent::RebuildWidget()
 				Slots[SlotIndex] = Slot = ConstructObject<UVerticalBoxSlot>(UVerticalBoxSlot::StaticClass(), this);
 			}
 
+			Slot->Parent = this;
+
 			auto& NewSlot = Canvas->AddSlot()
 				.Padding(Slot->Padding)
 				.HAlign(Slot->HorizontalAlignment)
@@ -90,6 +116,7 @@ UVerticalBoxSlot* UVerticalBoxComponent::AddSlot(UWidget* Content)
 {
 	UVerticalBoxSlot* Slot = ConstructObject<UVerticalBoxSlot>(UVerticalBoxSlot::StaticClass(), this);
 	Slot->Content = Content;
+	Slot->Parent = this;
 
 #if WITH_EDITOR
 	Content->Slot = Slot;
@@ -105,6 +132,9 @@ void UVerticalBoxComponent::ConnectEditorData()
 {
 	for ( UVerticalBoxSlot* Slot : Slots )
 	{
+		// TODO UMG Find a better way to assign the parent pointers for slots.
+		Slot->Parent = this;
+
 		if ( Slot->Content )
 		{
 			Slot->Content->Slot = Slot;
diff --git a/Engine/Source/Runtime/UMG/Public/Blueprint/UserWidget.h b/Engine/Source/Runtime/UMG/Public/Blueprint/UserWidget.h
--- a/Engine/Source/Runtime/UMG/Public/Blueprint/UserWidget.h
+++ b/Engine/Source/Runtime/UMG/Public/Blueprint/UserWidget.h
@@ -67,7 +67,7 @@ public:
 
 DECLARE_DYNAMIC_MULTICAST_DELEGATE_OneParam(FOnVisibilityChangedEvent, ESlateVisibility::Type, Visibility);
 
-UCLASS(Abstract, editinlinenew, hideCategories=(Object, Actor, Replication, Rendering, Input))
+UCLASS(Abstract, editinlinenew, BlueprintType, Blueprintable)
 class UMG_API UUserWidget : public UWidget
 {
 	GENERATED_UCLASS_BODY()
diff --git a/Engine/Source/Runtime/UMG/Public/Components/PanelSlot.h b/Engine/Source/Runtime/UMG/Public/Components/PanelSlot.h
--- a/Engine/Source/Runtime/UMG/Public/Components/PanelSlot.h
+++ b/Engine/Source/Runtime/UMG/Public/Components/PanelSlot.h
@@ -8,6 +8,9 @@ UCLASS()
 class UMG_API UPanelSlot : public UObject
 {
 	GENERATED_UCLASS_BODY()
+	
+	UPROPERTY()
+	class UPanelWidget* Parent;
 
 	UPROPERTY()
 	class UWidget* Content;
diff --git a/Engine/Source/Runtime/UMG/Public/Components/PanelWidget.h b/Engine/Source/Runtime/UMG/Public/Components/PanelWidget.h
--- a/Engine/Source/Runtime/UMG/Public/Components/PanelWidget.h
+++ b/Engine/Source/Runtime/UMG/Public/Components/PanelWidget.h
@@ -12,10 +12,16 @@ class UMG_API UPanelWidget : public UWidget
 	virtual int32 GetChildrenCount() const { return 0; }
 	virtual UWidget* GetChildAt(int32 Index) const { return NULL; }
 
+	virtual int32 GetChildIndex(UWidget* Content) const { return -1; }
+
 	virtual void ReplaceChildAt(int32 Index, UWidget* Child)
 	{
 	}
 
+	virtual void InsertChildAt(int32 Index, UWidget* Child)
+	{
+	}
+
 	virtual bool AddChild(UWidget* Child, FVector2D Position)
 	{
 		return false;
diff --git a/Engine/Source/Runtime/UMG/Public/Components/VerticalBoxComponent.h b/Engine/Source/Runtime/UMG/Public/Components/VerticalBoxComponent.h
--- a/Engine/Source/Runtime/UMG/Public/Components/VerticalBoxComponent.h
+++ b/Engine/Source/Runtime/UMG/Public/Components/VerticalBoxComponent.h
@@ -18,9 +18,11 @@ class UMG_API UVerticalBoxComponent : public UPanelWidget
 	// UPanelWidget
 	virtual int32 GetChildrenCount() const OVERRIDE;
 	virtual UWidget* GetChildAt(int32 Index) const OVERRIDE;
+	virtual int32 GetChildIndex(UWidget* Content) const OVERRIDE;
 	virtual bool AddChild(UWidget* Child, FVector2D Position) OVERRIDE;
 	virtual bool RemoveChild(UWidget* Child) OVERRIDE;
 	virtual void ReplaceChildAt(int32 Index, UWidget* Child) OVERRIDE;
+	virtual void InsertChildAt(int32 Index, UWidget* Child) OVERRIDE;
 	// End UPanelWidget
 
 #if WITH_EDITOR
