UMG - Reorganizing the designer extension code.

diff --git a/Engine/Source/Editor/UMGEditor/Private/Extensions/DesignerExtension.cpp b/Engine/Source/Editor/UMGEditor/Private/Extensions/DesignerExtension.cpp
new file mode 100644
--- /dev/null
+++ b/Engine/Source/Editor/UMGEditor/Private/Extensions/DesignerExtension.cpp
@@ -0,0 +1,17 @@
+// Copyright 1998-2014 Epic Games, Inc. All Rights Reserved.
+
+#include "UMGEditorPrivatePCH.h"
+
+#define LOCTEXT_NAMESPACE "UMG"
+
+void FDesignerExtension::Initialize(UWidgetBlueprint* InBlueprint)
+{
+	Blueprint = InBlueprint;
+}
+
+FName FDesignerExtension::GetExtensionId() const
+{
+	return ExtensionId;
+}
+
+#undef LOCTEXT_NAMESPACE
diff --git a/Engine/Source/Editor/UMGEditor/Private/Extensions/VerticalSlotExtension.cpp b/Engine/Source/Editor/UMGEditor/Private/Extensions/VerticalSlotExtension.cpp
new file mode 100644
--- /dev/null
+++ b/Engine/Source/Editor/UMGEditor/Private/Extensions/VerticalSlotExtension.cpp
@@ -0,0 +1,82 @@
+// Copyright 1998-2014 Epic Games, Inc. All Rights Reserved.
+
+#include "UMGEditorPrivatePCH.h"
+
+#define LOCTEXT_NAMESPACE "UMG"
+
+FVerticalSlotExtension::FVerticalSlotExtension()
+{
+	ExtensionId = FName(TEXT("VerticalSlot"));
+}
+
+bool FVerticalSlotExtension::IsActive(const TArray< UWidget* >& Selection)
+{
+	for ( UWidget* Widget : Selection )
+	{
+		if ( !Widget->Slot || !Widget->Slot->IsA(UVerticalBoxSlot::StaticClass()) )
+		{
+			return false;
+		}
+	}
+
+	return Selection.Num() == 1;
+}
+
+void FVerticalSlotExtension::BuildWidgetsForSelection(const TArray< UWidget* >& Selection, TArray< TSharedRef<SWidget> >& Widgets)
+{
+	SelectionCache = Selection;
+
+	if ( !IsActive(Selection) )
+	{
+		return;
+	}
+
+	TSharedRef<SButton> UpButton =
+		SNew(SButton)
+		.Text(LOCTEXT("UpArrow", "?"))
+		.OnClicked(this, &FVerticalSlotExtension::HandleUpPressed);
+
+	TSharedRef<SButton> DownButton =
+		SNew(SButton)
+		.Text(LOCTEXT("DownArrow", "?"))
+		.OnClicked(this, &FVerticalSlotExtension::HandleDownPressed);
+
+	Widgets.Add(UpButton);
+	Widgets.Add(DownButton);
+}
+
+FReply FVerticalSlotExtension::HandleUpPressed()
+{
+	for ( UWidget* Widget : SelectionCache )
+	{
+		UVerticalBoxSlot* VerticalSlot = Cast<UVerticalBoxSlot>(Widget->Slot);
+		UVerticalBoxComponent* Parent = Cast<UVerticalBoxComponent>(VerticalSlot->Parent);
+
+		int32 CurrentIndex = Parent->GetChildIndex(Widget);
+		Parent->Slots.RemoveAt(CurrentIndex);
+		Parent->Slots.Insert(VerticalSlot, FMath::Max(CurrentIndex - 1, 0));
+	}
+
+	FBlueprintEditorUtils::MarkBlueprintAsStructurallyModified(Blueprint);
+
+	return FReply::Handled();
+}
+
+FReply FVerticalSlotExtension::HandleDownPressed()
+{
+	for ( UWidget* Widget : SelectionCache )
+	{
+		UVerticalBoxSlot* VerticalSlot = Cast<UVerticalBoxSlot>(Widget->Slot);
+		UVerticalBoxComponent* Parent = Cast<UVerticalBoxComponent>(VerticalSlot->Parent);
+
+		int32 CurrentIndex = Parent->GetChildIndex(Widget);
+		Parent->Slots.RemoveAt(CurrentIndex);
+		Parent->Slots.Insert(VerticalSlot, FMath::Min(CurrentIndex + 1, Parent->GetChildrenCount()));
+	}
+
+	FBlueprintEditorUtils::MarkBlueprintAsStructurallyModified(Blueprint);
+
+	return FReply::Handled();
+}
+
+#undef LOCTEXT_NAMESPACE
diff --git a/Engine/Source/Editor/UMGEditor/Private/Extensions/VerticalSlotExtension.h b/Engine/Source/Editor/UMGEditor/Private/Extensions/VerticalSlotExtension.h
new file mode 100644
--- /dev/null
+++ b/Engine/Source/Editor/UMGEditor/Private/Extensions/VerticalSlotExtension.h
@@ -0,0 +1,22 @@
+// Copyright 1998-2014 Epic Games, Inc. All Rights Reserved.
+
+#pragma once
+
+#include "DesignerExtension.h"
+
+class FVerticalSlotExtension : public FDesignerExtension
+{
+public:
+	FVerticalSlotExtension();
+
+	bool IsActive(const TArray< UWidget* >& Selection);
+	
+	virtual void BuildWidgetsForSelection(const TArray< UWidget* >& Selection, TArray< TSharedRef<SWidget> >& Widgets) OVERRIDE;
+
+private:
+
+	FReply HandleUpPressed();
+	FReply HandleDownPressed();
+
+	TArray< UWidget* > SelectionCache;
+};
\ No newline at end of file
diff --git a/Engine/Source/Editor/UMGEditor/Private/SUMGDesigner.cpp b/Engine/Source/Editor/UMGEditor/Private/SUMGDesigner.cpp
--- a/Engine/Source/Editor/UMGEditor/Private/SUMGDesigner.cpp
+++ b/Engine/Source/Editor/UMGEditor/Private/SUMGDesigner.cpp
@@ -13,141 +13,6 @@
 
 #define LOCTEXT_NAMESPACE "UMG"
 
-class FDesignerExtension : public TSharedFromThis<FDesignerExtension>
-{
-public:
-	virtual void Initialize(UWidgetBlueprint* InBlueprint)
-	{
-		Blueprint = InBlueprint;
-	}
-
-	virtual void BuildWidgetsForSelection(const TArray< UWidget* >& Selection, TArray< TSharedRef<SWidget> >& Widgets) = 0;
-
-	FName GetExtensionId() const { return ExtensionId; }
-
-protected:
-	FName ExtensionId;
-	UWidgetBlueprint* Blueprint;
-};
-
-class FVerticalSlotExtension : public FDesignerExtension
-{
-public:
-	FVerticalSlotExtension()
-	{
-		ExtensionId = FName(TEXT("VerticalSlot"));
-	}
-
-	bool IsActive(const TArray< UWidget* >& Selection)
-	{
-		for ( UWidget* Widget : Selection )
-		{
-			if ( !Widget->Slot || !Widget->Slot->IsA(UVerticalBoxSlot::StaticClass()) )
-			{
-				return false;
-			}
-		}
-
-		return Selection.Num() == 1;
-	}
-	
-	virtual void BuildWidgetsForSelection(const TArray< UWidget* >& Selection, TArray< TSharedRef<SWidget> >& Widgets) OVERRIDE
-	{
-		SelectionCache = Selection;
-
-		if ( !IsActive(Selection) )
-		{
-			return;
-		}
-
-		TSharedRef<SButton> UpButton =
-			SNew(SButton)
-			.Text(LOCTEXT("UpArrow", "↑"))
-			.OnClicked(this, &FVerticalSlotExtension::HandleUpPressed);
-
-		TSharedRef<SButton> DownButton =
-			SNew(SButton)
-			.Text(LOCTEXT("DownArrow", "↓"))
-			.OnClicked(this, &FVerticalSlotExtension::HandleDownPressed);
-
-		Widgets.Add(UpButton);
-		Widgets.Add(DownButton);
-	}
-
-private:
-
-	FReply HandleUpPressed()
-	{
-		for ( UWidget* Widget : SelectionCache )
-		{
-			UVerticalBoxSlot* VerticalSlot = Cast<UVerticalBoxSlot>(Widget->Slot);
-			UVerticalBoxComponent* Parent = Cast<UVerticalBoxComponent>(VerticalSlot->Parent);
-
-			int32 CurrentIndex = Parent->GetChildIndex(Widget);
-			Parent->Slots.RemoveAt(CurrentIndex);
-			Parent->Slots.Insert(VerticalSlot, FMath::Max(CurrentIndex - 1, 0));
-		}
-
-		FBlueprintEditorUtils::MarkBlueprintAsStructurallyModified(Blueprint);
-
-		return FReply::Handled();
-	}
-
-	FReply HandleDownPressed()
-	{
-		for ( UWidget* Widget : SelectionCache )
-		{
-			UVerticalBoxSlot* VerticalSlot = Cast<UVerticalBoxSlot>(Widget->Slot);
-			UVerticalBoxComponent* Parent = Cast<UVerticalBoxComponent>(VerticalSlot->Parent);
-
-			int32 CurrentIndex = Parent->GetChildIndex(Widget);
-			Parent->Slots.RemoveAt(CurrentIndex);
-			Parent->Slots.Insert(VerticalSlot, FMath::Min(CurrentIndex + 1, Parent->GetChildrenCount()));
-		}
-
-		FBlueprintEditorUtils::MarkBlueprintAsStructurallyModified(Blueprint);
-
-		return FReply::Handled();
-	}
-
-	TArray< UWidget* > SelectionCache;
-};
-
-//class SDesignerWidget : public SBorder
-//{
-//	SLATE_BEGIN_ARGS(SDesignerWidget) {}
-//		/** Slot for the wrapped content (optional) */
-//		SLATE_DEFAULT_SLOT(FArguments, Content)
-//	SLATE_END_ARGS()
-//
-//	void Construct(const FArguments& InArgs, TSharedPtr<FBlueprintEditor> InBlueprintEditor)
-//	{
-//		LastPreviewActor = NULL;
-//		BlueprintEditor = InBlueprintEditor;
-//
-//		SBorder::Construct(SBorder::FArguments()
-//			.BorderImage(FStyleDefaults::GetNoBrush())
-//			.Padding(0.0f)
-//			[
-//
-//				InArgs._Content.Widget
-//			]);
-//		}
-//	}
-//
-//	//virtual FReply OnMouseButtonDown(const FGeometry& MyGeometry, const FPointerEvent& MouseEvent) OVERRIDE;
-//	//virtual FReply OnMouseButtonUp(const FGeometry& MyGeometry, const FPointerEvent& MouseEvent) OVERRIDE;
-///*virtual bool OnHitTest(const FGeometry& MyGeometry, FVector2D InAbsoluteCursorPosition) OVERRIDE;
-//{
-//return true;
-//}*/
-//
-//protected:
-//	TWeakPtr<FBlueprintEditor> BlueprintEditor;
-//	TWeakObjectPtr<AActor> LastPreviewActor;
-//};
-
-
 //////////////////////////////////////////////////////////////////////////
 
 static bool LocateWidgetsUnderCursor_Helper(FArrangedWidget& Candidate, FVector2D InAbsoluteCursorLocation, FArrangedChildren& OutWidgetsUnderCursor, bool bIgnoreEnabledStatus)
diff --git a/Engine/Source/Editor/UMGEditor/Private/UMGEditorPrivatePCH.h b/Engine/Source/Editor/UMGEditor/Private/UMGEditorPrivatePCH.h
--- a/Engine/Source/Editor/UMGEditor/Private/UMGEditorPrivatePCH.h
+++ b/Engine/Source/Editor/UMGEditor/Private/UMGEditorPrivatePCH.h
@@ -12,4 +12,7 @@
 #include "UMGEditorClasses.h"
 
 #include "AssetTypeActions_Base.h"
-#include "AssetTypeActions_Blueprint.h"
\ No newline at end of file
+#include "AssetTypeActions_Blueprint.h"
+
+#include "DesignerExtension.h"
+#include "VerticalSlotExtension.h"
\ No newline at end of file
diff --git a/Engine/Source/Editor/UMGEditor/Public/DesignerExtension.h b/Engine/Source/Editor/UMGEditor/Public/DesignerExtension.h
new file mode 100644
--- /dev/null
+++ b/Engine/Source/Editor/UMGEditor/Public/DesignerExtension.h
@@ -0,0 +1,21 @@
+// Copyright 1998-2014 Epic Games, Inc. All Rights Reserved.
+
+#pragma once
+
+class UWidget;
+class SWidget;
+class UWidgetBlueprint;
+
+class FDesignerExtension : public TSharedFromThis<FDesignerExtension>
+{
+public:
+	virtual void Initialize(UWidgetBlueprint* InBlueprint);
+
+	virtual void BuildWidgetsForSelection(const TArray< UWidget* >& Selection, TArray< TSharedRef<SWidget> >& Widgets) = 0;
+
+	FName GetExtensionId() const;
+
+protected:
+	FName ExtensionId;
+	UWidgetBlueprint* Blueprint;
+};
diff --git a/Engine/Source/Editor/UMGEditor/UMGEditor.Build.cs b/Engine/Source/Editor/UMGEditor/UMGEditor.Build.cs
--- a/Engine/Source/Editor/UMGEditor/UMGEditor.Build.cs
+++ b/Engine/Source/Editor/UMGEditor/UMGEditor.Build.cs
@@ -13,6 +13,7 @@ public class UMGEditor : ModuleRules
 				"Developer/AssetTools/Private/AssetTypeActions",
 				"Editor/UMGEditor/Private", // For PCH includes (because they don't work with relative paths, yet)
 				"Editor/UMGEditor/Private/Templates",
+                "Editor/UMGEditor/Private/Extensions",
 			});
 
 		PrivateIncludePathModuleNames.AddRange(
