UMG - More designer extensions for horizontal, uniform grid, canvas.  Stubbing in combobox.  Faster designer updates for dragging controls/resizing them.

diff --git a/Engine/Source/Editor/UMGEditor/Private/Extensions/CanvasSlotExtension.cpp b/Engine/Source/Editor/UMGEditor/Private/Extensions/CanvasSlotExtension.cpp
new file mode 100644
--- /dev/null
+++ b/Engine/Source/Editor/UMGEditor/Private/Extensions/CanvasSlotExtension.cpp
@@ -0,0 +1,93 @@
+// Copyright 1998-2014 Epic Games, Inc. All Rights Reserved.
+
+#include "UMGEditorPrivatePCH.h"
+
+#define LOCTEXT_NAMESPACE "UMG"
+
+FCanvasSlotExtension::FCanvasSlotExtension()
+	: bDragging(false)
+{
+	ExtensionId = FName(TEXT("CanvasSlot"));
+}
+
+bool FCanvasSlotExtension::IsActive(const TArray< FSelectedWidget >& Selection)
+{
+	for ( const FSelectedWidget& Widget : Selection )
+	{
+		if ( !Widget.Template->Slot || !Widget.Template->Slot->IsA(UCanvasPanelSlot::StaticClass()) )
+		{
+			return false;
+		}
+	}
+
+	return Selection.Num() == 1;
+}
+
+void FCanvasSlotExtension::BuildWidgetsForSelection(const TArray< FSelectedWidget >& Selection, TArray< TSharedRef<SWidget> >& Widgets)
+{
+	SelectionCache = Selection;
+
+	if ( !IsActive(Selection) )
+	{
+		return;
+	}
+
+	MoveHandle =
+		SNew(SBorder)
+		//.BorderImage(FEditorStyle::Get().GetBrush("ToolBar.Background"))
+		//.BorderBackgroundColor(FLinearColor(0,0,0))
+		.OnMouseButtonDown(this, &FCanvasSlotExtension::HandleBeginDrag)
+		.OnMouseButtonUp(this, &FCanvasSlotExtension::HandleEndDrag)
+		.OnMouseMove(this, &FCanvasSlotExtension::HandleDragging)
+		.Padding(FMargin(0))
+		[
+			SNew(SImage)
+			.Image(FCoreStyle::Get().GetBrush("SoftwareCursor_CardinalCross"))
+		];
+
+	Widgets.Add(MoveHandle.ToSharedRef());
+}
+
+FReply FCanvasSlotExtension::HandleBeginDrag(const FGeometry& Geometry, const FPointerEvent& Event)
+{
+	bDragging = true;
+	return FReply::Handled().CaptureMouse(MoveHandle.ToSharedRef());
+}
+
+FReply FCanvasSlotExtension::HandleEndDrag(const FGeometry& Geometry, const FPointerEvent& Event)
+{
+	bDragging = false;
+
+	FBlueprintEditorUtils::MarkBlueprintAsModified(Blueprint);
+
+	return FReply::Handled().ReleaseMouseCapture();
+}
+
+FReply FCanvasSlotExtension::HandleDragging(const FGeometry& Geometry, const FPointerEvent& Event)
+{
+	if ( bDragging )
+	{
+		for ( FSelectedWidget& Selection : SelectionCache )
+		{
+			MoveByAmount(Selection.Preview, Event.GetCursorDelta());
+			MoveByAmount(Selection.Template, Event.GetCursorDelta());
+		}
+
+		return FReply::Handled();
+	}
+
+	return FReply::Unhandled();
+}
+
+void FCanvasSlotExtension::MoveByAmount(UWidget* Widget, FVector2D Delta)
+{
+	if ( !Delta.IsZero() )
+	{
+		UCanvasPanelSlot* CanvasSlot = Cast<UCanvasPanelSlot>(Widget->Slot);
+		UCanvasPanelComponent* Parent = Cast<UCanvasPanelComponent>(CanvasSlot->Parent);
+
+		CanvasSlot->SetPosition(CanvasSlot->Position + Delta);
+	}
+}
+
+#undef LOCTEXT_NAMESPACE
diff --git a/Engine/Source/Editor/UMGEditor/Private/Extensions/CanvasSlotExtension.h b/Engine/Source/Editor/UMGEditor/Private/Extensions/CanvasSlotExtension.h
new file mode 100644
--- /dev/null
+++ b/Engine/Source/Editor/UMGEditor/Private/Extensions/CanvasSlotExtension.h
@@ -0,0 +1,28 @@
+// Copyright 1998-2014 Epic Games, Inc. All Rights Reserved.
+
+#pragma once
+
+#include "DesignerExtension.h"
+
+class FCanvasSlotExtension : public FDesignerExtension
+{
+public:
+	FCanvasSlotExtension();
+
+	bool IsActive(const TArray< FSelectedWidget >& Selection);
+	
+	virtual void BuildWidgetsForSelection(const TArray< FSelectedWidget >& Selection, TArray< TSharedRef<SWidget> >& Widgets) OVERRIDE;
+
+private:
+
+	FReply HandleBeginDrag(const FGeometry& Geometry, const FPointerEvent& Event);
+	FReply HandleEndDrag(const FGeometry& Geometry, const FPointerEvent& Event);
+	FReply HandleDragging(const FGeometry& Geometry, const FPointerEvent& Event);
+
+	void MoveByAmount(UWidget* Widget, FVector2D Delta);
+
+	TArray< FSelectedWidget > SelectionCache;
+	bool bDragging;
+
+	TSharedPtr<SBorder> MoveHandle;
+};
\ No newline at end of file
diff --git a/Engine/Source/Editor/UMGEditor/Private/Extensions/HorizontalSlotExtension.cpp b/Engine/Source/Editor/UMGEditor/Private/Extensions/HorizontalSlotExtension.cpp
new file mode 100644
--- /dev/null
+++ b/Engine/Source/Editor/UMGEditor/Private/Extensions/HorizontalSlotExtension.cpp
@@ -0,0 +1,73 @@
+﻿// Copyright 1998-2014 Epic Games, Inc. All Rights Reserved.
+
+#include "UMGEditorPrivatePCH.h"
+#include "Kismet2/BlueprintEditorUtils.h"
+
+#define LOCTEXT_NAMESPACE "UMG"
+
+FHorizontalSlotExtension::FHorizontalSlotExtension()
+{
+	ExtensionId = FName(TEXT("HorizontalSlot"));
+}
+
+bool FHorizontalSlotExtension::IsActive(const TArray< FSelectedWidget >& Selection)
+{
+	for ( const FSelectedWidget& Widget : Selection )
+	{
+		if ( !Widget.Template->Slot || !Widget.Template->Slot->IsA(UHorizontalBoxSlot::StaticClass()) )
+		{
+			return false;
+		}
+	}
+
+	return Selection.Num() == 1;
+}
+
+void FHorizontalSlotExtension::BuildWidgetsForSelection(const TArray< FSelectedWidget >& Selection, TArray< TSharedRef<SWidget> >& Widgets)
+{
+	SelectionCache = Selection;
+
+	if ( !IsActive(Selection) )
+	{
+		return;
+	}
+
+	TSharedRef<SButton> LeftButton =
+		SNew(SButton)
+		.Text(LOCTEXT("LeftArrow", "←"))
+		.OnClicked(this, &FHorizontalSlotExtension::HandleShift, -1);
+
+	TSharedRef<SButton> RightButton =
+		SNew(SButton)
+		.Text(LOCTEXT("RightArrow", "→"))
+		.OnClicked(this, &FHorizontalSlotExtension::HandleShift, 1);
+
+	Widgets.Add(LeftButton);
+	Widgets.Add(RightButton);
+}
+
+FReply FHorizontalSlotExtension::HandleShift(int32 ShiftAmount)
+{
+	for ( FSelectedWidget& Selection : SelectionCache )
+	{
+		ShiftHorizontal(Selection.Preview, ShiftAmount);
+		ShiftHorizontal(Selection.Template, ShiftAmount);
+	}
+
+	//TODO UMG Reorder the live slot without rebuilding the structure
+	FBlueprintEditorUtils::MarkBlueprintAsStructurallyModified(Blueprint);
+
+	return FReply::Handled();
+}
+
+void FHorizontalSlotExtension::ShiftHorizontal(UWidget* Widget, int32 ShiftAmount)
+{
+	UHorizontalBoxSlot* Slot = Cast<UHorizontalBoxSlot>(Widget->Slot);
+	UHorizontalBoxComponent* Parent = Cast<UHorizontalBoxComponent>(Slot->Parent);
+
+	int32 CurrentIndex = Parent->GetChildIndex(Widget);
+	Parent->Slots.RemoveAt(CurrentIndex);
+	Parent->Slots.Insert(Slot, FMath::Clamp(CurrentIndex + ShiftAmount, 0, Parent->GetChildrenCount()));
+}
+
+#undef LOCTEXT_NAMESPACE
diff --git a/Engine/Source/Editor/UMGEditor/Private/Extensions/HorizontalSlotExtension.h b/Engine/Source/Editor/UMGEditor/Private/Extensions/HorizontalSlotExtension.h
new file mode 100644
--- /dev/null
+++ b/Engine/Source/Editor/UMGEditor/Private/Extensions/HorizontalSlotExtension.h
@@ -0,0 +1,23 @@
+// Copyright 1998-2014 Epic Games, Inc. All Rights Reserved.
+
+#pragma once
+
+#include "DesignerExtension.h"
+
+class FHorizontalSlotExtension : public FDesignerExtension
+{
+public:
+	FHorizontalSlotExtension();
+
+	bool IsActive(const TArray< FSelectedWidget >& Selection);
+	
+	virtual void BuildWidgetsForSelection(const TArray< FSelectedWidget >& Selection, TArray< TSharedRef<SWidget> >& Widgets) OVERRIDE;
+
+private:
+
+	FReply HandleShift(int32 ShiftAmount);
+	
+	void ShiftHorizontal(UWidget* Widget, int32 ShiftAmount);
+
+	TArray< FSelectedWidget > SelectionCache;
+};
\ No newline at end of file
diff --git a/Engine/Source/Editor/UMGEditor/Private/Extensions/UniformGridSlotExtension.cpp b/Engine/Source/Editor/UMGEditor/Private/Extensions/UniformGridSlotExtension.cpp
new file mode 100644
--- /dev/null
+++ b/Engine/Source/Editor/UMGEditor/Private/Extensions/UniformGridSlotExtension.cpp
@@ -0,0 +1,98 @@
+﻿// Copyright 1998-2014 Epic Games, Inc. All Rights Reserved.
+
+#include "UMGEditorPrivatePCH.h"
+
+#define LOCTEXT_NAMESPACE "UMG"
+
+FUniformGridSlotExtension::FUniformGridSlotExtension()
+{
+	ExtensionId = FName(TEXT("UniformGridSlot"));
+}
+
+bool FUniformGridSlotExtension::IsActive(const TArray< FSelectedWidget >& Selection)
+{
+	for ( const FSelectedWidget& Widget : Selection )
+	{
+		if ( !Widget.Template->Slot || !Widget.Template->Slot->IsA(UUniformGridSlot::StaticClass()) )
+		{
+			return false;
+		}
+	}
+
+	return Selection.Num() == 1;
+}
+
+void FUniformGridSlotExtension::BuildWidgetsForSelection(const TArray< FSelectedWidget >& Selection, TArray< TSharedRef<SWidget> >& Widgets)
+{
+	SelectionCache = Selection;
+
+	if ( !IsActive(Selection) )
+	{
+		return;
+	}
+
+	TSharedRef<SButton> UpButton =
+		SNew(SButton)
+		.Text(LOCTEXT("UpArrow", "↑"))
+		.OnClicked(this, &FUniformGridSlotExtension::HandleShiftRow, -1);
+
+	TSharedRef<SButton> DownButton =
+		SNew(SButton)
+		.Text(LOCTEXT("DownArrow", "↓"))
+		.OnClicked(this, &FUniformGridSlotExtension::HandleShiftRow, 1);
+
+	TSharedRef<SButton> LeftButton =
+		SNew(SButton)
+		.Text(LOCTEXT("LeftArrow", "←"))
+		.OnClicked(this, &FUniformGridSlotExtension::HandleShiftColumn, -1);
+
+	TSharedRef<SButton> RightButton =
+		SNew(SButton)
+		.Text(LOCTEXT("RightArrow", "→"))
+		.OnClicked(this, &FUniformGridSlotExtension::HandleShiftColumn, 1);
+
+	Widgets.Add(UpButton);
+	Widgets.Add(DownButton);
+	Widgets.Add(LeftButton);
+	Widgets.Add(RightButton);
+}
+
+FReply FUniformGridSlotExtension::HandleShiftRow(int32 ShiftAmount)
+{
+	for ( FSelectedWidget& Selection : SelectionCache )
+	{
+		ShiftRow(Selection.Preview, ShiftAmount);
+		ShiftRow(Selection.Template, ShiftAmount);
+	}
+
+	FBlueprintEditorUtils::MarkBlueprintAsModified(Blueprint);
+
+	return FReply::Handled();
+}
+
+FReply FUniformGridSlotExtension::HandleShiftColumn(int32 ShiftAmount)
+{
+	for ( FSelectedWidget& Selection : SelectionCache )
+	{
+		ShiftColumn(Selection.Preview, ShiftAmount);
+		ShiftColumn(Selection.Template, ShiftAmount);
+	}
+
+	FBlueprintEditorUtils::MarkBlueprintAsModified(Blueprint);
+
+	return FReply::Handled();
+}
+
+void FUniformGridSlotExtension::ShiftRow(UWidget* Widget, int32 ShiftAmount)
+{
+	UUniformGridSlot* Slot = Cast<UUniformGridSlot>(Widget->Slot);
+	Slot->SetRow(FMath::Max(Slot->Row + ShiftAmount, 0));
+}
+
+void FUniformGridSlotExtension::ShiftColumn(UWidget* Widget, int32 ShiftAmount)
+{
+	UUniformGridSlot* Slot = Cast<UUniformGridSlot>(Widget->Slot);
+	Slot->SetColumn(FMath::Max(Slot->Column + ShiftAmount, 0));
+}
+
+#undef LOCTEXT_NAMESPACE
diff --git a/Engine/Source/Editor/UMGEditor/Private/Extensions/UniformGridSlotExtension.h b/Engine/Source/Editor/UMGEditor/Private/Extensions/UniformGridSlotExtension.h
new file mode 100644
--- /dev/null
+++ b/Engine/Source/Editor/UMGEditor/Private/Extensions/UniformGridSlotExtension.h
@@ -0,0 +1,27 @@
+// Copyright 1998-2014 Epic Games, Inc. All Rights Reserved.
+
+#pragma once
+
+#include "DesignerExtension.h"
+
+class FUniformGridSlotExtension : public FDesignerExtension
+{
+public:
+	FUniformGridSlotExtension();
+
+	bool IsActive(const TArray< FSelectedWidget >& Selection);
+	
+	virtual void BuildWidgetsForSelection(const TArray< FSelectedWidget >& Selection, TArray< TSharedRef<SWidget> >& Widgets) OVERRIDE;
+
+private:
+
+	FReply HandleShiftRow(int32 ShiftAmount);
+	FReply HandleShiftColumn(int32 ShiftAmount);
+
+	void ShiftRow(UWidget* Widget, int32 ShiftAmount);
+	void ShiftColumn(UWidget* Widget, int32 ShiftAmount);
+
+private:
+
+	TArray< FSelectedWidget > SelectionCache;
+};
\ No newline at end of file
diff --git a/Engine/Source/Editor/UMGEditor/Private/Extensions/VerticalSlotExtension.cpp b/Engine/Source/Editor/UMGEditor/Private/Extensions/VerticalSlotExtension.cpp
--- a/Engine/Source/Editor/UMGEditor/Private/Extensions/VerticalSlotExtension.cpp
+++ b/Engine/Source/Editor/UMGEditor/Private/Extensions/VerticalSlotExtension.cpp
@@ -1,4 +1,4 @@
-// Copyright 1998-2014 Epic Games, Inc. All Rights Reserved.
+﻿// Copyright 1998-2014 Epic Games, Inc. All Rights Reserved.
 
 #include "UMGEditorPrivatePCH.h"
 #include "Kismet2/BlueprintEditorUtils.h"
@@ -10,11 +10,11 @@ FVerticalSlotExtension::FVerticalSlotExtension()
 	ExtensionId = FName(TEXT("VerticalSlot"));
 }
 
-bool FVerticalSlotExtension::IsActive(const TArray< UWidget* >& Selection)
+bool FVerticalSlotExtension::IsActive(const TArray< FSelectedWidget >& Selection)
 {
-	for ( UWidget* Widget : Selection )
+	for ( const FSelectedWidget& Widget : Selection )
 	{
-		if ( !Widget->Slot || !Widget->Slot->IsA(UVerticalBoxSlot::StaticClass()) )
+		if ( !Widget.Template->Slot || !Widget.Template->Slot->IsA(UVerticalBoxSlot::StaticClass()) )
 		{
 			return false;
 		}
@@ -23,7 +23,7 @@ bool FVerticalSlotExtension::IsActive(const TArray< UWidget* >& Selection)
 	return Selection.Num() == 1;
 }
 
-void FVerticalSlotExtension::BuildWidgetsForSelection(const TArray< UWidget* >& Selection, TArray< TSharedRef<SWidget> >& Widgets)
+void FVerticalSlotExtension::BuildWidgetsForSelection(const TArray< FSelectedWidget >& Selection, TArray< TSharedRef<SWidget> >& Widgets)
 {
 	SelectionCache = Selection;
 
@@ -34,12 +34,12 @@ void FVerticalSlotExtension::BuildWidgetsForSelection(const TArray< UWidget* >&
 
 	TSharedRef<SButton> UpButton =
 		SNew(SButton)
-		.Text(LOCTEXT("UpArrow", "?"))
+		.Text(LOCTEXT("UpArrow", "↑"))
 		.OnClicked(this, &FVerticalSlotExtension::HandleUpPressed);
 
 	TSharedRef<SButton> DownButton =
 		SNew(SButton)
-		.Text(LOCTEXT("DownArrow", "?"))
+		.Text(LOCTEXT("DownArrow", "↓"))
 		.OnClicked(this, &FVerticalSlotExtension::HandleDownPressed);
 
 	Widgets.Add(UpButton);
@@ -48,16 +48,13 @@ void FVerticalSlotExtension::BuildWidgetsForSelection(const TArray< UWidget* >&
 
 FReply FVerticalSlotExtension::HandleUpPressed()
 {
-	for ( UWidget* Widget : SelectionCache )
+	for ( FSelectedWidget& Selection : SelectionCache )
 	{
-		UVerticalBoxSlot* VerticalSlot = Cast<UVerticalBoxSlot>(Widget->Slot);
-		UVerticalBoxComponent* Parent = Cast<UVerticalBoxComponent>(VerticalSlot->Parent);
-
-		int32 CurrentIndex = Parent->GetChildIndex(Widget);
-		Parent->Slots.RemoveAt(CurrentIndex);
-		Parent->Slots.Insert(VerticalSlot, FMath::Max(CurrentIndex - 1, 0));
+		MoveUp(Selection.Preview);
+		MoveUp(Selection.Template);
 	}
 
+	//TODO UMG Reorder the live slot without rebuilding the structure
 	FBlueprintEditorUtils::MarkBlueprintAsStructurallyModified(Blueprint);
 
 	return FReply::Handled();
@@ -65,19 +62,36 @@ FReply FVerticalSlotExtension::HandleUpPressed()
 
 FReply FVerticalSlotExtension::HandleDownPressed()
 {
-	for ( UWidget* Widget : SelectionCache )
+	for ( FSelectedWidget& Selection : SelectionCache )
 	{
-		UVerticalBoxSlot* VerticalSlot = Cast<UVerticalBoxSlot>(Widget->Slot);
-		UVerticalBoxComponent* Parent = Cast<UVerticalBoxComponent>(VerticalSlot->Parent);
-
-		int32 CurrentIndex = Parent->GetChildIndex(Widget);
-		Parent->Slots.RemoveAt(CurrentIndex);
-		Parent->Slots.Insert(VerticalSlot, FMath::Min(CurrentIndex + 1, Parent->GetChildrenCount()));
+		MoveDown(Selection.Preview);
+		MoveDown(Selection.Template);
 	}
 
+	//TODO UMG Reorder the live slot without rebuilding the structure
 	FBlueprintEditorUtils::MarkBlueprintAsStructurallyModified(Blueprint);
 
 	return FReply::Handled();
 }
 
+void FVerticalSlotExtension::MoveUp(UWidget* Widget)
+{
+	UVerticalBoxSlot* VerticalSlot = Cast<UVerticalBoxSlot>(Widget->Slot);
+	UVerticalBoxComponent* Parent = Cast<UVerticalBoxComponent>(VerticalSlot->Parent);
+
+	int32 CurrentIndex = Parent->GetChildIndex(Widget);
+	Parent->Slots.RemoveAt(CurrentIndex);
+	Parent->Slots.Insert(VerticalSlot, FMath::Max(CurrentIndex - 1, 0));
+}
+
+void FVerticalSlotExtension::MoveDown(UWidget* Widget)
+{
+	UVerticalBoxSlot* VerticalSlot = Cast<UVerticalBoxSlot>(Widget->Slot);
+	UVerticalBoxComponent* Parent = Cast<UVerticalBoxComponent>(VerticalSlot->Parent);
+
+	int32 CurrentIndex = Parent->GetChildIndex(Widget);
+	Parent->Slots.RemoveAt(CurrentIndex);
+	Parent->Slots.Insert(VerticalSlot, FMath::Min(CurrentIndex + 1, Parent->GetChildrenCount()));
+}
+
 #undef LOCTEXT_NAMESPACE
diff --git a/Engine/Source/Editor/UMGEditor/Private/Extensions/VerticalSlotExtension.h b/Engine/Source/Editor/UMGEditor/Private/Extensions/VerticalSlotExtension.h
--- a/Engine/Source/Editor/UMGEditor/Private/Extensions/VerticalSlotExtension.h
+++ b/Engine/Source/Editor/UMGEditor/Private/Extensions/VerticalSlotExtension.h
@@ -9,14 +9,17 @@ class FVerticalSlotExtension : public FDesignerExtension
 public:
 	FVerticalSlotExtension();
 
-	bool IsActive(const TArray< UWidget* >& Selection);
+	bool IsActive(const TArray< FSelectedWidget >& Selection);
 	
-	virtual void BuildWidgetsForSelection(const TArray< UWidget* >& Selection, TArray< TSharedRef<SWidget> >& Widgets) OVERRIDE;
+	virtual void BuildWidgetsForSelection(const TArray< FSelectedWidget >& Selection, TArray< TSharedRef<SWidget> >& Widgets) OVERRIDE;
 
 private:
 
 	FReply HandleUpPressed();
 	FReply HandleDownPressed();
 
-	TArray< UWidget* > SelectionCache;
+	void MoveUp(UWidget* Widget);
+	void MoveDown(UWidget* Widget);
+
+	TArray< FSelectedWidget > SelectionCache;
 };
\ No newline at end of file
diff --git a/Engine/Source/Editor/UMGEditor/Private/SUMGDesigner.cpp b/Engine/Source/Editor/UMGEditor/Private/SUMGDesigner.cpp
--- a/Engine/Source/Editor/UMGEditor/Private/SUMGDesigner.cpp
+++ b/Engine/Source/Editor/UMGEditor/Private/SUMGDesigner.cpp
@@ -67,7 +67,6 @@ void SUMGDesigner::Construct(const FArguments& InArgs, TSharedPtr<FWidgetBluepri
 	PreviewWidgetActor = NULL;
 	BlueprintEditor = InBlueprintEditor;
 	CurrentHandle = DH_NONE;
-	SelectedTemplate = NULL;
 
 	DragDirections.Init((int32)DH_MAX);
 	DragDirections[DH_TOP_LEFT] = FVector2D(-1, -1);
@@ -82,6 +81,9 @@ void SUMGDesigner::Construct(const FArguments& InArgs, TSharedPtr<FWidgetBluepri
 	DragDirections[DH_BOTTOM_RIGHT] = FVector2D(1, 1);
 
 	Register(MakeShareable(new FVerticalSlotExtension()));
+	Register(MakeShareable(new FHorizontalSlotExtension()));
+	Register(MakeShareable(new FCanvasSlotExtension()));
+	Register(MakeShareable(new FUniformGridSlotExtension()));
 
 	UWidgetBlueprint* Blueprint = GetBlueprint();
 	Blueprint->OnChanged().AddSP(this, &SUMGDesigner::OnBlueprintChanged);
@@ -178,7 +180,7 @@ void SUMGDesigner::ShowDetailsForObjects(TArray<UWidget*> Widgets)
 	BlueprintEditor.Pin()->GetInspector()->ShowDetailsForObjects(InspectorObjects, Options);
 }
 
-UWidget* SUMGDesigner::GetTemplateAtCursor(const FGeometry& MyGeometry, const FPointerEvent& MouseEvent, FArrangedWidget& ArrangedWidget)
+FSelectedWidget SUMGDesigner::GetWidgetAtCursor(const FGeometry& MyGeometry, const FPointerEvent& MouseEvent, FArrangedWidget& ArrangedWidget)
 {
 	//@TODO UMG Make it so you can request dropable widgets only, to find the first parentable.
 
@@ -190,15 +192,16 @@ UWidget* SUMGDesigner::GetTemplateAtCursor(const FGeometry& MyGeometry, const FP
 
 	PreviewSurface->SetVisibility(EVisibility::HitTestInvisible);
 
+	FSelectedWidget Selection;
+
 	UUserWidget* WidgetActor = BlueprintEditor.Pin()->GetPreview();
 	if ( WidgetActor )
 	{
-		UWidget* PreviewHandle = NULL;
 		for ( int32 ChildIndex = Children.Num() - 1; ChildIndex >= 0; ChildIndex-- )
 		{
 			FArrangedWidget& Child = Children.GetInternalArray()[ChildIndex];
-			PreviewHandle = WidgetActor->GetWidgetHandle(Child.Widget);
-			if ( PreviewHandle )
+			Selection.Preview = WidgetActor->GetWidgetHandle(Child.Widget);
+			if ( Selection.Preview )
 			{
 				ArrangedWidget = Child;
 				break;
@@ -207,15 +210,14 @@ UWidget* SUMGDesigner::GetTemplateAtCursor(const FGeometry& MyGeometry, const FP
 
 		UWidgetBlueprint* Blueprint = GetBlueprint();
 
-		if ( PreviewHandle )
+		if ( Selection.Preview )
 		{
-			FString Name = PreviewHandle->GetName();
-			UWidget* Template = Blueprint->WidgetTree->FindWidget(Name);
-			return Template;
+			FString Name = Selection.Preview->GetName();
+			Selection.Template = Blueprint->WidgetTree->FindWidget(Name);
 		}
 	}
 
-	return NULL;
+	return Selection;
 }
 
 FReply SUMGDesigner::OnMouseButtonDown(const FGeometry& MyGeometry, const FPointerEvent& MouseEvent)
@@ -225,14 +227,14 @@ FReply SUMGDesigner::OnMouseButtonDown(const FGeometry& MyGeometry, const FPoint
 	if ( CurrentHandle == DH_NONE )
 	{
 		FArrangedWidget ArrangedWidget(SNullWidget::NullWidget, FGeometry());
-		SelectedTemplate = GetTemplateAtCursor(MyGeometry, MouseEvent, ArrangedWidget);
+		CurrentSelection = GetWidgetAtCursor(MyGeometry, MouseEvent, ArrangedWidget);
 
-		if ( SelectedTemplate )
+		if ( CurrentSelection.IsValid() )
 		{
 			//@TODO UMG primary FBlueprintEditor needs to be inherited and selection control needs to be centralized.
 			// Set the template as selected in the details panel
 			TArray<UWidget*> SelectedTemplates;
-			SelectedTemplates.Add(SelectedTemplate);
+			SelectedTemplates.Add(CurrentSelection.Template);
 			ShowDetailsForObjects(SelectedTemplates);
 
 			// Remove all the current extension widgets
@@ -240,10 +242,13 @@ FReply SUMGDesigner::OnMouseButtonDown(const FGeometry& MyGeometry, const FPoint
 
 			ExtensionWidgets.Reset();
 
+			TArray<FSelectedWidget> Selected;
+			Selected.Add(CurrentSelection);
+
 			// Build extension widgets for new selection
 			for ( TSharedRef<FDesignerExtension>& Ext : DesignerExtensions )
 			{
-				Ext->BuildWidgetsForSelection(SelectedTemplates, ExtensionWidgets);
+				Ext->BuildWidgetsForSelection(Selected, ExtensionWidgets);
 			}
 
 			TSharedRef<SHorizontalBox> ExtensionBox = SNew(SHorizontalBox);
@@ -260,7 +265,7 @@ FReply SUMGDesigner::OnMouseButtonDown(const FGeometry& MyGeometry, const FPoint
 
 			ExtensionWidgetCanvas->AddSlot()
 				.Position(TAttribute<FVector2D>(this, &SUMGDesigner::GetSelectionDesignerWidgetsLocation))
-				.Size(FVector2D(100, 24))
+				.Size(FVector2D(100, 25))
 				[
 					ExtensionBox
 				];
@@ -284,13 +289,23 @@ FReply SUMGDesigner::OnMouseMove(const FGeometry& MyGeometry, const FPointerEven
 {
 	if ( CurrentHandle != DH_NONE )
 	{
-		if ( SelectedTemplate->Slot && !MouseEvent.GetCursorDelta().IsZero() )
+		if ( CurrentSelection.IsValid() && !MouseEvent.GetCursorDelta().IsZero() )
 		{
 			//@TODO UMG - Implement some system to query slots to know if they support dragging so we can provide visual feedback by hiding handles that wouldnt work and such.
 			//SelectedTemplate->Slot->CanResize(DH_TOP_CENTER)
 
-			SelectedTemplate->Slot->Resize(DragDirections[CurrentHandle], MouseEvent.GetCursorDelta());
-			FBlueprintEditorUtils::MarkBlueprintAsStructurallyModified(GetBlueprint());
+			if ( CurrentSelection.Preview->Slot )
+			{
+				CurrentSelection.Preview->Slot->Resize(DragDirections[CurrentHandle], MouseEvent.GetCursorDelta());
+			}
+
+			if ( CurrentSelection.Template->Slot )
+			{
+				CurrentSelection.Template->Slot->Resize(DragDirections[CurrentHandle], MouseEvent.GetCursorDelta());
+			}
+
+			FBlueprintEditorUtils::MarkBlueprintAsModified(GetBlueprint());
+
 			return FReply::Handled().PreventThrottling();
 		}
 	}
@@ -409,7 +424,7 @@ int32 SUMGDesigner::OnPaint(const FGeometry& AllottedGeometry, const FSlateRect&
 
 void SUMGDesigner::DrawDragHandles(const FPaintGeometry& SelectionGeometry, const FSlateRect& MyClippingRect, FSlateWindowElementList& OutDrawElements, int32 LayerId, const FWidgetStyle& InWidgetStyle) const
 {
-	if ( SelectedTemplate && SelectedTemplate->Slot )
+	if ( CurrentSelection.IsValid() && CurrentSelection.Template->Slot )
 	{
 		float X = SelectionGeometry.DrawPosition.X;
 		float Y = SelectionGeometry.DrawPosition.Y;
@@ -438,7 +453,7 @@ void SUMGDesigner::DrawDragHandles(const FPaintGeometry& SelectionGeometry, cons
 		for ( int32 HandleIndex = 0; HandleIndex < Handles.Num(); HandleIndex++ )
 		{
 			const FVector2D& Handle = Handles[HandleIndex];
-			if ( !SelectedTemplate->Slot->CanResize(DragDirections[HandleIndex]) )
+			if ( !CurrentSelection.Template->Slot->CanResize(DragDirections[HandleIndex]) )
 			{
 				// This isn't a valid handle
 				continue;
@@ -461,7 +476,7 @@ SUMGDesigner::DragHandle SUMGDesigner::HitTestDragHandles(const FGeometry& Allot
 {
 	FVector2D LocalPointer = AllottedGeometry.AbsoluteToLocal(PointerEvent.GetScreenSpacePosition());
 
-	if ( SelectedTemplate && SelectedTemplate->Slot && SelectedWidget.IsValid() )
+	if ( CurrentSelection.IsValid() && CurrentSelection.Template->Slot && SelectedWidget.IsValid() )
 	{
 		TSharedRef<SWidget> Widget = SelectedWidget.Pin().ToSharedRef();
 
@@ -494,7 +509,7 @@ SUMGDesigner::DragHandle SUMGDesigner::HitTestDragHandles(const FGeometry& Allot
 			FSlateRect Rect(FVector2D(Handle.X - HandleSize.X * 0.5f, Handle.Y - HandleSize.Y * 0.5f), Handle + HandleSize);
 			if ( Rect.ContainsPoint(LocalPointer) )
 			{
-				if ( !SelectedTemplate->Slot->CanResize(DragDirections[i]) )
+				if ( !CurrentSelection.Template->Slot->CanResize(DragDirections[i]) )
 				{
 					// This isn't a valid handle
 					break;
@@ -565,7 +580,7 @@ void SUMGDesigner::Tick(const FGeometry& AllottedGeometry, const double InCurren
 	}
 
 	// Update the selected widget to match the selected template.
-	if ( SelectedTemplate )
+	if ( CurrentSelection.IsValid() )
 	{
 		if ( WidgetActor )
 		{
@@ -573,7 +588,7 @@ void SUMGDesigner::Tick(const FGeometry& AllottedGeometry, const double InCurren
 			//@TODO UMG Don't store the transient Widget as the selected object, store the template component and lookup the widget from that.
 			//SelectedWidget = ArrangedWidget.Widget;
 
-			SelectedWidget = WidgetActor->GetWidgetFromName(SelectedTemplate->GetName());
+			SelectedWidget = WidgetActor->GetWidgetFromName(CurrentSelection.Template->GetName());
 		}
 	}
 
@@ -612,13 +627,13 @@ FReply SUMGDesigner::OnDrop(const FGeometry& MyGeometry, const FDragDropEvent& D
 	if ( DragDropOp.IsValid() )
 	{
 		FArrangedWidget ArrangedWidget(SNullWidget::NullWidget, FGeometry());
-		UWidget* Template = GetTemplateAtCursor(MyGeometry, DragDropEvent, ArrangedWidget);
+		FSelectedWidget Selection = GetWidgetAtCursor(MyGeometry, DragDropEvent, ArrangedWidget);
 
 		UWidgetBlueprint* BP = GetBlueprint();
 
-		if ( Template && Template->IsA(UPanelWidget::StaticClass()) )
+		if ( Selection.IsValid() && Selection.Template->IsA(UPanelWidget::StaticClass()) )
 		{
-			UPanelWidget* Parent = Cast<UPanelWidget>(Template);
+			UPanelWidget* Parent = Cast<UPanelWidget>(Selection.Template);
 
 			UWidget* Widget = DragDropOp->Template->Create(BP->WidgetTree);
 
@@ -628,7 +643,7 @@ FReply SUMGDesigner::OnDrop(const FGeometry& MyGeometry, const FDragDropEvent& D
 			//@TODO UMG We may need a desired size canvas, where the slots have no size, they only give you position, alternatively, maybe slots that don't clip, so center is still easy.
 
 			// Update the selected template to be the newly created one.
-			SelectedTemplate = Widget;
+			//SelectedTemplate = Widget;
 
 			FBlueprintEditorUtils::MarkBlueprintAsStructurallyModified(BP);
 
@@ -638,7 +653,7 @@ FReply SUMGDesigner::OnDrop(const FGeometry& MyGeometry, const FDragDropEvent& D
 		{
 			// No existing templates so just create it and make it the root widget.
 			UWidget* Widget = DragDropOp->Template->Create(BP->WidgetTree);
-			SelectedTemplate = Widget;
+			//SelectedTemplate = Widget;
 
 			FBlueprintEditorUtils::MarkBlueprintAsStructurallyModified(BP);
 
diff --git a/Engine/Source/Editor/UMGEditor/Private/SUMGDesigner.h b/Engine/Source/Editor/UMGEditor/Private/SUMGDesigner.h
--- a/Engine/Source/Editor/UMGEditor/Private/SUMGDesigner.h
+++ b/Engine/Source/Editor/UMGEditor/Private/SUMGDesigner.h
@@ -5,6 +5,7 @@
 #include "SCompoundWidget.h"
 #include "SNodePanel.h"
 #include "SDesignSurface.h"
+#include "DesignerExtension.h"
 
 class FDesignerExtension;
 
@@ -37,7 +38,7 @@ public:
 	void Register(TSharedRef<FDesignerExtension> Extension);
 
 private:
-	UWidget* GetTemplateAtCursor(const FGeometry& MyGeometry, const FPointerEvent& MouseEvent, FArrangedWidget& ArrangedWidget);
+	FSelectedWidget GetWidgetAtCursor(const FGeometry& MyGeometry, const FPointerEvent& MouseEvent, FArrangedWidget& ArrangedWidget);
 
 	UWidgetBlueprint* GetBlueprint() const;
 	void OnBlueprintChanged(UBlueprint* InBlueprint);
@@ -87,7 +88,7 @@ private:
 	TWeakPtr<SWidget> PreviewWidget;
 	UUserWidget* PreviewWidgetActor;
 
-	UWidget* SelectedTemplate;
+	FSelectedWidget CurrentSelection;
 	TWeakPtr<SWidget> SelectedWidget;
 
 	TSharedPtr<SBorder> PreviewSurface;
diff --git a/Engine/Source/Editor/UMGEditor/Private/UMGEditorPrivatePCH.h b/Engine/Source/Editor/UMGEditor/Private/UMGEditorPrivatePCH.h
--- a/Engine/Source/Editor/UMGEditor/Private/UMGEditorPrivatePCH.h
+++ b/Engine/Source/Editor/UMGEditor/Private/UMGEditorPrivatePCH.h
@@ -7,6 +7,8 @@
 #include "UnrealEd.h"
 #include "BlueprintGraphDefinitions.h"
 
+#include "BlueprintEditorUtils.h"
+
 #include "UMG.h"
 
 #include "UMGEditorClasses.h"
@@ -15,4 +17,7 @@
 #include "AssetTypeActions_Blueprint.h"
 
 #include "DesignerExtension.h"
-#include "VerticalSlotExtension.h"
\ No newline at end of file
+#include "VerticalSlotExtension.h"
+#include "HorizontalSlotExtension.h"
+#include "CanvasSlotExtension.h"
+#include "UniformGridSlotExtension.h"
diff --git a/Engine/Source/Editor/UMGEditor/Public/DesignerExtension.h b/Engine/Source/Editor/UMGEditor/Public/DesignerExtension.h
--- a/Engine/Source/Editor/UMGEditor/Public/DesignerExtension.h
+++ b/Engine/Source/Editor/UMGEditor/Public/DesignerExtension.h
@@ -6,12 +6,29 @@ class UWidget;
 class SWidget;
 class UWidgetBlueprint;
 
+struct FSelectedWidget
+{
+public:
+	FSelectedWidget()
+		: Preview(NULL)
+		, Template(NULL)
+	{ }
+
+	bool IsValid() const
+	{
+		return Template != NULL && Preview != NULL;
+	}
+
+	UWidget* Template;
+	UWidget* Preview;
+};
+
 class FDesignerExtension : public TSharedFromThis<FDesignerExtension>
 {
 public:
 	virtual void Initialize(UWidgetBlueprint* InBlueprint);
 
-	virtual void BuildWidgetsForSelection(const TArray< UWidget* >& Selection, TArray< TSharedRef<SWidget> >& Widgets) = 0;
+	virtual void BuildWidgetsForSelection(const TArray< FSelectedWidget >& Selection, TArray< TSharedRef<SWidget> >& Widgets) = 0;
 
 	FName GetExtensionId() const;
 
diff --git a/Engine/Source/Runtime/UMG/Private/Components/CanvasPanelComponent.cpp b/Engine/Source/Runtime/UMG/Private/Components/CanvasPanelComponent.cpp
--- a/Engine/Source/Runtime/UMG/Private/Components/CanvasPanelComponent.cpp
+++ b/Engine/Source/Runtime/UMG/Private/Components/CanvasPanelComponent.cpp
@@ -69,14 +69,7 @@ TSharedRef<SWidget> UCanvasPanelComponent::RebuildWidget()
 				Slots[SlotIndex] = Slot = ConstructObject<UCanvasPanelSlot>(UCanvasPanelSlot::StaticClass(), this);
 			}
 
-			auto NewSlot = Canvas->AddSlot()
-				.Position(Slot->Position)
-				.Size(Slot->Size)
-				.HAlign(Slot->HorizontalAlignment)
-				.VAlign(Slot->VerticalAlignment)
-				[
-					Slot->Content == NULL ? SNullWidget::NullWidget : Slot->Content->GetWidget()
-				];
+			Slot->BuildSlot(Canvas.ToSharedRef());
 		}
 	}
 
@@ -87,6 +80,7 @@ UCanvasPanelSlot* UCanvasPanelComponent::AddSlot(UWidget* Content)
 {
 	UCanvasPanelSlot* Slot = ConstructObject<UCanvasPanelSlot>(UCanvasPanelSlot::StaticClass(), this);
 	Slot->Content = Content;
+	Slot->Parent = this;
 
 #if WITH_EDITOR
 	Content->Slot = Slot;
@@ -112,6 +106,7 @@ void UCanvasPanelComponent::ConnectEditorData()
 {
 	for ( UCanvasPanelSlot* Slot : Slots )
 	{
+		Slot->Parent = this;
 		Slot->Content->Slot = Slot;
 	}
 }
@@ -128,18 +123,6 @@ void UCanvasPanelComponent::PostEditChangeProperty(struct FPropertyChangedEvent&
 		{
 			Slots[SlotIndex] = ConstructObject<UCanvasPanelSlot>(UCanvasPanelSlot::StaticClass(), this);
 		}
-
-		//FName& SlotName = Slots[SlotIndex].SlotName;
-
-		//if ((SlotName == NAME_None) || UniqueSlotNames.Contains(SlotName))
-		//{
-		//	do 
-		//	{
-		//		SlotName = FName(TEXT("Slot"), SlotNumbering++);
-		//	} while (UniqueSlotNames.Contains(SlotName));
-		//}
-
-		//UniqueSlotNames.Add(SlotName);
 	}
 }
 #endif
\ No newline at end of file
diff --git a/Engine/Source/Runtime/UMG/Private/Components/CanvasPanelSlot.cpp b/Engine/Source/Runtime/UMG/Private/Components/CanvasPanelSlot.cpp
--- a/Engine/Source/Runtime/UMG/Private/Components/CanvasPanelSlot.cpp
+++ b/Engine/Source/Runtime/UMG/Private/Components/CanvasPanelSlot.cpp
@@ -7,6 +7,7 @@
 
 UCanvasPanelSlot::UCanvasPanelSlot(const FPostConstructInitializeProperties& PCIP)
 	: Super(PCIP)
+	, Slot(NULL)
 {
 	Position = FVector2D::ZeroVector;
 	Size = FVector2D(1.0f, 1.0f);
@@ -14,6 +15,18 @@ UCanvasPanelSlot::UCanvasPanelSlot(const FPostConstructInitializeProperties& PCI
 	VerticalAlignment = VAlign_Top;
 }
 
+void UCanvasPanelSlot::BuildSlot(TSharedRef<SCanvas> Canvas)
+{
+	Slot = &Canvas->AddSlot()
+		.Position(Position)
+		.Size(Size)
+		.HAlign(HorizontalAlignment)
+		.VAlign(VerticalAlignment)
+		[
+			Content == NULL ? SNullWidget::NullWidget : Content->GetWidget()
+		];
+}
+
 void UCanvasPanelSlot::Resize(const FVector2D& Direction, const FVector2D& Amount)
 {
 	if ( Direction.X < 0 )
@@ -34,9 +47,24 @@ void UCanvasPanelSlot::Resize(const FVector2D& Direction, const FVector2D& Amoun
 	{
 		Size.Y += Amount.Y * Direction.Y;
 	}
+
+	if ( Slot )
+	{
+		Slot->Position(Position);
+		Slot->Size(Size);
+	}
 }
 
 bool UCanvasPanelSlot::CanResize(const FVector2D& Direction) const
 {
 	return true;
 }
+
+void UCanvasPanelSlot::SetPosition(FVector2D InPosition)
+{
+	Position = InPosition;
+	if ( Slot )
+	{
+		Slot->Position(InPosition);
+	}
+}
diff --git a/Engine/Source/Runtime/UMG/Private/Components/ComboBox.cpp b/Engine/Source/Runtime/UMG/Private/Components/ComboBox.cpp
new file mode 100644
--- /dev/null
+++ b/Engine/Source/Runtime/UMG/Private/Components/ComboBox.cpp
@@ -0,0 +1,43 @@
+// Copyright 1998-2014 Epic Games, Inc. All Rights Reserved.
+
+#include "UMGPrivatePCH.h"
+
+#define LOCTEXT_NAMESPACE "UMG"
+
+/////////////////////////////////////////////////////
+// UComboBox
+
+UComboBox::UComboBox(const FPostConstructInitializeProperties& PCIP)
+	: Super(PCIP)
+{
+}
+
+TSharedRef<SWidget> UComboBox::RebuildWidget()
+{
+	TSharedRef< SComboBox<UObject*> > NewComboBox =
+		SNew(SComboBox<UObject*>)
+		.OptionsSource(&Items)
+		.OnGenerateWidget(BIND_UOBJECT_DELEGATE(SComboBox<UObject*>::FOnGenerateWidget, HandleGenerateWidget));
+
+	return NewComboBox;
+}
+
+TSharedRef<SWidget> UComboBox::HandleGenerateWidget(UObject* Item) const
+{
+	// Call the user's delegate to see if they want to generate a custom widget bound to the data source.
+	if ( OnGenerateWidget.IsBound() )
+	{
+		UWidget* Widget = OnGenerateWidget.Execute(Item);
+		if ( Widget != NULL )
+		{
+			return Widget->GetWidget();
+		}
+	}
+
+	// If a row wasn't generated just create the default one, a simple text block of the item's name.
+	return SNew(STextBlock).Text(Item ? FText::FromString(Item->GetName()) : LOCTEXT("null", "null"));
+}
+
+/////////////////////////////////////////////////////
+
+#undef LOCTEXT_NAMESPACE
diff --git a/Engine/Source/Runtime/UMG/Private/Components/HorizontalBoxComponent.cpp b/Engine/Source/Runtime/UMG/Private/Components/HorizontalBoxComponent.cpp
--- a/Engine/Source/Runtime/UMG/Private/Components/HorizontalBoxComponent.cpp
+++ b/Engine/Source/Runtime/UMG/Private/Components/HorizontalBoxComponent.cpp
@@ -21,6 +21,21 @@ UWidget* UHorizontalBoxComponent::GetChildAt(int32 Index) const
 	return Slots[Index]->Content;
 }
 
+int32 UHorizontalBoxComponent::GetChildIndex(UWidget* Content) const
+{
+	for ( int32 SlotIndex = 0; SlotIndex < Slots.Num(); ++SlotIndex )
+	{
+		UHorizontalBoxSlot* Slot = Slots[SlotIndex];
+
+		if ( Slot->Content == Content )
+		{
+			return SlotIndex;
+		}
+	}
+
+	return -1;
+}
+
 bool UHorizontalBoxComponent::AddChild(UWidget* Child, FVector2D Position)
 {
 	AddSlot(Child);
@@ -29,15 +44,11 @@ bool UHorizontalBoxComponent::AddChild(UWidget* Child, FVector2D Position)
 
 bool UHorizontalBoxComponent::RemoveChild(UWidget* Child)
 {
-	for ( int32 SlotIndex = 0; SlotIndex < Slots.Num(); ++SlotIndex )
+	int32 SlotIndex = GetChildIndex(Child);
+	if ( SlotIndex != -1 )
 	{
-		UHorizontalBoxSlot* Slot = Slots[SlotIndex];
-
-		if ( Slot->Content == Child )
-		{
-			Slots.RemoveAt(SlotIndex);
-			return true;
-		}
+		Slots.RemoveAt(SlotIndex);
+		return true;
 	}
 
 	return false;
@@ -53,6 +64,19 @@ void UHorizontalBoxComponent::ReplaceChildAt(int32 Index, UWidget* Content)
 #endif
 }
 
+void UHorizontalBoxComponent::InsertChildAt(int32 Index, UWidget* Content)
+{
+	UHorizontalBoxSlot* Slot = ConstructObject<UHorizontalBoxSlot>(UHorizontalBoxSlot::StaticClass(), this);
+	Slot->Content = Content;
+	Slot->Parent = this;
+
+#if WITH_EDITOR
+	Content->Slot = Slot;
+#endif
+
+	Slots.Insert(Slot, Index);
+}
+
 TSharedRef<SWidget> UHorizontalBoxComponent::RebuildWidget()
 {
 	TSharedRef<SHorizontalBox> NewCanvas = SNew(SHorizontalBox);
@@ -90,6 +114,7 @@ UHorizontalBoxSlot* UHorizontalBoxComponent::AddSlot(UWidget* Content)
 {
 	UHorizontalBoxSlot* Slot = ConstructObject<UHorizontalBoxSlot>(UHorizontalBoxSlot::StaticClass(), this);
 	Slot->Content = Content;
+	Slot->Parent = this;
 
 #if WITH_EDITOR
 	Content->Slot = Slot;
@@ -106,6 +131,7 @@ void UHorizontalBoxComponent::ConnectEditorData()
 	for ( UHorizontalBoxSlot* Slot : Slots )
 	{
 		Slot->Content->Slot = Slot;
+		Slot->Parent = this;
 	}
 }
 
diff --git a/Engine/Source/Runtime/UMG/Private/Components/UniformGridPanel.cpp b/Engine/Source/Runtime/UMG/Private/Components/UniformGridPanel.cpp
--- a/Engine/Source/Runtime/UMG/Private/Components/UniformGridPanel.cpp
+++ b/Engine/Source/Runtime/UMG/Private/Components/UniformGridPanel.cpp
@@ -66,12 +66,7 @@ TSharedRef<SWidget> UUniformGridPanel::RebuildWidget()
 			Slots[SlotIndex] = Slot = ConstructObject<UUniformGridSlot>(UUniformGridSlot::StaticClass(), this);
 		}
 
-		auto& NewSlot = NewPanel->AddSlot(Slot->Column, Slot->Row)
-			.HAlign(Slot->HorizontalAlignment)
-			.VAlign(Slot->VerticalAlignment)
-			[
-				Slot->Content == NULL ? SNullWidget::NullWidget : Slot->Content->GetWidget()
-			];
+		Slot->BuildSlot(NewPanel);
 	}
 
 	return NewPanel;
@@ -81,6 +76,7 @@ UUniformGridSlot* UUniformGridPanel::AddSlot(UWidget* Content)
 {
 	UUniformGridSlot* Slot = ConstructObject<UUniformGridSlot>(UUniformGridSlot::StaticClass(), this);
 	Slot->Content = Content;
+	Slot->Parent = this;
 
 #if WITH_EDITOR
 	Content->Slot = Slot;
@@ -98,6 +94,7 @@ void UUniformGridPanel::ConnectEditorData()
 	for ( UUniformGridSlot* Slot : Slots )
 	{
 		Slot->Content->Slot = Slot;
+		Slot->Parent = this;
 	}
 }
 
diff --git a/Engine/Source/Runtime/UMG/Private/Components/UniformGridSlot.cpp b/Engine/Source/Runtime/UMG/Private/Components/UniformGridSlot.cpp
--- a/Engine/Source/Runtime/UMG/Private/Components/UniformGridSlot.cpp
+++ b/Engine/Source/Runtime/UMG/Private/Components/UniformGridSlot.cpp
@@ -7,7 +7,36 @@
 
 UUniformGridSlot::UUniformGridSlot(const FPostConstructInitializeProperties& PCIP)
 	: Super(PCIP)
+	, Slot(NULL)
 {
 	HorizontalAlignment = HAlign_Left;
 	VerticalAlignment = VAlign_Top;
 }
+
+void UUniformGridSlot::BuildSlot(TSharedRef<SUniformGridPanel> GridPanel)
+{
+	Slot = &GridPanel->AddSlot(Column, Row)
+		.HAlign(HorizontalAlignment)
+		.VAlign(VerticalAlignment)
+		[
+			Content == NULL ? SNullWidget::NullWidget : Content->GetWidget()
+		];
+}
+
+void UUniformGridSlot::SetRow(int32 InRow)
+{
+	Row = InRow;
+	if ( Slot )
+	{
+		Slot->Row = InRow;
+	}
+}
+
+void UUniformGridSlot::SetColumn(int32 InColumn)
+{
+	Column = InColumn;
+	if ( Slot )
+	{
+		Slot->Column = InColumn;
+	}
+}
diff --git a/Engine/Source/Runtime/UMG/Private/Components/VerticalBoxComponent.cpp b/Engine/Source/Runtime/UMG/Private/Components/VerticalBoxComponent.cpp
--- a/Engine/Source/Runtime/UMG/Private/Components/VerticalBoxComponent.cpp
+++ b/Engine/Source/Runtime/UMG/Private/Components/VerticalBoxComponent.cpp
@@ -96,16 +96,7 @@ TSharedRef<SWidget> UVerticalBoxComponent::RebuildWidget()
 			}
 
 			Slot->Parent = this;
-
-			auto& NewSlot = Canvas->AddSlot()
-				.Padding(Slot->Padding)
-				.HAlign(Slot->HorizontalAlignment)
-				.VAlign(Slot->VerticalAlignment)
-				[
-					Slot->Content == NULL ? SNullWidget::NullWidget : Slot->Content->GetWidget()
-				];
-
-			NewSlot.SizeParam = UWidget::ConvertSerializedSizeParamToRuntime(Slot->Size);
+			Slot->BuildSlot(NewCanvas);
 		}
 	}
 
diff --git a/Engine/Source/Runtime/UMG/Private/Components/VerticalBoxSlot.cpp b/Engine/Source/Runtime/UMG/Private/Components/VerticalBoxSlot.cpp
--- a/Engine/Source/Runtime/UMG/Private/Components/VerticalBoxSlot.cpp
+++ b/Engine/Source/Runtime/UMG/Private/Components/VerticalBoxSlot.cpp
@@ -7,7 +7,21 @@
 
 UVerticalBoxSlot::UVerticalBoxSlot(const FPostConstructInitializeProperties& PCIP)
 	: Super(PCIP)
+	, Slot(NULL)
 {
 	HorizontalAlignment = HAlign_Left;
 	VerticalAlignment = VAlign_Top;
 }
+
+void UVerticalBoxSlot::BuildSlot(TSharedRef<SVerticalBox> VerticalBox)
+{
+	Slot = &VerticalBox->AddSlot()
+		.Padding(Padding)
+		.HAlign(HorizontalAlignment)
+		.VAlign(VerticalAlignment)
+		[
+			Content == NULL ? SNullWidget::NullWidget : Content->GetWidget()
+		];
+
+	Slot->SizeParam = UWidget::ConvertSerializedSizeParamToRuntime(Size);
+}
diff --git a/Engine/Source/Runtime/UMG/Public/Components/CanvasPanelSlot.h b/Engine/Source/Runtime/UMG/Public/Components/CanvasPanelSlot.h
--- a/Engine/Source/Runtime/UMG/Public/Components/CanvasPanelSlot.h
+++ b/Engine/Source/Runtime/UMG/Public/Components/CanvasPanelSlot.h
@@ -48,7 +48,14 @@ class UMG_API UCanvasPanelSlot : public UPanelSlot
 	UPROPERTY(EditAnywhere, Category=Appearance)
 	TEnumAsByte<EVerticalAlignment> VerticalAlignment;
 
+	void BuildSlot(TSharedRef<SCanvas> Canvas);
+
 	virtual void Resize(const FVector2D& Direction, const FVector2D& Amount) OVERRIDE;
 
 	virtual bool CanResize(const FVector2D& Direction) const OVERRIDE;
+
+	void SetPosition(FVector2D Position);
+
+private:
+	SCanvas::FSlot* Slot;
 };
\ No newline at end of file
diff --git a/Engine/Source/Runtime/UMG/Public/Components/ComboBox.h b/Engine/Source/Runtime/UMG/Public/Components/ComboBox.h
new file mode 100644
--- /dev/null
+++ b/Engine/Source/Runtime/UMG/Public/Components/ComboBox.h
@@ -0,0 +1,38 @@
+// Copyright 1998-2014 Epic Games, Inc. All Rights Reserved.
+
+#pragma once
+
+#include "ComboBox.generated.h"
+
+UCLASS(meta=(BlueprintSpawnableComponent), ClassGroup=UserInterface)
+class UMG_API UComboBox : public UWidget
+{
+	GENERATED_UCLASS_BODY()
+
+public:
+
+	UPROPERTY(EditAnywhere, BlueprintReadWrite, Category=Content)
+	TArray<UObject*> Items;
+
+	/** Called when the widget is needed for the item. */
+	UPROPERTY(EditDefaultsOnly, Category=Events)
+	FGenerateWidgetUObject OnGenerateWidget;
+	
+protected:
+	TSharedPtr< SComboBox<UObject*> > ComboBoxWidget() const
+	{
+		if ( MyWidget.IsValid() )
+		{
+			return StaticCastSharedRef< SComboBox<UObject*> >(MyWidget.ToSharedRef());
+		}
+
+		return TSharedPtr< SComboBox<UObject*> >();
+	}
+
+protected:
+	TSharedRef<SWidget> HandleGenerateWidget(UObject* Item) const;
+
+	// UWidget interface
+	virtual TSharedRef<SWidget> RebuildWidget() OVERRIDE;
+	// End of UWidget interface
+};
diff --git a/Engine/Source/Runtime/UMG/Public/Components/HorizontalBoxComponent.h b/Engine/Source/Runtime/UMG/Public/Components/HorizontalBoxComponent.h
--- a/Engine/Source/Runtime/UMG/Public/Components/HorizontalBoxComponent.h
+++ b/Engine/Source/Runtime/UMG/Public/Components/HorizontalBoxComponent.h
@@ -18,9 +18,11 @@ class UMG_API UHorizontalBoxComponent : public UPanelWidget
 	// UPanelWidget
 	virtual int32 GetChildrenCount() const OVERRIDE;
 	virtual UWidget* GetChildAt(int32 Index) const OVERRIDE;
+	virtual int32 GetChildIndex(UWidget* Content) const OVERRIDE;
 	virtual bool AddChild(UWidget* Child, FVector2D Position) OVERRIDE;
 	virtual bool RemoveChild(UWidget* Child) OVERRIDE;
 	virtual void ReplaceChildAt(int32 Index, UWidget* Child) OVERRIDE;
+	virtual void InsertChildAt(int32 Index, UWidget* Child) OVERRIDE;
 	// End UPanelWidget
 
 #if WITH_EDITOR
diff --git a/Engine/Source/Runtime/UMG/Public/Components/UniformGridSlot.h b/Engine/Source/Runtime/UMG/Public/Components/UniformGridSlot.h
--- a/Engine/Source/Runtime/UMG/Public/Components/UniformGridSlot.h
+++ b/Engine/Source/Runtime/UMG/Public/Components/UniformGridSlot.h
@@ -22,4 +22,15 @@ class UMG_API UUniformGridSlot : public UPanelSlot
 	
 	UPROPERTY(EditAnywhere, Category=Appearance)
 	int32 Column;
+
+	UFUNCTION(BlueprintCallable, Category=Layout)
+	void SetRow(int32 InRow);
+
+	UFUNCTION(BlueprintCallable, Category=Layout)
+	void SetColumn(int32 InColumn);
+
+	void BuildSlot(TSharedRef<SUniformGridPanel> GridPanel);
+
+private:
+	SUniformGridPanel::FSlot* Slot;
 };
\ No newline at end of file
diff --git a/Engine/Source/Runtime/UMG/Public/Components/VerticalBoxSlot.h b/Engine/Source/Runtime/UMG/Public/Components/VerticalBoxSlot.h
--- a/Engine/Source/Runtime/UMG/Public/Components/VerticalBoxSlot.h
+++ b/Engine/Source/Runtime/UMG/Public/Components/VerticalBoxSlot.h
@@ -23,4 +23,9 @@ class UMG_API UVerticalBoxSlot : public UPanelSlot
 
 	UPROPERTY(EditAnywhere, Category=Appearance)
 	TEnumAsByte<EVerticalAlignment> VerticalAlignment;
+
+	void BuildSlot(TSharedRef<SVerticalBox> InVerticalBox);
+
+private:
+	SVerticalBox::FSlot* Slot;
 };
diff --git a/Engine/Source/Runtime/UMG/Public/Components/Widget.h b/Engine/Source/Runtime/UMG/Public/Components/Widget.h
--- a/Engine/Source/Runtime/UMG/Public/Components/Widget.h
+++ b/Engine/Source/Runtime/UMG/Public/Components/Widget.h
@@ -21,6 +21,8 @@ class UMG_API UWidget : public UObject
 	DECLARE_DYNAMIC_DELEGATE_RetVal(FLinearColor, FGetLinearColor);
 	DECLARE_DYNAMIC_DELEGATE_RetVal(ESlateVisibility::Type, FGetSlateVisibility);
 
+	DECLARE_DYNAMIC_DELEGATE_RetVal_OneParam(UWidget*, FGenerateWidgetUObject, UObject*, Item);
+
 	// Events
 	DECLARE_DYNAMIC_DELEGATE_RetVal(FSReply, FOnReply);
 
diff --git a/Engine/Source/Runtime/UMG/Public/UMG.h b/Engine/Source/Runtime/UMG/Public/UMG.h
--- a/Engine/Source/Runtime/UMG/Public/UMG.h
+++ b/Engine/Source/Runtime/UMG/Public/UMG.h
@@ -55,6 +55,7 @@
 #include "EditableTextBlockComponent.h"
 #include "ImageComponent.h"
 #include "Slider.h"
+#include "ComboBox.h"
 
 #include "GridPanelComponent.h"
 #include "ScrollBoxComponent.h"
