UMG - Organizing widget templates.  Working on customizations for the canvas.

diff --git a/Engine/Source/Editor/EditorStyle/Private/SlateEditorStyle.cpp b/Engine/Source/Editor/EditorStyle/Private/SlateEditorStyle.cpp
--- a/Engine/Source/Editor/EditorStyle/Private/SlateEditorStyle.cpp
+++ b/Engine/Source/Editor/EditorStyle/Private/SlateEditorStyle.cpp
@@ -138,6 +138,7 @@ void FSlateEditorStyle::FStyle::Initialize()
 	SetupMatineeStyle();
 	SetupSourceControlStyles();
 	SetupAutomationStyles();
+	SetupUMGEditorStyles();
 
 //	LogUnusedBrushResources();
 }
@@ -5165,6 +5166,33 @@ void FSlateEditorStyle::FStyle::SetupAutomationStyles()
 	}
 }
 
+void FSlateEditorStyle::FStyle::SetupUMGEditorStyles()
+{
+	Set("UMGEditor.AnchorCenter", new IMAGE_BRUSH("Icons/umg_anchor_center", Icon16x16));
+	Set("UMGEditor.AnchorTopBottom", new IMAGE_BRUSH("Icons/umg_anchor_top_bottom", FVector2D(16, 32)));
+	Set("UMGEditor.AnchorLeftRight", new IMAGE_BRUSH("Icons/umg_anchor_left_right", FVector2D(32, 16)));
+
+	Set("UMGEditor.TopLeft", new IMAGE_BRUSH("Icons/UMG/umg_anchors_top_left", Icon64x64));
+	Set("UMGEditor.TopCenter", new IMAGE_BRUSH("Icons/UMG/umg_anchors_top_center", Icon64x64));
+	Set("UMGEditor.TopRight", new IMAGE_BRUSH("Icons/UMG/umg_anchors_top_right", Icon64x64));
+	Set("UMGEditor.TopFill", new IMAGE_BRUSH("Icons/UMG/umg_anchors_top_fill", Icon64x64));
+
+	Set("UMGEditor.CenterLeft", new IMAGE_BRUSH("Icons/UMG/umg_anchors_center_left", Icon64x64));
+	Set("UMGEditor.CenterCenter", new IMAGE_BRUSH("Icons/UMG/umg_anchors_center_center", Icon64x64));
+	Set("UMGEditor.CenterRight", new IMAGE_BRUSH("Icons/UMG/umg_anchors_center_right", Icon64x64));
+	Set("UMGEditor.CenterFill", new IMAGE_BRUSH("Icons/UMG/umg_anchors_center_fill", Icon64x64));
+
+	Set("UMGEditor.BottomLeft", new IMAGE_BRUSH("Icons/UMG/umg_anchors_bottom_left", Icon64x64));
+	Set("UMGEditor.BottomCenter", new IMAGE_BRUSH("Icons/UMG/umg_anchors_bottom_center", Icon64x64));
+	Set("UMGEditor.BottomRight", new IMAGE_BRUSH("Icons/UMG/umg_anchors_bottom_right", Icon64x64));
+	Set("UMGEditor.BottomFill", new IMAGE_BRUSH("Icons/UMG/umg_anchors_bottom_fill", Icon64x64));
+
+	Set("UMGEditor.FillLeft", new IMAGE_BRUSH("Icons/UMG/umg_anchors_fill_left", Icon64x64));
+	Set("UMGEditor.FillCenter", new IMAGE_BRUSH("Icons/UMG/umg_anchors_fill_center", Icon64x64));
+	Set("UMGEditor.FillRight", new IMAGE_BRUSH("Icons/UMG/umg_anchors_fill_right", Icon64x64));
+	Set("UMGEditor.FillFill", new IMAGE_BRUSH("Icons/UMG/umg_anchors_fill_fill", Icon64x64));
+}
+
 void FSlateEditorStyle::FStyle::SetupTranslationEditorStyles()
 {
 	Set("TranslationEditor.Import", new IMAGE_BRUSH("Icons/icon_Import_40x", Icon40x40));
diff --git a/Engine/Source/Editor/EditorStyle/Private/SlateEditorStyle.h b/Engine/Source/Editor/EditorStyle/Private/SlateEditorStyle.h
--- a/Engine/Source/Editor/EditorStyle/Private/SlateEditorStyle.h
+++ b/Engine/Source/Editor/EditorStyle/Private/SlateEditorStyle.h
@@ -79,6 +79,7 @@ public:
 		void SetupMatineeStyle();
 		void SetupSourceControlStyles();
 		void SetupAutomationStyles();
+		void SetupUMGEditorStyles();
 
 		void SettingsChanged( UObject* ChangedObject, FPropertyChangedEvent& PropertyChangedEvent );
 		void SyncSettings();
diff --git a/Engine/Source/Editor/UMGEditor/Private/Customizations/CanvasSlotCustomization.cpp b/Engine/Source/Editor/UMGEditor/Private/Customizations/CanvasSlotCustomization.cpp
new file mode 100644
--- /dev/null
+++ b/Engine/Source/Editor/UMGEditor/Private/Customizations/CanvasSlotCustomization.cpp
@@ -0,0 +1,218 @@
+// Copyright 1998-2014 Epic Games, Inc. All Rights Reserved.
+
+#include "UMGEditorPrivatePCH.h"
+
+#include "PropertyEditing.h"
+#include "ObjectEditorUtils.h"
+#include "WidgetGraphSchema.h"
+#include "ScopedTransaction.h"
+#include "BlueprintEditorUtils.h"
+
+#define LOCTEXT_NAMESPACE "UMG"
+
+void FCanvasSlotCustomization::CustomizeDetails( IDetailLayoutBuilder& DetailLayout )
+{
+	// Build the layout category
+	IDetailCategoryBuilder& LayoutCategory = DetailLayout.EditCategory("Layout");
+
+	TSharedRef<IPropertyHandle> AnchorsHandle = DetailLayout.GetProperty(GET_MEMBER_NAME_CHECKED(UCanvasPanelSlot, Anchors));
+	IDetailPropertyRow& AnchorsPropertyRow = LayoutCategory.AddProperty(AnchorsHandle);
+
+	const bool bShowChildren = true;
+	AnchorsPropertyRow.CustomWidget(bShowChildren)
+		.ValueContent()
+		[
+			SNew(SHorizontalBox)
+					
+			+ SHorizontalBox::Slot()
+			.AutoWidth()
+			[
+				SNew(SComboButton)
+				.ButtonContent()
+				[
+					SNew(STextBlock)
+					.Text(LOCTEXT("AnchorsText", "Anchors"))
+				]
+				.MenuContent()
+				[
+					SNew(SUniformGridPanel)
+
+					// Top Row
+					+ SUniformGridPanel::Slot(0, 0)
+					[
+						SNew(SButton)
+						.OnClicked(this, &FCanvasSlotCustomization::OnAnchorClicked, AnchorsHandle, FAnchors(0, 0, 0, 0))
+						[
+							SNew(SImage)
+							.Image(FEditorStyle::Get().GetBrush("UMGEditor.TopLeft"))
+						]
+					]
+
+					+ SUniformGridPanel::Slot(1, 0)
+					[
+						SNew(SButton)
+						.OnClicked(this, &FCanvasSlotCustomization::OnAnchorClicked, AnchorsHandle, FAnchors(0.5, 0, 0.5, 0))
+						[
+							SNew(SImage)
+							.Image(FEditorStyle::Get().GetBrush("UMGEditor.TopCenter"))
+						]
+					]
+
+					+ SUniformGridPanel::Slot(2, 0)
+					[
+						SNew(SButton)
+						.OnClicked(this, &FCanvasSlotCustomization::OnAnchorClicked, AnchorsHandle, FAnchors(1, 0, 1, 0))
+						[
+							SNew(SImage)
+							.Image(FEditorStyle::Get().GetBrush("UMGEditor.TopRight"))
+						]
+					]
+
+					+ SUniformGridPanel::Slot(3, 0)
+					[
+						SNew(SButton)
+						.OnClicked(this, &FCanvasSlotCustomization::OnAnchorClicked, AnchorsHandle, FAnchors(0, 0, 1, 0))
+						[
+							SNew(SImage)
+							.Image(FEditorStyle::Get().GetBrush("UMGEditor.TopFill"))
+						]
+					]
+
+					// Center Row
+					+ SUniformGridPanel::Slot(0, 1)
+					[
+						SNew(SButton)
+						.OnClicked(this, &FCanvasSlotCustomization::OnAnchorClicked, AnchorsHandle, FAnchors(0, 0.5, 0, 0.5))
+						[
+							SNew(SImage)
+							.Image(FEditorStyle::Get().GetBrush("UMGEditor.CenterLeft"))
+						]
+					]
+
+					+ SUniformGridPanel::Slot(1, 1)
+					[
+						SNew(SButton)
+						.OnClicked(this, &FCanvasSlotCustomization::OnAnchorClicked, AnchorsHandle, FAnchors(0.5, 0.5, 0.5, 0.5))
+						[
+							SNew(SImage)
+							.Image(FEditorStyle::Get().GetBrush("UMGEditor.CenterCenter"))
+						]
+					]
+
+					+ SUniformGridPanel::Slot(2, 1)
+					[
+						SNew(SButton)
+						.OnClicked(this, &FCanvasSlotCustomization::OnAnchorClicked, AnchorsHandle, FAnchors(1, 0.5, 1, 0.5))
+						[
+							SNew(SImage)
+							.Image(FEditorStyle::Get().GetBrush("UMGEditor.CenterRight"))
+						]
+					]
+
+					+ SUniformGridPanel::Slot(3, 1)
+					[
+						SNew(SButton)
+						.OnClicked(this, &FCanvasSlotCustomization::OnAnchorClicked, AnchorsHandle, FAnchors(0, 0.5f, 1, 0.5f))
+						[
+							SNew(SImage)
+							.Image(FEditorStyle::Get().GetBrush("UMGEditor.CenterFill"))
+						]
+					]
+
+					// Bottom Row
+					+ SUniformGridPanel::Slot(0, 2)
+					[
+						SNew(SButton)
+						.OnClicked(this, &FCanvasSlotCustomization::OnAnchorClicked, AnchorsHandle, FAnchors(0, 1, 0, 1))
+						[
+							SNew(SImage)
+							.Image(FEditorStyle::Get().GetBrush("UMGEditor.BottomLeft"))
+						]
+					]
+
+					+ SUniformGridPanel::Slot(1, 2)
+					[
+						SNew(SButton)
+						.OnClicked(this, &FCanvasSlotCustomization::OnAnchorClicked, AnchorsHandle, FAnchors(0.5, 1, 0.5, 1))
+						[
+							SNew(SImage)
+							.Image(FEditorStyle::Get().GetBrush("UMGEditor.BottomCenter"))
+						]
+					]
+
+					+ SUniformGridPanel::Slot(2, 2)
+					[
+						SNew(SButton)
+						.OnClicked(this, &FCanvasSlotCustomization::OnAnchorClicked, AnchorsHandle, FAnchors(1, 1, 1, 1))
+						[
+							SNew(SImage)
+							.Image(FEditorStyle::Get().GetBrush("UMGEditor.BottomRight"))
+						]
+					]
+
+					+ SUniformGridPanel::Slot(3, 2)
+					[
+						SNew(SButton)
+						.OnClicked(this, &FCanvasSlotCustomization::OnAnchorClicked, AnchorsHandle, FAnchors(0, 1, 1, 1))
+						[
+							SNew(SImage)
+							.Image(FEditorStyle::Get().GetBrush("UMGEditor.BottomFill"))
+						]
+					]
+
+					// Fill Row
+					+ SUniformGridPanel::Slot(0, 3)
+					[
+						SNew(SButton)
+						.OnClicked(this, &FCanvasSlotCustomization::OnAnchorClicked, AnchorsHandle, FAnchors(0, 0, 0, 1))
+						[
+							SNew(SImage)
+							.Image(FEditorStyle::Get().GetBrush("UMGEditor.FillLeft"))
+						]
+					]
+
+					+ SUniformGridPanel::Slot(1, 3)
+					[
+						SNew(SButton)
+						.OnClicked(this, &FCanvasSlotCustomization::OnAnchorClicked, AnchorsHandle, FAnchors(0.5, 0, 0.5, 1))
+						[
+							SNew(SImage)
+							.Image(FEditorStyle::Get().GetBrush("UMGEditor.FillCenter"))
+						]
+					]
+
+					+ SUniformGridPanel::Slot(2, 3)
+					[
+						SNew(SButton)
+						.OnClicked(this, &FCanvasSlotCustomization::OnAnchorClicked, AnchorsHandle, FAnchors(1, 0, 1, 1))
+						[
+							SNew(SImage)
+							.Image(FEditorStyle::Get().GetBrush("UMGEditor.FillRight"))
+						]
+					]
+
+					+ SUniformGridPanel::Slot(3, 3)
+					[
+						SNew(SButton)
+						.OnClicked(this, &FCanvasSlotCustomization::OnAnchorClicked, AnchorsHandle, FAnchors(1, 1, 1, 1))
+						[
+							SNew(SImage)
+							.Image(FEditorStyle::Get().GetBrush("UMGEditor.FillFill"))
+						]
+					]
+				]
+			]
+		];
+}
+
+FReply FCanvasSlotCustomization::OnAnchorClicked(TSharedRef<IPropertyHandle> AnchorsHandle, FAnchors Anchors)
+{
+	FString Value = FString::Printf(TEXT("(Minimum=(X=%f,Y=%f),Maximum=(X=%f,Y=%f))"), Anchors.Minimum.X, Anchors.Minimum.Y, Anchors.Maximum.X, Anchors.Maximum.Y);
+	AnchorsHandle->SetValueFromFormattedString(Value);
+	
+	FSlateApplication::Get().DismissAllMenus();
+
+	return FReply::Handled();
+}
+
+#undef LOCTEXT_NAMESPACE
diff --git a/Engine/Source/Editor/UMGEditor/Private/Customizations/CanvasSlotCustomization.h b/Engine/Source/Editor/UMGEditor/Private/Customizations/CanvasSlotCustomization.h
new file mode 100644
--- /dev/null
+++ b/Engine/Source/Editor/UMGEditor/Private/Customizations/CanvasSlotCustomization.h
@@ -0,0 +1,29 @@
+// Copyright 1998-2014 Epic Games, Inc. All Rights Reserved.
+
+#pragma once
+
+class FCanvasSlotCustomization : public IDetailCustomization
+{
+public:
+	/** Makes a new instance of this detail layout class for a specific detail view requesting it */
+	static TSharedRef<class IDetailCustomization> MakeInstance(UBlueprint* Blueprint)
+	{
+		return MakeShareable(new FCanvasSlotCustomization(Blueprint));
+	}
+
+	FCanvasSlotCustomization(UBlueprint* InBlueprint)
+		: Blueprint(CastChecked<UWidgetBlueprint>(InBlueprint))
+	{
+	}
+	
+	/** IDetailCustomization interface */
+	virtual void CustomizeDetails( IDetailLayoutBuilder& DetailLayout ) OVERRIDE;
+	
+private:
+	
+	FReply OnAnchorClicked(TSharedRef<IPropertyHandle> AnchorsHandle, FAnchors Anchors);
+
+private:
+	UWidgetBlueprint* Blueprint;
+	
+};
diff --git a/Engine/Source/Editor/UMGEditor/Private/Customizations/DetailCustomizations.cpp b/Engine/Source/Editor/UMGEditor/Private/Customizations/DetailCustomizations.cpp
new file mode 100644
--- /dev/null
+++ b/Engine/Source/Editor/UMGEditor/Private/Customizations/DetailCustomizations.cpp
@@ -0,0 +1,382 @@
+// Copyright 1998-2014 Epic Games, Inc. All Rights Reserved.
+
+#include "UMGEditorPrivatePCH.h"
+
+#include "DetailCustomizations.h"
+
+#include "PropertyEditing.h"
+#include "ObjectEditorUtils.h"
+#include "WidgetGraphSchema.h"
+#include "ScopedTransaction.h"
+#include "BlueprintEditorUtils.h"
+
+#define LOCTEXT_NAMESPACE "UMG"
+
+void FBlueprintWidgetCustomization::CustomizeDetails( IDetailLayoutBuilder& DetailLayout )
+{
+	TArray< TWeakObjectPtr<UObject> > OutObjects;
+	DetailLayout.GetObjectsBeingCustomized(OutObjects);
+
+	if ( OutObjects.Num() == 1 )
+	{
+		UClass* PropertyClass = OutObjects[0].Get()->GetClass();
+
+		for ( TFieldIterator<UProperty> PropertyIt(PropertyClass, EFieldIteratorFlags::IncludeSuper); PropertyIt; ++PropertyIt )
+		{
+			if ( UDelegateProperty* Property = Cast<UDelegateProperty>(*PropertyIt) )
+			{
+				if ( !Property->GetName().EndsWith(TEXT("Delegate")) )
+				{
+					continue;
+				}
+
+				FString ConstPropertyName = Property->GetName();
+				ConstPropertyName.RemoveFromEnd(TEXT("Delegate"));
+
+				//TSharedRef<IPropertyHandle> DelegatePropertyHandle = DetailLayout.GetProperty(Property->GetFName(), PropertyClass);
+				TSharedRef<IPropertyHandle> ConstPropertyHandle = DetailLayout.GetProperty(FName(*ConstPropertyName), CastChecked<UClass>(Property->GetOuter()));
+
+				const bool bHasValidHandle = ConstPropertyHandle->IsValidHandle();
+
+				if ( !bHasValidHandle )
+				{
+					continue;
+				}
+
+				UProperty* ConstProperty = ConstPropertyHandle->GetProperty();
+
+				const bool bIsEditable = ConstProperty->HasAnyPropertyFlags(CPF_Edit | CPF_EditConst);
+				const bool bDoSignaturesMatch = Property->SignatureFunction->GetReturnProperty()->SameType(ConstProperty);
+
+				if ( !(bIsEditable && bDoSignaturesMatch) )
+				{
+					continue;
+				}
+
+				IDetailCategoryBuilder& PropertyCategory = DetailLayout.EditCategory(FObjectEditorUtils::GetCategoryFName(ConstProperty));
+
+				IDetailPropertyRow& PropertyRow = PropertyCategory.AddProperty(ConstPropertyHandle);
+
+				TSharedPtr<SWidget> NameWidget;
+				TSharedPtr<SWidget> ValueWidget;
+				FDetailWidgetRow Row;
+				PropertyRow.GetDefaultWidgets(NameWidget, ValueWidget, Row);
+
+				TSharedRef<SComboButton> DelegateButton = SNew(SComboButton)
+				.OnGetMenuContent(this, &FBlueprintWidgetCustomization::OnGenerateDelegateMenu, ConstPropertyHandle, Property->SignatureFunction)
+				.ContentPadding(0)
+				.ButtonContent()
+				[
+					SNew(STextBlock)
+					.Text(this, &FBlueprintWidgetCustomization::GetCurrentBindingText, ConstPropertyHandle)
+					.Font(IDetailLayoutBuilder::GetDetailFont())
+				];
+
+				const bool bShowChildren = true;
+				PropertyRow.CustomWidget(bShowChildren)
+				.NameContent()
+				.MinDesiredWidth(Row.NameWidget.MinWidth)
+				.MaxDesiredWidth(Row.NameWidget.MaxWidth)
+				[
+					NameWidget.ToSharedRef()
+				]
+				.ValueContent()
+				.MinDesiredWidth(Row.ValueWidget.MinWidth * 2)
+				.MaxDesiredWidth(Row.ValueWidget.MaxWidth * 2)
+				[
+					SNew(SHorizontalBox)
+
+					+ SHorizontalBox::Slot()
+					.FillWidth(1)
+					.Padding(FMargin(0, 0, 4, 0))
+					[
+						ValueWidget.ToSharedRef()
+					]
+					
+					+ SHorizontalBox::Slot()
+					.AutoWidth()
+					[
+						DelegateButton
+					]
+				];
+			}
+		}
+
+		//--------------------------------------------
+
+		for ( TFieldIterator<UProperty> PropertyIt(PropertyClass, EFieldIteratorFlags::IncludeSuper); PropertyIt; ++PropertyIt )
+		{
+			if ( UDelegateProperty* Property = Cast<UDelegateProperty>(*PropertyIt) )
+			{
+				if ( !Property->GetName().EndsWith(TEXT("Event")) )
+				{
+					continue;
+				}
+
+				TSharedRef<IPropertyHandle> DelegatePropertyHandle = DetailLayout.GetProperty(Property->GetFName(), CastChecked<UClass>(Property->GetOuter()));
+
+				const bool bHasValidHandle = DelegatePropertyHandle->IsValidHandle();
+
+				if ( !bHasValidHandle )
+				{
+					continue;
+				}
+
+				IDetailCategoryBuilder& PropertyCategory = DetailLayout.EditCategory(FObjectEditorUtils::GetCategoryFName(Property));
+
+				IDetailPropertyRow& PropertyRow = PropertyCategory.AddProperty(DelegatePropertyHandle);
+
+				TSharedPtr<SWidget> NameWidget;
+				TSharedPtr<SWidget> ValueWidget;
+				FDetailWidgetRow Row;
+				PropertyRow.GetDefaultWidgets(NameWidget, ValueWidget, Row);
+
+				TSharedRef<SComboButton> DelegateButton = SNew(SComboButton)
+				.OnGetMenuContent(this, &FBlueprintWidgetCustomization::OnGenerateDelegateMenu, DelegatePropertyHandle, Property->SignatureFunction)
+				.ContentPadding(0)
+				.ButtonContent()
+				[
+					SNew(STextBlock)
+					.Text(this, &FBlueprintWidgetCustomization::GetCurrentBindingText, DelegatePropertyHandle)
+					.Font(IDetailLayoutBuilder::GetDetailFont())
+				];
+
+				const bool bShowChildren = true;
+				PropertyRow.CustomWidget(bShowChildren)
+				.NameContent()
+				.MinDesiredWidth(Row.NameWidget.MinWidth)
+				.MaxDesiredWidth(Row.NameWidget.MaxWidth)
+				[
+					NameWidget.ToSharedRef()
+				]
+				.ValueContent()
+				.MinDesiredWidth(Row.ValueWidget.MinWidth)
+				.MaxDesiredWidth(Row.ValueWidget.MaxWidth)
+				[
+					SNew(SHorizontalBox)
+					
+					+ SHorizontalBox::Slot()
+					.FillWidth(1)
+					[
+						DelegateButton
+					]
+				];
+			}
+		}
+
+	}
+}
+
+void FBlueprintWidgetCustomization::RefreshBlueprintFunctionCache(const UFunction* DelegateSignature)
+{
+	const UWidgetGraphSchema* Schema = GetDefault<UWidgetGraphSchema>();
+	const FSlateFontInfo DetailFontInfo = IDetailLayoutBuilder::GetDetailFont();
+
+	BlueprintFunctionCache.Reset();
+
+	// Get the current skeleton class, think header for the blueprint.
+	UBlueprintGeneratedClass* SkeletonClass = Cast<UBlueprintGeneratedClass>(Blueprint->SkeletonGeneratedClass);
+
+	// Grab functions implemented by the blueprint
+	for ( int32 i = 0; i < Blueprint->FunctionGraphs.Num(); i++ )
+	{
+		UEdGraph* Graph = Blueprint->FunctionGraphs[i];
+		check(Graph);
+
+		// Ignore the construction script, we just want user created functions.
+		if ( Graph->GetFName() == Schema->FN_UserConstructionScript )
+		{
+			continue;
+		}
+
+		FName FunctionFName = Graph->GetFName();
+
+		UFunction* Function = SkeletonClass->FindFunctionByName(FunctionFName, EIncludeSuperFlag::IncludeSuper);
+		if ( Function == NULL )
+		{
+			continue;
+		}
+
+		// We ignore CPF_ReturnParm because all that matters for binding to script functions is that the number of out parameters match.
+		if ( Function->IsSignatureCompatibleWith(DelegateSignature, UFunction::GetDefaultIgnoredSignatureCompatibilityFlags() | CPF_ReturnParm) )
+		{
+			FGraphDisplayInfo DisplayInfo;
+			Graph->GetSchema()->GetGraphDisplayInformation(*Graph, DisplayInfo);
+
+			TSharedPtr<FunctionInfo> NewFuncAction = MakeShareable(new FunctionInfo());
+			NewFuncAction->DisplayName = DisplayInfo.PlainName;
+			NewFuncAction->Tooltip = DisplayInfo.Tooltip;
+			NewFuncAction->FuncName = FunctionFName;
+			NewFuncAction->EdGraph = Graph;
+
+			BlueprintFunctionCache.Add(NewFuncAction);
+		}
+	}
+}
+
+TSharedRef<SWidget> FBlueprintWidgetCustomization::OnGenerateDelegateMenu(TSharedRef<IPropertyHandle> PropertyHandle, UFunction* DelegateSignature)
+{
+	RefreshBlueprintFunctionCache(DelegateSignature);
+
+	const bool bInShouldCloseWindowAfterMenuSelection = true;
+	FMenuBuilder MenuBuilder(bInShouldCloseWindowAfterMenuSelection, NULL);
+
+	//TODO UMG Enable or disable Remove Binding if needed.
+
+	MenuBuilder.BeginSection("ClearBinding");
+	{
+		MenuBuilder.AddMenuEntry(
+			LOCTEXT("RemoveBinding", "Remove Binding"),
+			LOCTEXT("RemoveBindingToolTip", "Removes the current binding"),
+			FSlateIcon(),
+			FUIAction(FExecuteAction::CreateSP(this, &FBlueprintWidgetCustomization::HandleRemoveBinding, PropertyHandle))
+			);
+	}
+	MenuBuilder.EndSection(); //ClearBinding
+
+	MenuBuilder.BeginSection("CreateBinding");
+	{
+		MenuBuilder.AddMenuEntry(
+			LOCTEXT("CreateBinding", "Create Binding"),
+			LOCTEXT("CreateBindingToolTip", "Creates a new function for this property to be bound to"),
+			FSlateIcon(),
+			FUIAction(FExecuteAction::CreateSP(this, &FBlueprintWidgetCustomization::HandleCreateAndAddBinding, PropertyHandle, DelegateSignature))
+			);
+	}
+	MenuBuilder.EndSection(); //CreateBinding
+
+	MenuBuilder.BeginSection("Functions");
+	{
+		for ( TSharedPtr<FunctionInfo>& Function : BlueprintFunctionCache )
+		{
+			MenuBuilder.AddMenuEntry(
+				Function->DisplayName,
+				FText::FromString(Function->Tooltip),
+				FSlateIcon(),
+				FUIAction(FExecuteAction::CreateSP(this, &FBlueprintWidgetCustomization::HandleAddBinding, PropertyHandle, Function))
+				);
+		}
+	}
+	MenuBuilder.EndSection(); //Functions
+
+
+	FDisplayMetrics DisplayMetrics;
+	FSlateApplication::Get().GetDisplayMetrics(DisplayMetrics);
+
+	return
+		SNew(SVerticalBox)
+
+		+ SVerticalBox::Slot()
+		.MaxHeight(DisplayMetrics.PrimaryDisplayHeight * 0.5)
+		[
+			MenuBuilder.MakeWidget()
+		];
+}
+
+FText FBlueprintWidgetCustomization::GetCurrentBindingText(TSharedRef<IPropertyHandle> PropertyHandle) const
+{
+	TArray<UObject*> OuterObjects;
+	PropertyHandle->GetOuterObjects(OuterObjects);
+
+	//TODO UMG O(N) Isn't good for this, needs to be map, but map isn't serialized, need cached runtime map for fast lookups.
+
+	FName PropertyName = PropertyHandle->GetProperty()->GetFName();
+	for ( int32 ObjectIndex = 0; ObjectIndex < OuterObjects.Num(); ObjectIndex++ )
+	{
+		// Ignore null outter objects
+		if ( OuterObjects[ObjectIndex] == NULL )
+		{
+			continue;
+		}
+
+		//TODO UMG handle multiple things selected
+
+		for ( const FDelegateEditorBinding& Binding : Blueprint->Bindings )
+		{
+			if ( Binding.ObjectName == OuterObjects[ObjectIndex]->GetName() && Binding.PropertyName == PropertyName )
+			{
+				FName FoundName = Blueprint->GetFieldNameFromClassByGuid<UFunction>(Blueprint->GeneratedClass, Binding.MemberGuid);
+				return FText::FromName(FoundName);
+				//return FText::FromName(Binding.FunctionName);
+			}
+		}
+
+		//TODO UMG Do something about missing functions, little exclamation points if they're missing and such.
+
+		break;
+	}
+
+	return LOCTEXT("Bind", "Bind");
+}
+
+void FBlueprintWidgetCustomization::HandleRemoveBinding(TSharedRef<IPropertyHandle> PropertyHandle)
+{
+	const FScopedTransaction Transaction(LOCTEXT("BindDelegate", "Remove Binding"));
+
+	Blueprint->Modify();
+
+	TArray<UObject*> OuterObjects;
+	PropertyHandle->GetOuterObjects(OuterObjects);
+	for ( UObject* SelectedObject : OuterObjects )
+	{
+		FDelegateEditorBinding Binding;
+		Binding.ObjectName = SelectedObject->GetName();
+		Binding.PropertyName = PropertyHandle->GetProperty()->GetFName();
+
+		Blueprint->Bindings.Remove(Binding);
+	}
+
+	FBlueprintEditorUtils::MarkBlueprintAsStructurallyModified(Blueprint);
+}
+
+void FBlueprintWidgetCustomization::HandleAddBinding(TSharedRef<IPropertyHandle> PropertyHandle, TSharedPtr<FunctionInfo> SelectedFunction)
+{
+	const FScopedTransaction Transaction(LOCTEXT("BindDelegate", "Set Binding"));
+
+	Blueprint->Modify();
+
+	TArray<UObject*> OuterObjects;
+	PropertyHandle->GetOuterObjects(OuterObjects);
+	for ( UObject* SelectedObject : OuterObjects )
+	{
+		FDelegateEditorBinding Binding;
+		Binding.ObjectName = SelectedObject->GetName();
+		Binding.PropertyName = PropertyHandle->GetProperty()->GetFName();
+		Binding.FunctionName = SelectedFunction->FuncName;
+		Binding.MemberGuid = SelectedFunction->EdGraph->GraphGuid;
+
+		Blueprint->Bindings.Remove(Binding);
+		Blueprint->Bindings.AddUnique(Binding);
+	}
+
+	FBlueprintEditorUtils::MarkBlueprintAsStructurallyModified(Blueprint);
+}
+
+void FBlueprintWidgetCustomization::HandleCreateAndAddBinding(TSharedRef<IPropertyHandle> PropertyHandle, UFunction* DelegateSignature)
+{
+	const FScopedTransaction Transaction(LOCTEXT("CreateDelegate", "Create Binding"));
+
+	Blueprint->Modify();
+
+	// Create the function graph.
+	FString FunctionName = FString(TEXT("Get")) + PropertyHandle->GetProperty()->GetName();
+	UEdGraph* FunctionGraph = FBlueprintEditorUtils::CreateNewGraph(Blueprint, FBlueprintEditorUtils::FindUniqueKismetName(Blueprint, FunctionName), UEdGraph::StaticClass(), UEdGraphSchema_K2::StaticClass());
+
+	bool bUserCreated = true;
+	FBlueprintEditorUtils::AddFunctionGraph(Blueprint, FunctionGraph, bUserCreated, DelegateSignature);
+
+	//TODO UMG Once we start binding event functions, pure needs to be variable, so maybe pass it in.
+	
+	// All bound delegates should be pure.
+	const UEdGraphSchema_K2* Schema_K2 = Cast<UEdGraphSchema_K2>(FunctionGraph->GetSchema());
+	Schema_K2->AddExtraFunctionFlags(FunctionGraph, FUNC_BlueprintPure);
+
+	// Add the binding to the blueprint
+	TSharedPtr<FunctionInfo> SelectedFunction = MakeShareable(new FunctionInfo());
+	SelectedFunction->FuncName = FunctionGraph->GetFName();
+	SelectedFunction->EdGraph = FunctionGraph;
+
+	HandleAddBinding(PropertyHandle, SelectedFunction);
+}
+
+#undef LOCTEXT_NAMESPACE
diff --git a/Engine/Source/Editor/UMGEditor/Private/Customizations/DetailCustomizations.h b/Engine/Source/Editor/UMGEditor/Private/Customizations/DetailCustomizations.h
new file mode 100644
--- /dev/null
+++ b/Engine/Source/Editor/UMGEditor/Private/Customizations/DetailCustomizations.h
@@ -0,0 +1,47 @@
+// Copyright 1998-2014 Epic Games, Inc. All Rights Reserved.
+
+#pragma once
+
+class FBlueprintWidgetCustomization : public IDetailCustomization
+{
+public:
+	/** Makes a new instance of this detail layout class for a specific detail view requesting it */
+	static TSharedRef<class IDetailCustomization> MakeInstance(UBlueprint* Blueprint)
+	{
+		return MakeShareable(new FBlueprintWidgetCustomization(Blueprint));
+	}
+
+	FBlueprintWidgetCustomization(UBlueprint* InBlueprint)
+		: Blueprint(CastChecked<UWidgetBlueprint>(InBlueprint))
+	{
+	}
+	
+	/** IDetailCustomization interface */
+	virtual void CustomizeDetails( IDetailLayoutBuilder& DetailLayout ) OVERRIDE;
+
+private:
+	struct FunctionInfo
+	{
+		FText DisplayName;
+		FString Tooltip;
+
+		FName FuncName;
+		UEdGraph* EdGraph;
+	};
+
+	void RefreshBlueprintFunctionCache(const UFunction* DelegateSignature);
+
+	TSharedRef<SWidget> OnGenerateDelegateMenu(TSharedRef<IPropertyHandle> PropertyHandle, UFunction* DelegateSignature);
+
+	FText GetCurrentBindingText(TSharedRef<IPropertyHandle> PropertyHandle) const;
+
+	void HandleRemoveBinding(TSharedRef<IPropertyHandle> PropertyHandle);
+	void HandleAddBinding(TSharedRef<IPropertyHandle> PropertyHandle, TSharedPtr<FunctionInfo> SelectedFunction);
+	void HandleCreateAndAddBinding(TSharedRef<IPropertyHandle> PropertyHandle, UFunction* DelegateSignature);
+
+private:
+
+	UWidgetBlueprint* Blueprint;
+
+	TArray< TSharedPtr<FunctionInfo> > BlueprintFunctionCache;
+};
diff --git a/Engine/Source/Editor/UMGEditor/Private/DetailCustomizations.cpp b/Engine/Source/Editor/UMGEditor/Private/DetailCustomizations.cpp
deleted file mode 100644
diff --git a/Engine/Source/Editor/UMGEditor/Private/DetailCustomizations.h b/Engine/Source/Editor/UMGEditor/Private/DetailCustomizations.h
deleted file mode 100644
diff --git a/Engine/Source/Editor/UMGEditor/Private/SUMGDesigner.cpp b/Engine/Source/Editor/UMGEditor/Private/SUMGDesigner.cpp
--- a/Engine/Source/Editor/UMGEditor/Private/SUMGDesigner.cpp
+++ b/Engine/Source/Editor/UMGEditor/Private/SUMGDesigner.cpp
@@ -417,6 +417,78 @@ int32 SUMGDesigner::OnPaint(const FGeometry& AllottedGeometry, const FSlateRect&
 		);
 
 		DrawDragHandles(SelectionGeometry, MyClippingRect, OutDrawElements, LayerId, InWidgetStyle);
+
+		if ( CurrentSelection.IsValid() && CurrentSelection.Template->Slot )
+		{
+			const float X = AllottedGeometry.AbsolutePosition.X;
+			const float Y = AllottedGeometry.AbsolutePosition.Y;
+			const float Width = AllottedGeometry.Size.X;
+			const float Height = AllottedGeometry.Size.Y;
+
+			float Selection_X = SelectionGeometry.DrawPosition.X;
+			float Selection_Y = SelectionGeometry.DrawPosition.Y;
+			float Selection_Width = SelectionGeometry.DrawSize.X;
+			float Selection_Height = SelectionGeometry.DrawSize.Y;
+
+			const FVector2D LeftRightSize = FVector2D(32, 16);
+			const FVector2D TopBottomSize = FVector2D(16, 32);
+
+			// @TODO UMG - Don't use the curve editors brushes
+			const FSlateBrush* KeyBrush = FEditorStyle::GetBrush("CurveEd.CurveKey");
+			FLinearColor KeyColor = FLinearColor(1.0f, 0.0f, 0.0f, 1.f);
+
+			if ( UCanvasPanelSlot* CanvasSlot = Cast<UCanvasPanelSlot>(CurrentSelection.Template->Slot) )
+			{
+				//Set("UMGEditor.AnchorCenter", new IMAGE_BRUSH("Icons/umg_anchor_center", Icon16x16));
+				//Set("UMGEditor.AnchorTopBottom", new IMAGE_BRUSH("Icons/umg_anchor_top_bottom", FVector2D(16, 32)));
+				//Set("UMGEditor.AnchorLeftRight", new IMAGE_BRUSH("Icons/umg_anchor_left_right", FVector2D(32, 16)));
+
+				// Left
+				FSlateDrawElement::MakeBox(
+					OutDrawElements,
+					++LayerId,
+					FPaintGeometry(FVector2D(X + Width * CanvasSlot->Anchors.Minimum.X - LeftRightSize.X * 0.5f, SelectionGeometry.DrawPosition.Y + Selection_Height * 0.5f - LeftRightSize.Y * 0.5f), LeftRightSize, 1.0f),
+					FEditorStyle::Get().GetBrush("UMGEditor.AnchorLeftRight"),
+					MyClippingRect,
+					ESlateDrawEffect::None,
+					KeyBrush->GetTint(InWidgetStyle) * InWidgetStyle.GetColorAndOpacityTint() * KeyColor
+					);
+
+				// Right
+				FSlateDrawElement::MakeBox(
+					OutDrawElements,
+					++LayerId,
+					FPaintGeometry(FVector2D(X + Width * CanvasSlot->Anchors.Maximum.X - LeftRightSize.X * 0.5f, SelectionGeometry.DrawPosition.Y + Selection_Height * 0.5f - LeftRightSize.Y * 0.5f), LeftRightSize, 1.0f),
+					FEditorStyle::Get().GetBrush("UMGEditor.AnchorLeftRight"),
+					MyClippingRect,
+					ESlateDrawEffect::None,
+					KeyBrush->GetTint(InWidgetStyle) * InWidgetStyle.GetColorAndOpacityTint() * KeyColor
+					);
+
+
+				// Top
+				FSlateDrawElement::MakeBox(
+					OutDrawElements,
+					++LayerId,
+					FPaintGeometry(FVector2D(SelectionGeometry.DrawPosition.X + Selection_Width * 0.5f - TopBottomSize.X * 0.5f, Y + Height * CanvasSlot->Anchors.Minimum.Y - TopBottomSize.Y * 0.5f), TopBottomSize, 1.0f),
+					FEditorStyle::Get().GetBrush("UMGEditor.AnchorTopBottom"),
+					MyClippingRect,
+					ESlateDrawEffect::None,
+					KeyBrush->GetTint(InWidgetStyle) * InWidgetStyle.GetColorAndOpacityTint() * KeyColor
+					);
+
+				// Bottom
+				FSlateDrawElement::MakeBox(
+					OutDrawElements,
+					++LayerId,
+					FPaintGeometry(FVector2D(SelectionGeometry.DrawPosition.X + Selection_Width * 0.5f - TopBottomSize.X * 0.5f, Y + Height * CanvasSlot->Anchors.Maximum.Y - TopBottomSize.Y * 0.5f), TopBottomSize, 1.0f),
+					FEditorStyle::Get().GetBrush("UMGEditor.AnchorTopBottom"),
+					MyClippingRect,
+					ESlateDrawEffect::None,
+					KeyBrush->GetTint(InWidgetStyle) * InWidgetStyle.GetColorAndOpacityTint() * KeyColor
+					);
+			}
+		}
 	}
 
 	return LayerId;
diff --git a/Engine/Source/Editor/UMGEditor/Private/SUMGEditorWidgetTemplates.cpp b/Engine/Source/Editor/UMGEditor/Private/SUMGEditorWidgetTemplates.cpp
--- a/Engine/Source/Editor/UMGEditor/Private/SUMGEditorWidgetTemplates.cpp
+++ b/Engine/Source/Editor/UMGEditor/Private/SUMGEditorWidgetTemplates.cpp
@@ -18,6 +18,8 @@
 #include "WidgetTemplateButton.h"
 #include "WidgetTemplateCheckBox.h"
 
+#define LOCTEXT_NAMESPACE "UMG"
+
 class SWidgetTemplateItem : public SCompoundWidget
 {
 public:
@@ -35,31 +37,93 @@ public:
 		Template = InTemplate;
 
 		ChildSlot
+		[
+			SNew(SHorizontalBox)
+
+			+ SHorizontalBox::Slot()
+			.AutoWidth()
+			.Padding(2.0f)
 			[
-				SNew(SHorizontalBox)
-
-				+ SHorizontalBox::Slot()
-				.AutoWidth()
-				.Padding(2.0f)
-				[
-					SNew(SImage)
-					.Image(Template->Icon.GetIcon())
-				]
-				
-				+ SHorizontalBox::Slot()
-				.AutoWidth()
-				.Padding(2.0f)
-				[
-					SNew(STextBlock)
-					.Text(Template->Name)
-				]
-			];
+				SNew(SImage)
+				.Image(Template->Icon.GetIcon())
+			]
+
+			+ SHorizontalBox::Slot()
+			.AutoWidth()
+			.Padding(2.0f)
+			[
+				SNew(STextBlock)
+				.Text(Template->Name)
+			]
+		];
 	}
 
 private:
 	TSharedPtr<FWidgetTemplate> Template;
 };
 
+class FWidgetTemplateViewModel : public FWidgetViewModel
+{
+public:
+	virtual FText GetName() const
+	{
+		return Template->Name;
+	}
+
+	virtual TSharedRef<ITableRow> BuildRow(const TSharedRef<STableViewBase>& OwnerTable) OVERRIDE
+	{
+		return SNew(STableRow< TSharedPtr<FWidgetViewModel> >, OwnerTable)
+			.Padding(2.0f)
+			.OnDragDetected(this, &FWidgetTemplateViewModel::OnDraggingWidgetTemplateItem)
+			[
+				SNew(SWidgetTemplateItem, Template)
+			];
+	}
+
+	FReply OnDraggingWidgetTemplateItem(const FGeometry& MyGeometry, const FPointerEvent& MouseEvent)
+	{
+		return FReply::Handled().BeginDragDrop(FWidgetTemplateDragDropOp::New(Template));
+	}
+
+	TSharedPtr<FWidgetTemplate> Template;
+};
+
+class FWidgetHeaderViewModel : public FWidgetViewModel
+{
+public:
+	virtual FText GetName() const
+	{
+		return GroupName;
+	}
+
+	virtual bool IsExpanded() const
+	{
+		return true;
+	}
+
+	virtual TSharedRef<ITableRow> BuildRow(const TSharedRef<STableViewBase>& OwnerTable) OVERRIDE
+	{
+		return SNew(STableRow< TSharedPtr<FWidgetViewModel> >, OwnerTable)
+			.Padding(2.0f)
+			.ShowSelection(false)
+			[
+				SNew(STextBlock)
+				.Text(GroupName)
+			];
+	}
+
+	virtual void GetChildren(TArray< TSharedPtr<FWidgetViewModel> >& OutChildren) OVERRIDE
+	{
+		for ( TSharedPtr<FWidgetViewModel>& Child : Children )
+		{
+			OutChildren.Add(Child);
+		}
+	}
+
+	FText GroupName;
+	TArray< TSharedPtr<FWidgetViewModel> > Children;
+};
+
 void SUMGEditorWidgetTemplates::Construct(const FArguments& InArgs, TSharedPtr<FBlueprintEditor> InBlueprintEditor, USimpleConstructionScript* InSCS)
 {
 	BlueprintEditor = InBlueprintEditor;
@@ -77,12 +141,22 @@ void SUMGEditorWidgetTemplates::Construct(const FArguments& InArgs, TSharedPtr<F
 		+ SVerticalBox::Slot()
 		.FillHeight(1.0f)
 		[
-			SAssignNew(WidgetsListView, SListView< TSharedPtr<FWidgetTemplate> >)
+			SAssignNew(WidgetTemplatesView, STreeView< TSharedPtr<FWidgetViewModel> >)
 			.OnGenerateRow(this, &SUMGEditorWidgetTemplates::OnGenerateWidgetTemplateItem)
-			.ListItemsSource(&WidgetTemplates)
+			.OnGetChildren(this, &SUMGEditorWidgetTemplates::OnGetChildren)
+			.TreeItemsSource(&WidgetViewModels)
 			.SelectionMode(ESelectionMode::Single)
 		]
 	];
+
+	// Restore the expansion state of the widget groups.
+	for ( TSharedPtr<FWidgetViewModel>& ViewModel : WidgetViewModels )
+	{
+		if ( ViewModel->IsExpanded() )
+		{
+			WidgetTemplatesView->SetItemExpansion(ViewModel, true);
+		}
+	}
 }
 
 SUMGEditorWidgetTemplates::~SUMGEditorWidgetTemplates()
@@ -103,7 +177,7 @@ UWidgetBlueprint* SUMGEditorWidgetTemplates::GetBlueprint() const
 
 void SUMGEditorWidgetTemplates::BuildWidgetList()
 {
-	WidgetTemplates.Reset();
+	WidgetViewModels.Reset();
 	WidgetTemplateCategories.Reset();
 
 	BuildClassWidgetList();
@@ -111,11 +185,23 @@ void SUMGEditorWidgetTemplates::BuildWidgetList()
 
 	for ( auto& Entry : WidgetTemplateCategories )
 	{
+		TSharedPtr<FWidgetHeaderViewModel> Header = MakeShareable(new FWidgetHeaderViewModel());
+		Header->GroupName = FText::FromString(Entry.Key);
+
 		for ( auto& Template : Entry.Value )
 		{
-			WidgetTemplates.Add(Template);
+			TSharedPtr<FWidgetTemplateViewModel> TemplateViewModel = MakeShareable(new FWidgetTemplateViewModel());
+			TemplateViewModel->Template = Template;
+
+			Header->Children.Add(TemplateViewModel);
 		}
+
+		Header->Children.Sort([] (TSharedPtr<FWidgetViewModel> L, TSharedPtr<FWidgetViewModel> R) { return R->GetName().CompareTo(L->GetName()) > 0; });
+
+		WidgetViewModels.Add(Header);
 	}
+
+	WidgetViewModels.Sort([] (TSharedPtr<FWidgetViewModel> L, TSharedPtr<FWidgetViewModel> R) { return R->GetName().CompareTo(L->GetName()) > 0; });
 }
 
 void SUMGEditorWidgetTemplates::BuildClassWidgetList()
@@ -157,24 +243,14 @@ void SUMGEditorWidgetTemplates::AddWidgetTemplate(TSharedPtr<FWidgetTemplate> Te
 	Group.Add(Template);
 }
 
-TSharedRef<ITableRow> SUMGEditorWidgetTemplates::OnGenerateWidgetTemplateItem(TSharedPtr<FWidgetTemplate> Item, const TSharedRef<STableViewBase>& OwnerTable)
+void SUMGEditorWidgetTemplates::OnGetChildren(TSharedPtr<FWidgetViewModel> Item, TArray< TSharedPtr<FWidgetViewModel> >& Children)
 {
-	return
-		SNew(STableRow< TSharedPtr<FWidgetTemplate> >, OwnerTable)
-		.Padding(2.0f)
-		.OnDragDetected(this, &SUMGEditorWidgetTemplates::OnDraggingWidgetTemplateItem)
-		[
-			SNew(SWidgetTemplateItem, Item)
-		];
+	return Item->GetChildren(Children);
 }
 
-FReply SUMGEditorWidgetTemplates::OnDraggingWidgetTemplateItem(const FGeometry& MyGeometry, const FPointerEvent& MouseEvent)
+TSharedRef<ITableRow> SUMGEditorWidgetTemplates::OnGenerateWidgetTemplateItem(TSharedPtr<FWidgetViewModel> Item, const TSharedRef<STableViewBase>& OwnerTable)
 {
-	if ( WidgetsListView->GetNumItemsSelected() > 0 )
-	{
-		TSharedPtr<FWidgetTemplate> Template = WidgetsListView->GetSelectedItems()[0];
-		return FReply::Handled().BeginDragDrop(FWidgetTemplateDragDropOp::New(Template));
-	}
-
-	return FReply::Unhandled();
+	return Item->BuildRow(OwnerTable);
 }
+
+#undef LOCTEXT_NAMESPACE
diff --git a/Engine/Source/Editor/UMGEditor/Private/SUMGEditorWidgetTemplates.h b/Engine/Source/Editor/UMGEditor/Private/SUMGEditorWidgetTemplates.h
--- a/Engine/Source/Editor/UMGEditor/Private/SUMGEditorWidgetTemplates.h
+++ b/Engine/Source/Editor/UMGEditor/Private/SUMGEditorWidgetTemplates.h
@@ -7,6 +7,25 @@
 
 #include "WidgetTemplate.h"
 
+/** View model for the items in the widget template list */
+class FWidgetViewModel : public TSharedFromThis<FWidgetViewModel>
+{
+public:
+	virtual FText GetName() const = 0;
+
+	virtual TSharedRef<ITableRow> BuildRow(const TSharedRef<STableViewBase>& OwnerTable) = 0;
+
+	virtual void GetChildren(TArray< TSharedPtr<FWidgetViewModel> >& OutChildren)
+	{
+	}
+
+	virtual bool IsExpanded() const
+	{
+		return false;
+	}
+};
+
+/**  */
 class SUMGEditorWidgetTemplates : public SCompoundWidget
 {
 public:
@@ -23,9 +42,8 @@ private:
 	void BuildClassWidgetList();
 	void BuildSpecialWidgetList();
 
-	TSharedRef<ITableRow> OnGenerateWidgetTemplateItem(TSharedPtr<FWidgetTemplate> Item, const TSharedRef<STableViewBase>& OwnerTable);
-
-	FReply OnDraggingWidgetTemplateItem(const FGeometry& MyGeometry, const FPointerEvent& MouseEvent);
+	void OnGetChildren(TSharedPtr<FWidgetViewModel> Item, TArray< TSharedPtr<FWidgetViewModel> >& Children);
+	TSharedRef<ITableRow> OnGenerateWidgetTemplateItem(TSharedPtr<FWidgetViewModel> Item, const TSharedRef<STableViewBase>& OwnerTable);
 
 private:
 	void AddWidgetTemplate(TSharedPtr<FWidgetTemplate> Template);
@@ -34,7 +52,9 @@ private:
 
 	typedef TArray< TSharedPtr<FWidgetTemplate> > WidgetTemplateArray;
 	TMap< FString, WidgetTemplateArray > WidgetTemplateCategories;
-	WidgetTemplateArray WidgetTemplates;
 
-	TSharedPtr< SListView< TSharedPtr<FWidgetTemplate> > > WidgetsListView;
+	typedef TArray< TSharedPtr<FWidgetViewModel> > ViewModelsArray;
+	ViewModelsArray WidgetViewModels;
+
+	TSharedPtr< STreeView< TSharedPtr<FWidgetViewModel> > > WidgetTemplatesView;
 };
diff --git a/Engine/Source/Editor/UMGEditor/Private/Templates/WidgetTemplateButton.cpp b/Engine/Source/Editor/UMGEditor/Private/Templates/WidgetTemplateButton.cpp
--- a/Engine/Source/Editor/UMGEditor/Private/Templates/WidgetTemplateButton.cpp
+++ b/Engine/Source/Editor/UMGEditor/Private/Templates/WidgetTemplateButton.cpp
@@ -15,7 +15,7 @@ FWidgetTemplateButton::FWidgetTemplateButton()
 
 FText FWidgetTemplateButton::GetCategory()
 {
-	return LOCTEXT("Common", "Common");
+	return LOCTEXT("Preset", "Preset");
 }
 
 UWidget* FWidgetTemplateButton::Create(UWidgetTree* Tree)
diff --git a/Engine/Source/Editor/UMGEditor/Private/Templates/WidgetTemplateCheckBox.cpp b/Engine/Source/Editor/UMGEditor/Private/Templates/WidgetTemplateCheckBox.cpp
--- a/Engine/Source/Editor/UMGEditor/Private/Templates/WidgetTemplateCheckBox.cpp
+++ b/Engine/Source/Editor/UMGEditor/Private/Templates/WidgetTemplateCheckBox.cpp
@@ -15,7 +15,7 @@ FWidgetTemplateCheckBox::FWidgetTemplateCheckBox()
 
 FText FWidgetTemplateCheckBox::GetCategory()
 {
-	return LOCTEXT("Common", "Common");
+	return LOCTEXT("Preset", "Preset");
 }
 
 UWidget* FWidgetTemplateCheckBox::Create(UWidgetTree* Tree)
diff --git a/Engine/Source/Editor/UMGEditor/Private/Templates/WidgetTemplateClass.cpp b/Engine/Source/Editor/UMGEditor/Private/Templates/WidgetTemplateClass.cpp
--- a/Engine/Source/Editor/UMGEditor/Private/Templates/WidgetTemplateClass.cpp
+++ b/Engine/Source/Editor/UMGEditor/Private/Templates/WidgetTemplateClass.cpp
@@ -13,8 +13,14 @@ FWidgetTemplateClass::FWidgetTemplateClass(TSubclassOf<UWidget> InWidgetClass)
 
 FText FWidgetTemplateClass::GetCategory()
 {
-	// TODO this should be pulled from metadata from the class, that describes its categorization, is it a panel, a common control...etc?
-	return LOCTEXT("Class", "Class");
+	const FString& Category = WidgetClass->GetMetaData("Category");
+
+	if ( Category.IsEmpty() )
+	{
+		return LOCTEXT("Misc", "Misc");
+	}
+
+	return FText::FromString(Category);
 }
 
 UWidget* FWidgetTemplateClass::Create(UWidgetTree* Tree)
diff --git a/Engine/Source/Editor/UMGEditor/Private/UMGBlueprintEditorExtensionHook.cpp b/Engine/Source/Editor/UMGEditor/Private/UMGBlueprintEditorExtensionHook.cpp
--- a/Engine/Source/Editor/UMGEditor/Private/UMGBlueprintEditorExtensionHook.cpp
+++ b/Engine/Source/Editor/UMGEditor/Private/UMGBlueprintEditorExtensionHook.cpp
@@ -18,6 +18,7 @@
 #include "IDetailsView.h"
 
 #include "DetailCustomizations.h"
+#include "CanvasSlotCustomization.h"
 
 #include "WidgetBlueprintEditor.h"
 
@@ -225,7 +226,9 @@ public:
 		TSharedRef<class SKismetInspector> Inspector = MyBlueprintEditor.Pin()->GetInspector();
 		FOnGetDetailCustomizationInstance LayoutDelegateDetails = FOnGetDetailCustomizationInstance::CreateStatic(&FBlueprintWidgetCustomization::MakeInstance, MyBlueprintEditor.Pin()->GetBlueprintObj());
 		Inspector->GetPropertyView()->RegisterInstancedCustomPropertyLayout(UWidget::StaticClass(), LayoutDelegateDetails);
-		//Inspector->GetPropertyView()->RegisterInstancedCustomPropertyLayout(UDelegateProperty::StaticClass(), LayoutDelegateDetails);
+		
+		FOnGetDetailCustomizationInstance CanvasSlotCustomization = FOnGetDetailCustomizationInstance::CreateStatic(&FCanvasSlotCustomization::MakeInstance, MyBlueprintEditor.Pin()->GetBlueprintObj());
+		Inspector->GetPropertyView()->RegisterInstancedCustomPropertyLayout(UCanvasPanelSlot::StaticClass(), CanvasSlotCustomization);
 	}
 
 	UWidgetBlueprint* GetBlueprint() const
diff --git a/Engine/Source/Editor/UMGEditor/Private/UMGEditorPrivatePCH.h b/Engine/Source/Editor/UMGEditor/Private/UMGEditorPrivatePCH.h
--- a/Engine/Source/Editor/UMGEditor/Private/UMGEditorPrivatePCH.h
+++ b/Engine/Source/Editor/UMGEditor/Private/UMGEditorPrivatePCH.h
@@ -21,3 +21,6 @@
 #include "HorizontalSlotExtension.h"
 #include "CanvasSlotExtension.h"
 #include "UniformGridSlotExtension.h"
+
+#include "DetailCustomizations.h"
+#include "CanvasSlotCustomization.h"
diff --git a/Engine/Source/Editor/UMGEditor/UMGEditor.Build.cs b/Engine/Source/Editor/UMGEditor/UMGEditor.Build.cs
--- a/Engine/Source/Editor/UMGEditor/UMGEditor.Build.cs
+++ b/Engine/Source/Editor/UMGEditor/UMGEditor.Build.cs
@@ -14,6 +14,7 @@ public class UMGEditor : ModuleRules
 				"Editor/UMGEditor/Private", // For PCH includes (because they don't work with relative paths, yet)
 				"Editor/UMGEditor/Private/Templates",
                 "Editor/UMGEditor/Private/Extensions",
+                "Editor/UMGEditor/Private/Customizations",
 			});
 
 		PrivateIncludePathModuleNames.AddRange(
diff --git a/Engine/Source/Runtime/UMG/Private/Components/Button.cpp b/Engine/Source/Runtime/UMG/Private/Components/Button.cpp
--- a/Engine/Source/Runtime/UMG/Private/Components/Button.cpp
+++ b/Engine/Source/Runtime/UMG/Private/Components/Button.cpp
@@ -10,11 +10,12 @@
 UButton::UButton(const FPostConstructInitializeProperties& PCIP)
 	: Super(PCIP)
 {
+	HorizontalAlignment = HAlign_Center;
+	VerticalAlignment = VAlign_Center;
+
 	SButton::FArguments ButtonDefaults;
 
 	Style = NULL;
-	HorizontalAlignment = ButtonDefaults._HAlign;
-	VerticalAlignment = ButtonDefaults._VAlign;
 	ContentPadding = ButtonDefaults._ContentPadding.Get();
 
 	DesiredSizeScale = ButtonDefaults._DesiredSizeScale.Get();
diff --git a/Engine/Source/Runtime/UMG/Private/Components/CanvasPanelSlot.cpp b/Engine/Source/Runtime/UMG/Private/Components/CanvasPanelSlot.cpp
--- a/Engine/Source/Runtime/UMG/Private/Components/CanvasPanelSlot.cpp
+++ b/Engine/Source/Runtime/UMG/Private/Components/CanvasPanelSlot.cpp
@@ -82,4 +82,20 @@ void UCanvasPanelSlot::SetAlignment(FVector2D InAlignment)
 	{
 		Slot->Alignment(InAlignment);
 	}
-}
\ No newline at end of file
+}
+
+#if WITH_EDITOR
+void UCanvasPanelSlot::PreEditChange(UProperty* PropertyAboutToChange)
+{
+	//static FName AnchorsProperty(TEXT("Anchors"));
+
+	//if ( PropertyChangedEvent.MemberProperty && PropertyChangedEvent.MemberProperty->GetFName() == AnchorsProperty )
+	//{
+
+	//}
+}
+
+void UCanvasPanelSlot::PostEditChangeProperty(struct FPropertyChangedEvent& PropertyChangedEvent)
+{
+}
+#endif
\ No newline at end of file
diff --git a/Engine/Source/Runtime/UMG/Public/Components/Border.h b/Engine/Source/Runtime/UMG/Public/Components/Border.h
--- a/Engine/Source/Runtime/UMG/Public/Components/Border.h
+++ b/Engine/Source/Runtime/UMG/Public/Components/Border.h
@@ -10,7 +10,7 @@ class USlateBrushAsset;
  * A border is a container widget that can contain one child widget, providing an opportunity 
  * to surround it with a border image and adjustable padding.
  */
-UCLASS(meta=(BlueprintSpawnableComponent), ClassGroup=UserInterface)
+UCLASS(meta=( BlueprintSpawnableComponent, Category="Common" ), ClassGroup=UserInterface)
 class UMG_API UBorder : public UContentWidget
 {
 	GENERATED_UCLASS_BODY()
diff --git a/Engine/Source/Runtime/UMG/Public/Components/Button.h b/Engine/Source/Runtime/UMG/Public/Components/Button.h
--- a/Engine/Source/Runtime/UMG/Public/Components/Button.h
+++ b/Engine/Source/Runtime/UMG/Public/Components/Button.h
@@ -7,7 +7,7 @@
 DECLARE_DYNAMIC_MULTICAST_DELEGATE(FOnButtonClickedEvent);
 
 /** Buttons are clickable widgets */
-UCLASS(meta=(BlueprintSpawnableComponent), ClassGroup=UserInterface)
+UCLASS(meta=( BlueprintSpawnableComponent, Category="Common" ), ClassGroup=UserInterface)
 class UMG_API UButton : public UContentWidget
 {
 	GENERATED_UCLASS_BODY()
diff --git a/Engine/Source/Runtime/UMG/Public/Components/CanvasPanel.h b/Engine/Source/Runtime/UMG/Public/Components/CanvasPanel.h
--- a/Engine/Source/Runtime/UMG/Public/Components/CanvasPanel.h
+++ b/Engine/Source/Runtime/UMG/Public/Components/CanvasPanel.h
@@ -4,7 +4,7 @@
 
 #include "CanvasPanel.generated.h"
 
-UCLASS(meta=(BlueprintSpawnableComponent), ClassGroup=UserInterface)
+UCLASS(meta=( BlueprintSpawnableComponent, Category="Panel" ), ClassGroup=UserInterface)
 class UMG_API UCanvasPanel : public UPanelWidget
 {
 	GENERATED_UCLASS_BODY()
diff --git a/Engine/Source/Runtime/UMG/Public/Components/CanvasPanelSlot.h b/Engine/Source/Runtime/UMG/Public/Components/CanvasPanelSlot.h
--- a/Engine/Source/Runtime/UMG/Public/Components/CanvasPanelSlot.h
+++ b/Engine/Source/Runtime/UMG/Public/Components/CanvasPanelSlot.h
@@ -41,6 +41,13 @@ class UMG_API UCanvasPanelSlot : public UPanelSlot
 	UFUNCTION(BlueprintCallable, Category="Appearance")
 	void SetAlignment(FVector2D InAlignment);
 
+#if WITH_EDITOR
+	// UObject interface
+	virtual void PreEditChange(UProperty* PropertyAboutToChange) OVERRIDE;
+	virtual void PostEditChangeProperty(struct FPropertyChangedEvent& PropertyChangedEvent) OVERRIDE;
+	// End of UObject interface
+#endif
+
 private:
 	SConstraintCanvas::FSlot* Slot;
 };
\ No newline at end of file
diff --git a/Engine/Source/Runtime/UMG/Public/Components/CheckBox.h b/Engine/Source/Runtime/UMG/Public/Components/CheckBox.h
--- a/Engine/Source/Runtime/UMG/Public/Components/CheckBox.h
+++ b/Engine/Source/Runtime/UMG/Public/Components/CheckBox.h
@@ -7,7 +7,7 @@
 DECLARE_DYNAMIC_MULTICAST_DELEGATE_OneParam( FOnCheckBoxComponentStateChanged, bool, bIsChecked );
 
 /** Check box widget */
-UCLASS(meta=(BlueprintSpawnableComponent), ClassGroup=UserInterface)
+UCLASS(meta=( BlueprintSpawnableComponent, Category="Common" ), ClassGroup=UserInterface)
 class UMG_API UCheckBox : public UWidget
 {
 	GENERATED_UCLASS_BODY()
diff --git a/Engine/Source/Runtime/UMG/Public/Components/EditableText.h b/Engine/Source/Runtime/UMG/Public/Components/EditableText.h
--- a/Engine/Source/Runtime/UMG/Public/Components/EditableText.h
+++ b/Engine/Source/Runtime/UMG/Public/Components/EditableText.h
@@ -8,7 +8,7 @@ DECLARE_DYNAMIC_MULTICAST_DELEGATE_OneParam(FOnEditableTextChangedEvent, const F
 DECLARE_DYNAMIC_MULTICAST_DELEGATE_TwoParams(FOnEditableTextCommittedEvent, const FText&, Text, ETextCommit::Type, CommitMethod);
 
 /** Editable text box widget */
-UCLASS(meta=(BlueprintSpawnableComponent), ClassGroup=UserInterface)
+UCLASS(meta=( BlueprintSpawnableComponent, Category="Common" ), ClassGroup=UserInterface)
 class UMG_API UEditableText : public UWidget
 {
 	GENERATED_UCLASS_BODY()
diff --git a/Engine/Source/Runtime/UMG/Public/Components/EditableTextBlock.h b/Engine/Source/Runtime/UMG/Public/Components/EditableTextBlock.h
--- a/Engine/Source/Runtime/UMG/Public/Components/EditableTextBlock.h
+++ b/Engine/Source/Runtime/UMG/Public/Components/EditableTextBlock.h
@@ -8,7 +8,7 @@ DECLARE_DYNAMIC_MULTICAST_DELEGATE_OneParam(FOnEditableTextBlock_TextChangedEven
 //DECLARE_DYNAMIC_MULTICAST_DELEGATE_TwoParams(FOnEditableTextBlock_TextCommittedEvent, const FText&, Text, ETextCommit::Type, Type);
 
 /** Editable text box widget */
-UCLASS(meta=(BlueprintSpawnableComponent), ClassGroup=UserInterface)
+UCLASS(meta=( BlueprintSpawnableComponent, Category="Common" ), ClassGroup=UserInterface)
 class UMG_API UEditableTextBlock : public UWidget
 {
 	GENERATED_UCLASS_BODY()
diff --git a/Engine/Source/Runtime/UMG/Public/Components/HorizontalBox.h b/Engine/Source/Runtime/UMG/Public/Components/HorizontalBox.h
--- a/Engine/Source/Runtime/UMG/Public/Components/HorizontalBox.h
+++ b/Engine/Source/Runtime/UMG/Public/Components/HorizontalBox.h
@@ -4,7 +4,7 @@
 
 #include "HorizontalBox.generated.h"
 
-UCLASS(meta=(BlueprintSpawnableComponent), ClassGroup=UserInterface)
+UCLASS(meta=( BlueprintSpawnableComponent, Category="Panel" ), ClassGroup=UserInterface)
 class UMG_API UHorizontalBox : public UPanelWidget
 {
 	GENERATED_UCLASS_BODY()
diff --git a/Engine/Source/Runtime/UMG/Public/Components/Image.h b/Engine/Source/Runtime/UMG/Public/Components/Image.h
--- a/Engine/Source/Runtime/UMG/Public/Components/Image.h
+++ b/Engine/Source/Runtime/UMG/Public/Components/Image.h
@@ -7,7 +7,7 @@
 class USlateBrushAsset;
 
 /** Image widget */
-UCLASS(meta=(BlueprintSpawnableComponent), ClassGroup=UserInterface)
+UCLASS(meta=( BlueprintSpawnableComponent, Category="Common" ), ClassGroup=UserInterface)
 class UMG_API UImage : public UWidget
 {
 	GENERATED_UCLASS_BODY()
diff --git a/Engine/Source/Runtime/UMG/Public/Components/Overlay.h b/Engine/Source/Runtime/UMG/Public/Components/Overlay.h
--- a/Engine/Source/Runtime/UMG/Public/Components/Overlay.h
+++ b/Engine/Source/Runtime/UMG/Public/Components/Overlay.h
@@ -4,7 +4,7 @@
 
 #include "Overlay.generated.h"
 
-UCLASS(meta=(BlueprintSpawnableComponent), ClassGroup=UserInterface)
+UCLASS(meta=( BlueprintSpawnableComponent, Category="Panel" ), ClassGroup=UserInterface)
 class UMG_API UOverlay : public UPanelWidget
 {
 	GENERATED_UCLASS_BODY()
diff --git a/Engine/Source/Runtime/UMG/Public/Components/ScrollBox.h b/Engine/Source/Runtime/UMG/Public/Components/ScrollBox.h
--- a/Engine/Source/Runtime/UMG/Public/Components/ScrollBox.h
+++ b/Engine/Source/Runtime/UMG/Public/Components/ScrollBox.h
@@ -5,7 +5,7 @@
 #include "ScrollBox.generated.h"
 
 /** An arbitrary scrollable collection of widgets.  Great for presenting 10-100 widgets in a list.  Doesn't support virtualization. */
-UCLASS(meta=(BlueprintSpawnableComponent), ClassGroup=UserInterface)
+UCLASS(meta=( BlueprintSpawnableComponent, Category="Panel" ), ClassGroup=UserInterface)
 class UMG_API UScrollBox : public UPanelWidget
 {
 	GENERATED_UCLASS_BODY()
diff --git a/Engine/Source/Runtime/UMG/Public/Components/TextBlock.h b/Engine/Source/Runtime/UMG/Public/Components/TextBlock.h
--- a/Engine/Source/Runtime/UMG/Public/Components/TextBlock.h
+++ b/Engine/Source/Runtime/UMG/Public/Components/TextBlock.h
@@ -5,7 +5,7 @@
 #include "TextBlock.generated.h"
 
 /** A simple static text widget */
-UCLASS(meta=(BlueprintSpawnableComponent), ClassGroup=UserInterface)
+UCLASS(meta=( BlueprintSpawnableComponent, Category="Common" ), ClassGroup=UserInterface)
 class UMG_API UTextBlock : public UWidget
 {
 	GENERATED_UCLASS_BODY()
diff --git a/Engine/Source/Runtime/UMG/Public/Components/UniformGridPanel.h b/Engine/Source/Runtime/UMG/Public/Components/UniformGridPanel.h
--- a/Engine/Source/Runtime/UMG/Public/Components/UniformGridPanel.h
+++ b/Engine/Source/Runtime/UMG/Public/Components/UniformGridPanel.h
@@ -5,7 +5,7 @@
 #include "UniformGridPanel.generated.h"
 
 /** A panel that evenly divides up available space between all of its children. */
-UCLASS(meta=(BlueprintSpawnableComponent), ClassGroup=UserInterface)
+UCLASS(meta=( BlueprintSpawnableComponent, Category="Panel"), ClassGroup=UserInterface)
 class UMG_API UUniformGridPanel : public UPanelWidget
 {
 	GENERATED_UCLASS_BODY()
diff --git a/Engine/Source/Runtime/UMG/Public/Components/VerticalBox.h b/Engine/Source/Runtime/UMG/Public/Components/VerticalBox.h
--- a/Engine/Source/Runtime/UMG/Public/Components/VerticalBox.h
+++ b/Engine/Source/Runtime/UMG/Public/Components/VerticalBox.h
@@ -8,7 +8,7 @@
  * A vertical box widget is a layout panel allowing child widgets to be automatically laid out
  * vertically.
  */
-UCLASS(meta=(BlueprintSpawnableComponent), ClassGroup=UserInterface)
+UCLASS(meta=( BlueprintSpawnableComponent, Category="Panel" ), ClassGroup=UserInterface)
 class UMG_API UVerticalBox : public UPanelWidget
 {
 	GENERATED_UCLASS_BODY()
