UMG - Customizing the way slots are displayed, flattened, hiding the editinline look, everything just appears as a layout property.

diff --git a/Engine/Source/Editor/UMGEditor/Private/Customizations/CanvasSlotCustomization.cpp b/Engine/Source/Editor/UMGEditor/Private/Customizations/CanvasSlotCustomization.cpp
--- a/Engine/Source/Editor/UMGEditor/Private/Customizations/CanvasSlotCustomization.cpp
+++ b/Engine/Source/Editor/UMGEditor/Private/Customizations/CanvasSlotCustomization.cpp
@@ -10,16 +10,66 @@
 
 #define LOCTEXT_NAMESPACE "UMG"
 
-void FCanvasSlotCustomization::CustomizeDetails( IDetailLayoutBuilder& DetailLayout )
+void FCanvasSlotCustomization::CustomizeHeader(TSharedRef<IPropertyHandle> PropertyHandle, FDetailWidgetRow& HeaderRow, IPropertyTypeCustomizationUtils& CustomizationUtils)
 {
 	// Build the layout category
-	IDetailCategoryBuilder& LayoutCategory = DetailLayout.EditCategory("Layout");
+	//IDetailCategoryBuilder& LayoutCategory = DetailLayout.EditCategory("Layout");
+}
+
+void FCanvasSlotCustomization::CustomizeChildren(TSharedRef<IPropertyHandle> PropertyHandle, IDetailChildrenBuilder& ChildBuilder, IPropertyTypeCustomizationUtils& CustomizationUtils)
+{
+	CustomizeAnchors(PropertyHandle, ChildBuilder, CustomizationUtils);
+
+	FillOutChildren(PropertyHandle, ChildBuilder, CustomizationUtils);
+}
 
-	TSharedRef<IPropertyHandle> AnchorsHandle = DetailLayout.GetProperty(GET_MEMBER_NAME_CHECKED(UCanvasPanelSlot, Anchors));
-	IDetailPropertyRow& AnchorsPropertyRow = LayoutCategory.AddProperty(AnchorsHandle);
+void FCanvasSlotCustomization::FillOutChildren(TSharedRef<IPropertyHandle> PropertyHandle, IDetailChildrenBuilder& ChildBuilder, IPropertyTypeCustomizationUtils& CustomizationUtils)
+{
+	// Generate all the other children
+	uint32 NumChildren;
+	PropertyHandle->GetNumChildren(NumChildren);
+
+	for ( uint32 ChildIndex = 0; ChildIndex < NumChildren; ChildIndex++ )
+	{
+		TSharedRef<IPropertyHandle> ChildHandle = PropertyHandle->GetChildHandle(ChildIndex).ToSharedRef();
+		if ( ChildHandle->IsCustomized() )
+		{
+			continue;
+		}
+
+		if ( ChildHandle->GetProperty() == NULL )
+		{
+			FillOutChildren(ChildHandle, ChildBuilder, CustomizationUtils);
+		}
+		else
+		{
+			ChildBuilder.AddChildProperty(ChildHandle);
+		}
+	}
+}
+
+void FCanvasSlotCustomization::CustomizeAnchors(TSharedRef<IPropertyHandle> PropertyHandle, IDetailChildrenBuilder& ChildBuilder, IPropertyTypeCustomizationUtils& CustomizationUtils)
+{
+	TSharedPtr<IPropertyHandle> Slot = PropertyHandle->GetChildHandle(GET_MEMBER_NAME_CHECKED(UCanvasPanelSlot, Anchors));
+
+	if ( !Slot.IsValid() )
+	{
+		return;
+	}
+
+	Slot->MarkHiddenByCustomization();
+
+	TSharedRef<IPropertyHandle> AnchorsHandle = Slot.ToSharedRef();
+
+	IDetailPropertyRow& AnchorsPropertyRow = ChildBuilder.AddChildProperty(AnchorsHandle);
 
 	const bool bShowChildren = true;
 	AnchorsPropertyRow.CustomWidget(bShowChildren)
+		.NameContent()
+		[
+			SNew(STextBlock)
+			.Text(LOCTEXT("Anchors", "Anchors"))
+		]
 		.ValueContent()
 		[
 			SNew(SHorizontalBox)
@@ -194,7 +244,7 @@ void FCanvasSlotCustomization::CustomizeDetails( IDetailLayoutBuilder& DetailLay
 					+ SUniformGridPanel::Slot(3, 3)
 					[
 						SNew(SButton)
-						.OnClicked(this, &FCanvasSlotCustomization::OnAnchorClicked, AnchorsHandle, FAnchors(1, 1, 1, 1))
+						.OnClicked(this, &FCanvasSlotCustomization::OnAnchorClicked, AnchorsHandle, FAnchors(0, 0, 1, 1))
 						[
 							SNew(SImage)
 							.Image(FEditorStyle::Get().GetBrush("UMGEditor.FillFill"))
@@ -207,8 +257,14 @@ void FCanvasSlotCustomization::CustomizeDetails( IDetailLayoutBuilder& DetailLay
 
 FReply FCanvasSlotCustomization::OnAnchorClicked(TSharedRef<IPropertyHandle> AnchorsHandle, FAnchors Anchors)
 {
+	//TODO UMG Group under single undo transaction.
+
 	FString Value = FString::Printf(TEXT("(Minimum=(X=%f,Y=%f),Maximum=(X=%f,Y=%f))"), Anchors.Minimum.X, Anchors.Minimum.Y, Anchors.Maximum.X, Anchors.Maximum.Y);
 	AnchorsHandle->SetValueFromFormattedString(Value);
+
+	//TODO UMG Reset position, pivot.
+	//TODO UMG Moving the anchors shouldn't change the position of the widget, unless stretched.  But maybe that too should be based off current position and size.
+	//         Either way need access to cached geometry of the slot in the preview world...
 	
 	FSlateApplication::Get().DismissAllMenus();
 
diff --git a/Engine/Source/Editor/UMGEditor/Private/Customizations/CanvasSlotCustomization.h b/Engine/Source/Editor/UMGEditor/Private/Customizations/CanvasSlotCustomization.h
--- a/Engine/Source/Editor/UMGEditor/Private/Customizations/CanvasSlotCustomization.h
+++ b/Engine/Source/Editor/UMGEditor/Private/Customizations/CanvasSlotCustomization.h
@@ -2,11 +2,11 @@
 
 #pragma once
 
-class FCanvasSlotCustomization : public IDetailCustomization
+class FCanvasSlotCustomization : public IPropertyTypeCustomization
 {
 public:
 	/** Makes a new instance of this detail layout class for a specific detail view requesting it */
-	static TSharedRef<class IDetailCustomization> MakeInstance(UBlueprint* Blueprint)
+	static TSharedRef<class IPropertyTypeCustomization> MakeInstance(UBlueprint* Blueprint)
 	{
 		return MakeShareable(new FCanvasSlotCustomization(Blueprint));
 	}
@@ -16,10 +16,14 @@ public:
 	{
 	}
 	
-	/** IDetailCustomization interface */
-	virtual void CustomizeDetails( IDetailLayoutBuilder& DetailLayout ) OVERRIDE;
+	virtual void CustomizeHeader(TSharedRef<IPropertyHandle> PropertyHandle, class FDetailWidgetRow& HeaderRow, IPropertyTypeCustomizationUtils& CustomizationUtils) OVERRIDE;
+
+	virtual void CustomizeChildren(TSharedRef<IPropertyHandle> PropertyHandle, class IDetailChildrenBuilder& ChildBuilder, IPropertyTypeCustomizationUtils& CustomizationUtils) OVERRIDE;
 	
 private:
+
+	void CustomizeAnchors(TSharedRef<IPropertyHandle> PropertyHandle, class IDetailChildrenBuilder& ChildBuilder, IPropertyTypeCustomizationUtils& CustomizationUtils);
+	void FillOutChildren(TSharedRef<IPropertyHandle> PropertyHandle, class IDetailChildrenBuilder& ChildBuilder, IPropertyTypeCustomizationUtils& CustomizationUtils);
 	
 	FReply OnAnchorClicked(TSharedRef<IPropertyHandle> AnchorsHandle, FAnchors Anchors);
 
diff --git a/Engine/Source/Editor/UMGEditor/Private/SUMGEditorWidgetTemplates.cpp b/Engine/Source/Editor/UMGEditor/Private/SUMGEditorWidgetTemplates.cpp
--- a/Engine/Source/Editor/UMGEditor/Private/SUMGEditorWidgetTemplates.cpp
+++ b/Engine/Source/Editor/UMGEditor/Private/SUMGEditorWidgetTemplates.cpp
@@ -95,12 +95,7 @@ public:
 	{
 		return GroupName;
 	}
-
-	virtual bool IsExpanded() const
-	{
-		return true;
-	}
-
+	
 	virtual TSharedRef<ITableRow> BuildRow(const TSharedRef<STableViewBase>& OwnerTable) OVERRIDE
 	{
 		return SNew(STableRow< TSharedPtr<FWidgetViewModel> >, OwnerTable)
@@ -149,19 +144,36 @@ void SUMGEditorWidgetTemplates::Construct(const FArguments& InArgs, TSharedPtr<F
 		]
 	];
 
+	LoadItemExpanssion();
+}
+
+SUMGEditorWidgetTemplates::~SUMGEditorWidgetTemplates()
+{
+	//FCoreDelegates::OnObjectPropertyChanged.RemoveAll(this);
+	SaveItemExpansion();
+}
+
+void SUMGEditorWidgetTemplates::LoadItemExpanssion()
+{
 	// Restore the expansion state of the widget groups.
 	for ( TSharedPtr<FWidgetViewModel>& ViewModel : WidgetViewModels )
 	{
-		if ( ViewModel->IsExpanded() )
+		bool IsExpanded;
+		if ( GConfig->GetBool(TEXT("WidgetTemplatesExpanded"), *ViewModel->GetName().ToString(), IsExpanded, GEditorUserSettingsIni) && IsExpanded )
 		{
 			WidgetTemplatesView->SetItemExpansion(ViewModel, true);
 		}
 	}
 }
 
-SUMGEditorWidgetTemplates::~SUMGEditorWidgetTemplates()
+void SUMGEditorWidgetTemplates::SaveItemExpansion()
 {
-	//FCoreDelegates::OnObjectPropertyChanged.RemoveAll(this);
+	// Restore the expansion state of the widget groups.
+	for ( TSharedPtr<FWidgetViewModel>& ViewModel : WidgetViewModels )
+	{
+		const bool IsExpanded = WidgetTemplatesView->IsItemExpanded(ViewModel);
+		GConfig->SetBool(TEXT("WidgetTemplatesExpanded"), *ViewModel->GetName().ToString(), IsExpanded, GEditorUserSettingsIni);
+	}
 }
 
 UWidgetBlueprint* SUMGEditorWidgetTemplates::GetBlueprint() const
diff --git a/Engine/Source/Editor/UMGEditor/Private/SUMGEditorWidgetTemplates.h b/Engine/Source/Editor/UMGEditor/Private/SUMGEditorWidgetTemplates.h
--- a/Engine/Source/Editor/UMGEditor/Private/SUMGEditorWidgetTemplates.h
+++ b/Engine/Source/Editor/UMGEditor/Private/SUMGEditorWidgetTemplates.h
@@ -18,11 +18,6 @@ public:
 	virtual void GetChildren(TArray< TSharedPtr<FWidgetViewModel> >& OutChildren)
 	{
 	}
-
-	virtual bool IsExpanded() const
-	{
-		return false;
-	}
 };
 
 /**  */
@@ -46,6 +41,9 @@ private:
 	TSharedRef<ITableRow> OnGenerateWidgetTemplateItem(TSharedPtr<FWidgetViewModel> Item, const TSharedRef<STableViewBase>& OwnerTable);
 
 private:
+	void LoadItemExpanssion();
+	void SaveItemExpansion();
+
 	void AddWidgetTemplate(TSharedPtr<FWidgetTemplate> Template);
 
 	TWeakPtr<class FBlueprintEditor> BlueprintEditor;
diff --git a/Engine/Source/Editor/UMGEditor/Private/UMGBlueprintEditorExtensionHook.cpp b/Engine/Source/Editor/UMGEditor/Private/UMGBlueprintEditorExtensionHook.cpp
--- a/Engine/Source/Editor/UMGEditor/Private/UMGBlueprintEditorExtensionHook.cpp
+++ b/Engine/Source/Editor/UMGEditor/Private/UMGBlueprintEditorExtensionHook.cpp
@@ -215,6 +215,10 @@ public:
 		FBlueprintDefaultsApplicationMode::PreDeactivateMode();
 
 		GetBlueprintEditor()->GetInspector()->EnableComponentDetailsCustomization(false);
+
+		static FName PropertyEditor("PropertyEditor");
+		FPropertyEditorModule& PropertyModule = FModuleManager::GetModuleChecked<FPropertyEditorModule>(PropertyEditor);
+		PropertyModule.UnregisterCustomPropertyTypeLayout(TEXT("CanvasPanelSlot"));
 	}
 
 	virtual void PostActivateMode() OVERRIDE
@@ -226,9 +230,12 @@ public:
 		TSharedRef<class SKismetInspector> Inspector = MyBlueprintEditor.Pin()->GetInspector();
 		FOnGetDetailCustomizationInstance LayoutDelegateDetails = FOnGetDetailCustomizationInstance::CreateStatic(&FBlueprintWidgetCustomization::MakeInstance, MyBlueprintEditor.Pin()->GetBlueprintObj());
 		Inspector->GetPropertyView()->RegisterInstancedCustomPropertyLayout(UWidget::StaticClass(), LayoutDelegateDetails);
+
+		FOnGetPropertyTypeCustomizationInstance CanvasSlotCustomization = FOnGetPropertyTypeCustomizationInstance::CreateStatic(&FCanvasSlotCustomization::MakeInstance, MyBlueprintEditor.Pin()->GetBlueprintObj());
 		
-		FOnGetDetailCustomizationInstance CanvasSlotCustomization = FOnGetDetailCustomizationInstance::CreateStatic(&FCanvasSlotCustomization::MakeInstance, MyBlueprintEditor.Pin()->GetBlueprintObj());
-		Inspector->GetPropertyView()->RegisterInstancedCustomPropertyLayout(UCanvasPanelSlot::StaticClass(), CanvasSlotCustomization);
+		static FName PropertyEditor("PropertyEditor");
+		FPropertyEditorModule& PropertyModule = FModuleManager::GetModuleChecked<FPropertyEditorModule>(PropertyEditor);
+		PropertyModule.RegisterCustomPropertyTypeLayout(TEXT("PanelSlot"), CanvasSlotCustomization);
 	}
 
 	UWidgetBlueprint* GetBlueprint() const
diff --git a/Engine/Source/Runtime/UMG/Public/Components/TextBlock.h b/Engine/Source/Runtime/UMG/Public/Components/TextBlock.h
--- a/Engine/Source/Runtime/UMG/Public/Components/TextBlock.h
+++ b/Engine/Source/Runtime/UMG/Public/Components/TextBlock.h
@@ -5,7 +5,7 @@
 #include "TextBlock.generated.h"
 
 /** A simple static text widget */
-UCLASS(meta=( BlueprintSpawnableComponent, Category="Common" ), ClassGroup=UserInterface)
+UCLASS(meta=( BlueprintSpawnableComponent, DisplayName="Text", Category="Common" ), ClassGroup=UserInterface )
 class UMG_API UTextBlock : public UWidget
 {
 	GENERATED_UCLASS_BODY()
diff --git a/Engine/Source/Runtime/UMG/Public/Slate/Anchors.h b/Engine/Source/Runtime/UMG/Public/Slate/Anchors.h
--- a/Engine/Source/Runtime/UMG/Public/Slate/Anchors.h
+++ b/Engine/Source/Runtime/UMG/Public/Slate/Anchors.h
@@ -49,4 +49,10 @@ public:
 		: Minimum(InLeft, InTop)
 		, Maximum(InRight, InBottom)
 	{ }
+
+	/** Returns true if the anchors represent a stretch along the vertical axis */
+	bool IsStretchedVertical() const { return Minimum.Y != Maximum.Y; }
+
+	/** Returns true if the anchors represent a stretch along the horizontal axis */
+	bool IsStretchedHorizontal() const { return Minimum.X != Maximum.X; }
 };
