UMG - Designer revamp, how properties are migrated, which objects are changed...etc.  Widget selection is now better at handling previews disappearing.  Lots of cleanup.  The designer now live edits the preview, and the changes are migrated back to the templates, probably a bit buggy right now.  All widgets must now implement SyncronizeProperties so that the designer can update the visuals without completely regenerating the GUI.

diff --git a/Engine/Source/Editor/UMGEditor/Private/Customizations/CanvasSlotCustomization.cpp b/Engine/Source/Editor/UMGEditor/Private/Customizations/CanvasSlotCustomization.cpp
--- a/Engine/Source/Editor/UMGEditor/Private/Customizations/CanvasSlotCustomization.cpp
+++ b/Engine/Source/Editor/UMGEditor/Private/Customizations/CanvasSlotCustomization.cpp
@@ -50,7 +50,7 @@ void FCanvasSlotCustomization::FillOutChildren(TSharedRef<IPropertyHandle> Prope
 
 void FCanvasSlotCustomization::CustomizeAnchors(TSharedRef<IPropertyHandle> PropertyHandle, IDetailChildrenBuilder& ChildBuilder, IPropertyTypeCustomizationUtils& CustomizationUtils)
 {
-	TSharedPtr<IPropertyHandle> Slot = PropertyHandle->GetChildHandle(GET_MEMBER_NAME_CHECKED(UCanvasPanelSlot, Anchors));
+	TSharedPtr<IPropertyHandle> Slot = PropertyHandle->GetChildHandle(GET_MEMBER_NAME_CHECKED(FAnchorData, Anchors));
 
 	if ( !Slot.IsValid() )
 	{
diff --git a/Engine/Source/Editor/UMGEditor/Private/Extensions/CanvasSlotExtension.cpp b/Engine/Source/Editor/UMGEditor/Private/Extensions/CanvasSlotExtension.cpp
--- a/Engine/Source/Editor/UMGEditor/Private/Extensions/CanvasSlotExtension.cpp
+++ b/Engine/Source/Editor/UMGEditor/Private/Extensions/CanvasSlotExtension.cpp
@@ -14,7 +14,7 @@ bool FCanvasSlotExtension::IsActive(const TArray< FSelectedWidget >& Selection)
 {
 	for ( const FSelectedWidget& Widget : Selection )
 	{
-		if ( !Widget.Template->Slot || !Widget.Template->Slot->IsA(UCanvasPanelSlot::StaticClass()) )
+		if ( !Widget.GetTemplate()->Slot || !Widget.GetTemplate()->Slot->IsA(UCanvasPanelSlot::StaticClass()) )
 		{
 			return false;
 		}
@@ -69,8 +69,8 @@ FReply FCanvasSlotExtension::HandleDragging(const FGeometry& Geometry, const FPo
 	{
 		for ( FSelectedWidget& Selection : SelectionCache )
 		{
-			MoveByAmount(Selection.Preview, Event.GetCursorDelta());
-			MoveByAmount(Selection.Template, Event.GetCursorDelta());
+			MoveByAmount(Selection.GetPreview(), Event.GetCursorDelta());
+			MoveByAmount(Selection.GetTemplate(), Event.GetCursorDelta());
 		}
 
 		return FReply::Handled();
@@ -86,11 +86,11 @@ void FCanvasSlotExtension::MoveByAmount(UWidget* Widget, FVector2D Delta)
 		UCanvasPanelSlot* CanvasSlot = Cast<UCanvasPanelSlot>(Widget->Slot);
 		UCanvasPanel* Parent = Cast<UCanvasPanel>(CanvasSlot->Parent);
 
-		FMargin Offset = CanvasSlot->Offset;
-		Offset.Left += Delta.X;
-		Offset.Top += Delta.Y;
+		FMargin Offsets = CanvasSlot->LayoutData.Offsets;
+		Offsets.Left += Delta.X;
+		Offsets.Top += Delta.Y;
 
-		CanvasSlot->SetOffset(Offset);
+		CanvasSlot->SetOffset(Offsets);
 	}
 }
 
diff --git a/Engine/Source/Editor/UMGEditor/Private/Extensions/DesignerExtension.cpp b/Engine/Source/Editor/UMGEditor/Private/Extensions/DesignerExtension.cpp
--- a/Engine/Source/Editor/UMGEditor/Private/Extensions/DesignerExtension.cpp
+++ b/Engine/Source/Editor/UMGEditor/Private/Extensions/DesignerExtension.cpp
@@ -2,8 +2,77 @@
 
 #include "UMGEditorPrivatePCH.h"
 
+#include "WidgetBlueprintEditor.h"
+
 #define LOCTEXT_NAMESPACE "UMG"
 
+// FSelectedWidget
+
+FSelectedWidget::FSelectedWidget()
+	: WidgetEditor()
+	, TemplateWidget(NULL)
+{
+}
+
+FSelectedWidget::FSelectedWidget(TSharedPtr<FWidgetBlueprintEditor> WidgetEditor, UWidget* TemplateWidget)
+	: WidgetEditor(WidgetEditor)
+	, TemplateWidget(TemplateWidget)
+{
+}
+
+FSelectedWidget FSelectedWidget::FromTemplate(TSharedPtr<FWidgetBlueprintEditor> WidgetEditor, UWidget* TemplateWidget)
+{
+	return FSelectedWidget(WidgetEditor, TemplateWidget);
+}
+
+FSelectedWidget FSelectedWidget::FromPreview(TSharedPtr<FWidgetBlueprintEditor> WidgetEditor, UWidget* PreviewWidget)
+{
+	if ( WidgetEditor.IsValid() )
+	{
+		UUserWidget* PreviewRoot = WidgetEditor->GetPreview();
+		if ( PreviewRoot )
+		{
+			UWidgetBlueprint* Blueprint = WidgetEditor->GetWidgetBlueprintObj();
+
+			if ( PreviewWidget )
+			{
+				FString Name = PreviewWidget->GetName();
+				return FromTemplate(WidgetEditor, Blueprint->WidgetTree->FindWidget(Name));
+			}
+		}
+	}
+
+	return FSelectedWidget();
+}
+
+bool FSelectedWidget::IsValid() const
+{
+	return TemplateWidget != NULL && GetPreview();
+}
+
+UWidget* FSelectedWidget::GetTemplate() const
+{
+	return TemplateWidget;
+}
+
+UWidget* FSelectedWidget::GetPreview() const
+{
+	if ( WidgetEditor.IsValid() )
+	{
+		UUserWidget* PreviewRoot = WidgetEditor->GetPreview();
+
+		if ( PreviewRoot && TemplateWidget )
+		{
+			UWidget* PreviewWidget = PreviewRoot->GetHandleFromName(TemplateWidget->GetName());
+			return PreviewWidget;
+		}
+	}
+
+	return NULL;
+}
+
+// Designer Extension
+
 void FDesignerExtension::Initialize(UWidgetBlueprint* InBlueprint)
 {
 	Blueprint = InBlueprint;
diff --git a/Engine/Source/Editor/UMGEditor/Private/Extensions/HorizontalSlotExtension.cpp b/Engine/Source/Editor/UMGEditor/Private/Extensions/HorizontalSlotExtension.cpp
--- a/Engine/Source/Editor/UMGEditor/Private/Extensions/HorizontalSlotExtension.cpp
+++ b/Engine/Source/Editor/UMGEditor/Private/Extensions/HorizontalSlotExtension.cpp
@@ -14,7 +14,7 @@ bool FHorizontalSlotExtension::IsActive(const TArray< FSelectedWidget >& Selecti
 {
 	for ( const FSelectedWidget& Widget : Selection )
 	{
-		if ( !Widget.Template->Slot || !Widget.Template->Slot->IsA(UHorizontalBoxSlot::StaticClass()) )
+		if ( !Widget.GetTemplate()->Slot || !Widget.GetTemplate()->Slot->IsA(UHorizontalBoxSlot::StaticClass()) )
 		{
 			return false;
 		}
@@ -50,8 +50,8 @@ FReply FHorizontalSlotExtension::HandleShift(int32 ShiftAmount)
 {
 	for ( FSelectedWidget& Selection : SelectionCache )
 	{
-		ShiftHorizontal(Selection.Preview, ShiftAmount);
-		ShiftHorizontal(Selection.Template, ShiftAmount);
+		ShiftHorizontal(Selection.GetPreview(), ShiftAmount);
+		ShiftHorizontal(Selection.GetTemplate(), ShiftAmount);
 	}
 
 	//TODO UMG Reorder the live slot without rebuilding the structure
diff --git a/Engine/Source/Editor/UMGEditor/Private/Extensions/UniformGridSlotExtension.cpp b/Engine/Source/Editor/UMGEditor/Private/Extensions/UniformGridSlotExtension.cpp
--- a/Engine/Source/Editor/UMGEditor/Private/Extensions/UniformGridSlotExtension.cpp
+++ b/Engine/Source/Editor/UMGEditor/Private/Extensions/UniformGridSlotExtension.cpp
@@ -13,7 +13,7 @@ bool FUniformGridSlotExtension::IsActive(const TArray< FSelectedWidget >& Select
 {
 	for ( const FSelectedWidget& Widget : Selection )
 	{
-		if ( !Widget.Template->Slot || !Widget.Template->Slot->IsA(UUniformGridSlot::StaticClass()) )
+		if ( !Widget.GetTemplate()->Slot || !Widget.GetTemplate()->Slot->IsA(UUniformGridSlot::StaticClass()) )
 		{
 			return false;
 		}
@@ -61,8 +61,8 @@ FReply FUniformGridSlotExtension::HandleShiftRow(int32 ShiftAmount)
 {
 	for ( FSelectedWidget& Selection : SelectionCache )
 	{
-		ShiftRow(Selection.Preview, ShiftAmount);
-		ShiftRow(Selection.Template, ShiftAmount);
+		ShiftRow(Selection.GetPreview(), ShiftAmount);
+		ShiftRow(Selection.GetTemplate(), ShiftAmount);
 	}
 
 	FBlueprintEditorUtils::MarkBlueprintAsModified(Blueprint);
@@ -74,8 +74,8 @@ FReply FUniformGridSlotExtension::HandleShiftColumn(int32 ShiftAmount)
 {
 	for ( FSelectedWidget& Selection : SelectionCache )
 	{
-		ShiftColumn(Selection.Preview, ShiftAmount);
-		ShiftColumn(Selection.Template, ShiftAmount);
+		ShiftColumn(Selection.GetPreview(), ShiftAmount);
+		ShiftColumn(Selection.GetTemplate(), ShiftAmount);
 	}
 
 	FBlueprintEditorUtils::MarkBlueprintAsModified(Blueprint);
diff --git a/Engine/Source/Editor/UMGEditor/Private/Extensions/VerticalSlotExtension.cpp b/Engine/Source/Editor/UMGEditor/Private/Extensions/VerticalSlotExtension.cpp
--- a/Engine/Source/Editor/UMGEditor/Private/Extensions/VerticalSlotExtension.cpp
+++ b/Engine/Source/Editor/UMGEditor/Private/Extensions/VerticalSlotExtension.cpp
@@ -14,7 +14,7 @@ bool FVerticalSlotExtension::IsActive(const TArray< FSelectedWidget >& Selection
 {
 	for ( const FSelectedWidget& Widget : Selection )
 	{
-		if ( !Widget.Template->Slot || !Widget.Template->Slot->IsA(UVerticalBoxSlot::StaticClass()) )
+		if ( !Widget.GetTemplate()->Slot || !Widget.GetTemplate()->Slot->IsA(UVerticalBoxSlot::StaticClass()) )
 		{
 			return false;
 		}
@@ -50,8 +50,8 @@ FReply FVerticalSlotExtension::HandleUpPressed()
 {
 	for ( FSelectedWidget& Selection : SelectionCache )
 	{
-		MoveUp(Selection.Preview);
-		MoveUp(Selection.Template);
+		MoveUp(Selection.GetPreview());
+		MoveUp(Selection.GetTemplate());
 	}
 
 	//TODO UMG Reorder the live slot without rebuilding the structure
@@ -64,8 +64,8 @@ FReply FVerticalSlotExtension::HandleDownPressed()
 {
 	for ( FSelectedWidget& Selection : SelectionCache )
 	{
-		MoveDown(Selection.Preview);
-		MoveDown(Selection.Template);
+		MoveDown(Selection.GetPreview());
+		MoveDown(Selection.GetTemplate());
 	}
 
 	//TODO UMG Reorder the live slot without rebuilding the structure
diff --git a/Engine/Source/Editor/UMGEditor/Private/SUMGDesigner.cpp b/Engine/Source/Editor/UMGEditor/Private/SUMGDesigner.cpp
--- a/Engine/Source/Editor/UMGEditor/Private/SUMGDesigner.cpp
+++ b/Engine/Source/Editor/UMGEditor/Private/SUMGDesigner.cpp
@@ -158,26 +158,7 @@ void SUMGDesigner::OnObjectPropertyChanged(UObject* ObjectBeingModified, FProper
 
 void SUMGDesigner::ShowDetailsForObjects(TArray<UWidget*> Widgets)
 {
-	// @TODO COde duplication
-
-	// Convert the selection set to an array of UObject* pointers
-	FString InspectorTitle;
-	TArray<UObject*> InspectorObjects;
-	InspectorObjects.Empty(Widgets.Num());
-	for ( UWidget* Widget : Widgets )
-	{
-		//if ( NodePtr->CanEditDefaults() )
-		{
-			InspectorTitle = "Widget";// Widget->GetDisplayString();
-			InspectorObjects.Add(Widget);
-		}
-	}
-
-	UWidgetBlueprint* Blueprint = GetBlueprint();
-
-	// Update the details panel
-	SKismetInspector::FShowDetailsOptions Options(InspectorTitle, true);
-	BlueprintEditor.Pin()->GetInspector()->ShowDetailsForObjects(InspectorObjects, Options);
+	BlueprintEditor.Pin()->SelectWidgets(Widgets);
 }
 
 FSelectedWidget SUMGDesigner::GetWidgetAtCursor(const FGeometry& MyGeometry, const FPointerEvent& MouseEvent, FArrangedWidget& ArrangedWidget)
@@ -192,32 +173,29 @@ FSelectedWidget SUMGDesigner::GetWidgetAtCursor(const FGeometry& MyGeometry, con
 
 	PreviewSurface->SetVisibility(EVisibility::HitTestInvisible);
 
-	FSelectedWidget Selection;
-
 	UUserWidget* WidgetActor = BlueprintEditor.Pin()->GetPreview();
 	if ( WidgetActor )
 	{
+		UWidget* Preview = NULL;
+
 		for ( int32 ChildIndex = Children.Num() - 1; ChildIndex >= 0; ChildIndex-- )
 		{
 			FArrangedWidget& Child = Children.GetInternalArray()[ChildIndex];
-			Selection.Preview = WidgetActor->GetWidgetHandle(Child.Widget);
-			if ( Selection.Preview )
+			Preview = WidgetActor->GetWidgetHandle(Child.Widget);
+			if ( Preview )
 			{
 				ArrangedWidget = Child;
 				break;
 			}
 		}
 
-		UWidgetBlueprint* Blueprint = GetBlueprint();
-
-		if ( Selection.Preview )
+		if ( Preview )
 		{
-			FString Name = Selection.Preview->GetName();
-			Selection.Template = Blueprint->WidgetTree->FindWidget(Name);
+			return FSelectedWidget::FromPreview(BlueprintEditor.Pin(), Preview);
 		}
 	}
 
-	return Selection;
+	return FSelectedWidget();
 }
 
 FReply SUMGDesigner::OnMouseButtonDown(const FGeometry& MyGeometry, const FPointerEvent& MouseEvent)
@@ -234,7 +212,8 @@ FReply SUMGDesigner::OnMouseButtonDown(const FGeometry& MyGeometry, const FPoint
 			//@TODO UMG primary FBlueprintEditor needs to be inherited and selection control needs to be centralized.
 			// Set the template as selected in the details panel
 			TArray<UWidget*> SelectedTemplates;
-			SelectedTemplates.Add(CurrentSelection.Template);
+			//SelectedTemplates.Add(CurrentSelection.Template);
+			SelectedTemplates.Add(CurrentSelection.GetPreview());
 			ShowDetailsForObjects(SelectedTemplates);
 
 			// Remove all the current extension widgets
@@ -264,7 +243,7 @@ FReply SUMGDesigner::OnMouseButtonDown(const FGeometry& MyGeometry, const FPoint
 			}
 
 			ExtensionWidgetCanvas->AddSlot()
-				.Position(TAttribute<FVector2D>(this, &SUMGDesigner::GetSelectionDesignerWidgetsLocation))
+				.Position(TAttribute<FVector2D>(this, &SUMGDesigner::GetCachedSelectionDesignerWidgetsLocation))
 				.Size(FVector2D(100, 25))
 				[
 					ExtensionBox
@@ -294,14 +273,17 @@ FReply SUMGDesigner::OnMouseMove(const FGeometry& MyGeometry, const FPointerEven
 			//@TODO UMG - Implement some system to query slots to know if they support dragging so we can provide visual feedback by hiding handles that wouldnt work and such.
 			//SelectedTemplate->Slot->CanResize(DH_TOP_CENTER)
 
-			if ( CurrentSelection.Preview->Slot )
+			UWidget* Template = CurrentSelection.GetTemplate();
+			UWidget* Preview = CurrentSelection.GetPreview();
+
+			if ( Preview->Slot )
 			{
-				CurrentSelection.Preview->Slot->Resize(DragDirections[CurrentHandle], MouseEvent.GetCursorDelta());
+				Preview->Slot->Resize(DragDirections[CurrentHandle], MouseEvent.GetCursorDelta());
 			}
 
-			if ( CurrentSelection.Template->Slot )
+			if ( Template->Slot )
 			{
-				CurrentSelection.Template->Slot->Resize(DragDirections[CurrentHandle], MouseEvent.GetCursorDelta());
+				Template->Slot->Resize(DragDirections[CurrentHandle], MouseEvent.GetCursorDelta());
 			}
 
 			FBlueprintEditorUtils::MarkBlueprintAsModified(GetBlueprint());
@@ -371,6 +353,11 @@ bool SUMGDesigner::GetArrangedWidgetRelativeToDesigner(TSharedRef<SWidget> Widge
 	return false;
 }
 
+FVector2D SUMGDesigner::GetCachedSelectionDesignerWidgetsLocation() const
+{
+	return CachedDesignerWidgetLocation;
+}
+
 FVector2D SUMGDesigner::GetSelectionDesignerWidgetsLocation() const
 {
 	if ( SelectedWidget.IsValid() )
@@ -418,7 +405,7 @@ int32 SUMGDesigner::OnPaint(const FGeometry& AllottedGeometry, const FSlateRect&
 
 		DrawDragHandles(SelectionGeometry, MyClippingRect, OutDrawElements, LayerId, InWidgetStyle);
 
-		if ( CurrentSelection.IsValid() && CurrentSelection.Template->Slot )
+		if ( CurrentSelection.IsValid() && CurrentSelection.GetTemplate()->Slot )
 		{
 			const float X = AllottedGeometry.AbsolutePosition.X;
 			const float Y = AllottedGeometry.AbsolutePosition.Y;
@@ -437,7 +424,7 @@ int32 SUMGDesigner::OnPaint(const FGeometry& AllottedGeometry, const FSlateRect&
 			const FSlateBrush* KeyBrush = FEditorStyle::GetBrush("CurveEd.CurveKey");
 			FLinearColor KeyColor = FLinearColor(1.0f, 0.0f, 0.0f, 1.f);
 
-			if ( UCanvasPanelSlot* CanvasSlot = Cast<UCanvasPanelSlot>(CurrentSelection.Template->Slot) )
+			if ( UCanvasPanelSlot* CanvasSlot = Cast<UCanvasPanelSlot>(CurrentSelection.GetTemplate()->Slot) )
 			{
 				//Set("UMGEditor.AnchorCenter", new IMAGE_BRUSH("Icons/umg_anchor_center", Icon16x16));
 				//Set("UMGEditor.AnchorTopBottom", new IMAGE_BRUSH("Icons/umg_anchor_top_bottom", FVector2D(16, 32)));
@@ -447,7 +434,7 @@ int32 SUMGDesigner::OnPaint(const FGeometry& AllottedGeometry, const FSlateRect&
 				FSlateDrawElement::MakeBox(
 					OutDrawElements,
 					++LayerId,
-					FPaintGeometry(FVector2D(X + Width * CanvasSlot->Anchors.Minimum.X - LeftRightSize.X * 0.5f, SelectionGeometry.DrawPosition.Y + Selection_Height * 0.5f - LeftRightSize.Y * 0.5f), LeftRightSize, 1.0f),
+					FPaintGeometry(FVector2D(X + Width * CanvasSlot->LayoutData.Anchors.Minimum.X - LeftRightSize.X * 0.5f, SelectionGeometry.DrawPosition.Y + Selection_Height * 0.5f - LeftRightSize.Y * 0.5f), LeftRightSize, 1.0f),
 					FEditorStyle::Get().GetBrush("UMGEditor.AnchorLeftRight"),
 					MyClippingRect,
 					ESlateDrawEffect::None,
@@ -458,7 +445,7 @@ int32 SUMGDesigner::OnPaint(const FGeometry& AllottedGeometry, const FSlateRect&
 				FSlateDrawElement::MakeBox(
 					OutDrawElements,
 					++LayerId,
-					FPaintGeometry(FVector2D(X + Width * CanvasSlot->Anchors.Maximum.X - LeftRightSize.X * 0.5f, SelectionGeometry.DrawPosition.Y + Selection_Height * 0.5f - LeftRightSize.Y * 0.5f), LeftRightSize, 1.0f),
+					FPaintGeometry(FVector2D(X + Width * CanvasSlot->LayoutData.Anchors.Maximum.X - LeftRightSize.X * 0.5f, SelectionGeometry.DrawPosition.Y + Selection_Height * 0.5f - LeftRightSize.Y * 0.5f), LeftRightSize, 1.0f),
 					FEditorStyle::Get().GetBrush("UMGEditor.AnchorLeftRight"),
 					MyClippingRect,
 					ESlateDrawEffect::None,
@@ -470,7 +457,7 @@ int32 SUMGDesigner::OnPaint(const FGeometry& AllottedGeometry, const FSlateRect&
 				FSlateDrawElement::MakeBox(
 					OutDrawElements,
 					++LayerId,
-					FPaintGeometry(FVector2D(SelectionGeometry.DrawPosition.X + Selection_Width * 0.5f - TopBottomSize.X * 0.5f, Y + Height * CanvasSlot->Anchors.Minimum.Y - TopBottomSize.Y * 0.5f), TopBottomSize, 1.0f),
+					FPaintGeometry(FVector2D(SelectionGeometry.DrawPosition.X + Selection_Width * 0.5f - TopBottomSize.X * 0.5f, Y + Height * CanvasSlot->LayoutData.Anchors.Minimum.Y - TopBottomSize.Y * 0.5f), TopBottomSize, 1.0f),
 					FEditorStyle::Get().GetBrush("UMGEditor.AnchorTopBottom"),
 					MyClippingRect,
 					ESlateDrawEffect::None,
@@ -481,7 +468,7 @@ int32 SUMGDesigner::OnPaint(const FGeometry& AllottedGeometry, const FSlateRect&
 				FSlateDrawElement::MakeBox(
 					OutDrawElements,
 					++LayerId,
-					FPaintGeometry(FVector2D(SelectionGeometry.DrawPosition.X + Selection_Width * 0.5f - TopBottomSize.X * 0.5f, Y + Height * CanvasSlot->Anchors.Maximum.Y - TopBottomSize.Y * 0.5f), TopBottomSize, 1.0f),
+					FPaintGeometry(FVector2D(SelectionGeometry.DrawPosition.X + Selection_Width * 0.5f - TopBottomSize.X * 0.5f, Y + Height * CanvasSlot->LayoutData.Anchors.Maximum.Y - TopBottomSize.Y * 0.5f), TopBottomSize, 1.0f),
 					FEditorStyle::Get().GetBrush("UMGEditor.AnchorTopBottom"),
 					MyClippingRect,
 					ESlateDrawEffect::None,
@@ -496,7 +483,7 @@ int32 SUMGDesigner::OnPaint(const FGeometry& AllottedGeometry, const FSlateRect&
 
 void SUMGDesigner::DrawDragHandles(const FPaintGeometry& SelectionGeometry, const FSlateRect& MyClippingRect, FSlateWindowElementList& OutDrawElements, int32 LayerId, const FWidgetStyle& InWidgetStyle) const
 {
-	if ( CurrentSelection.IsValid() && CurrentSelection.Template->Slot )
+	if ( CurrentSelection.IsValid() && CurrentSelection.GetTemplate()->Slot )
 	{
 		float X = SelectionGeometry.DrawPosition.X;
 		float Y = SelectionGeometry.DrawPosition.Y;
@@ -525,7 +512,7 @@ void SUMGDesigner::DrawDragHandles(const FPaintGeometry& SelectionGeometry, cons
 		for ( int32 HandleIndex = 0; HandleIndex < Handles.Num(); HandleIndex++ )
 		{
 			const FVector2D& Handle = Handles[HandleIndex];
-			if ( !CurrentSelection.Template->Slot->CanResize(DragDirections[HandleIndex]) )
+			if ( !CurrentSelection.GetTemplate()->Slot->CanResize(DragDirections[HandleIndex]) )
 			{
 				// This isn't a valid handle
 				continue;
@@ -548,7 +535,7 @@ SUMGDesigner::DragHandle SUMGDesigner::HitTestDragHandles(const FGeometry& Allot
 {
 	FVector2D LocalPointer = AllottedGeometry.AbsoluteToLocal(PointerEvent.GetScreenSpacePosition());
 
-	if ( CurrentSelection.IsValid() && CurrentSelection.Template->Slot && SelectedWidget.IsValid() )
+	if ( CurrentSelection.IsValid() && CurrentSelection.GetTemplate()->Slot && SelectedWidget.IsValid() )
 	{
 		TSharedRef<SWidget> Widget = SelectedWidget.Pin().ToSharedRef();
 
@@ -581,7 +568,7 @@ SUMGDesigner::DragHandle SUMGDesigner::HitTestDragHandles(const FGeometry& Allot
 			FSlateRect Rect(FVector2D(Handle.X - HandleSize.X * 0.5f, Handle.Y - HandleSize.Y * 0.5f), Handle + HandleSize);
 			if ( Rect.ContainsPoint(LocalPointer) )
 			{
-				if ( !CurrentSelection.Template->Slot->CanResize(DragDirections[i]) )
+				if ( !CurrentSelection.GetTemplate()->Slot->CanResize(DragDirections[i]) )
 				{
 					// This isn't a valid handle
 					break;
@@ -660,10 +647,12 @@ void SUMGDesigner::Tick(const FGeometry& AllottedGeometry, const double InCurren
 			//@TODO UMG Don't store the transient Widget as the selected object, store the template component and lookup the widget from that.
 			//SelectedWidget = ArrangedWidget.Widget;
 
-			SelectedWidget = WidgetActor->GetWidgetFromName(CurrentSelection.Template->GetName());
+			SelectedWidget = WidgetActor->GetWidgetFromName(CurrentSelection.GetTemplate()->GetName());
 		}
 	}
 
+	CachedDesignerWidgetLocation = GetSelectionDesignerWidgetsLocation();
+
 	SCompoundWidget::Tick(AllottedGeometry, InCurrentTime, InDeltaTime);
 }
 
@@ -703,9 +692,9 @@ FReply SUMGDesigner::OnDrop(const FGeometry& MyGeometry, const FDragDropEvent& D
 
 		UWidgetBlueprint* BP = GetBlueprint();
 
-		if ( Selection.IsValid() && Selection.Template->IsA(UPanelWidget::StaticClass()) )
+		if ( Selection.IsValid() && Selection.GetTemplate()->IsA(UPanelWidget::StaticClass()) )
 		{
-			UPanelWidget* Parent = Cast<UPanelWidget>(Selection.Template);
+			UPanelWidget* Parent = Cast<UPanelWidget>(Selection.GetTemplate());
 
 			UWidget* Widget = DragDropOp->Template->Create(BP->WidgetTree);
 
diff --git a/Engine/Source/Editor/UMGEditor/Private/SUMGDesigner.h b/Engine/Source/Editor/UMGEditor/Private/SUMGDesigner.h
--- a/Engine/Source/Editor/UMGEditor/Private/SUMGDesigner.h
+++ b/Engine/Source/Editor/UMGEditor/Private/SUMGDesigner.h
@@ -50,6 +50,7 @@ private:
 	bool GetArrangedWidgetRelativeToDesigner(TSharedRef<SWidget> Widget, FArrangedWidget& ArrangedWidget) const;
 
 	FVector2D GetSelectionDesignerWidgetsLocation() const;
+	FVector2D GetCachedSelectionDesignerWidgetsLocation() const;
 
 private:
 	enum DragHandle
@@ -95,4 +96,6 @@ private:
 	TSharedPtr<SCanvas> ExtensionWidgetCanvas;
 
 	DragHandle CurrentHandle;
+
+	FVector2D CachedDesignerWidgetLocation;
 };
diff --git a/Engine/Source/Editor/UMGEditor/Private/SUMGEditorTree.cpp b/Engine/Source/Editor/UMGEditor/Private/SUMGEditorTree.cpp
--- a/Engine/Source/Editor/UMGEditor/Private/SUMGEditorTree.cpp
+++ b/Engine/Source/Editor/UMGEditor/Private/SUMGEditorTree.cpp
@@ -16,7 +16,7 @@
 
 #define LOCTEXT_NAMESPACE "UMG"
 
-void SUMGEditorTree::Construct(const FArguments& InArgs, TSharedPtr<FBlueprintEditor> InBlueprintEditor, USimpleConstructionScript* InSCS)
+void SUMGEditorTree::Construct(const FArguments& InArgs, TSharedPtr<FWidgetBlueprintEditor> InBlueprintEditor, USimpleConstructionScript* InSCS)
 {
 	BlueprintEditor = InBlueprintEditor;
 
@@ -95,26 +95,15 @@ void SUMGEditorTree::OnObjectPropertyChanged(UObject* ObjectBeingModified, FProp
 	}
 }
 
-void SUMGEditorTree::ShowDetailsForObjects(TArray<UWidget*> Widgets)
+void SUMGEditorTree::ShowDetailsForObjects(TArray<UWidget*> TemplateWidgets)
 {
-	// Convert the selection set to an array of UObject* pointers
-	FString InspectorTitle;
-	TArray<UObject*> InspectorObjects;
-	InspectorObjects.Empty(Widgets.Num());
-	for ( UWidget* Widget : Widgets )
+	TArray<UWidget*> PreviewWidgets;
+	for ( UWidget* TemplateWidget : TemplateWidgets )
 	{
-		//if ( NodePtr->CanEditDefaults() )
-		{
-			InspectorTitle = "Widget";// Widget->GetDisplayString();
-			InspectorObjects.Add(Widget);
-		}
+		PreviewWidgets.Add(FSelectedWidget::FromTemplate(BlueprintEditor.Pin(), TemplateWidget).GetPreview());
 	}
 
-	UWidgetBlueprint* Blueprint = GetBlueprint();
-
-	// Update the details panel
-	SKismetInspector::FShowDetailsOptions Options(InspectorTitle, true);
-	BlueprintEditor.Pin()->GetInspector()->ShowDetailsForObjects(InspectorObjects, Options);
+	BlueprintEditor.Pin()->SelectWidgets(PreviewWidgets);
 }
 
 void SUMGEditorTree::BuildWrapWithMenu(FMenuBuilder& Menu)
@@ -256,7 +245,7 @@ FReply SUMGEditorTree::CreateTestUI()
 		Button4->SetContent(Text1);
 
 		UCanvasPanelSlot* Slot = Canvas->AddSlot(Border);
-		Slot->Offset = FMargin(200, 300, 20, 50);
+		Slot->LayoutData.Offsets = FMargin(200, 300, 20, 50);
 
 		Border->SetContent(Vertical);
 
diff --git a/Engine/Source/Editor/UMGEditor/Private/SUMGEditorTree.h b/Engine/Source/Editor/UMGEditor/Private/SUMGEditorTree.h
--- a/Engine/Source/Editor/UMGEditor/Private/SUMGEditorTree.h
+++ b/Engine/Source/Editor/UMGEditor/Private/SUMGEditorTree.h
@@ -12,7 +12,7 @@ public:
 	SLATE_BEGIN_ARGS( SUMGEditorTree ){}
 	SLATE_END_ARGS()
 
-	void Construct(const FArguments& InArgs, TSharedPtr<FBlueprintEditor> InBlueprintEditor, USimpleConstructionScript* InSCS);
+	void Construct(const FArguments& InArgs, TSharedPtr<FWidgetBlueprintEditor> InBlueprintEditor, USimpleConstructionScript* InSCS);
 	virtual ~SUMGEditorTree();
 
 private:
@@ -39,7 +39,7 @@ private:
 	
 private:
 
-	TWeakPtr<class FBlueprintEditor> BlueprintEditor;
+	TWeakPtr<class FWidgetBlueprintEditor> BlueprintEditor;
 
 	TArray< UWidget* > RootWidgets;
 
diff --git a/Engine/Source/Editor/UMGEditor/Private/UMGBlueprintEditorExtensionHook.cpp b/Engine/Source/Editor/UMGEditor/Private/UMGBlueprintEditorExtensionHook.cpp
--- a/Engine/Source/Editor/UMGEditor/Private/UMGBlueprintEditorExtensionHook.cpp
+++ b/Engine/Source/Editor/UMGEditor/Private/UMGBlueprintEditorExtensionHook.cpp
@@ -90,7 +90,9 @@ public:
 
 	virtual TSharedRef<SWidget> CreateTabBody(const FWorkflowTabSpawnInfo& Info) const OVERRIDE
 	{
-		return SNew(SUMGEditorTree, BlueprintEditor.Pin(), BlueprintEditor.Pin()->GetBlueprintObj()->SimpleConstructionScript);
+		TSharedPtr<FWidgetBlueprintEditor> BlueprintEditorPtr = StaticCastSharedPtr<FWidgetBlueprintEditor>(BlueprintEditor.Pin());
+
+		return SNew(SUMGEditorTree, BlueprintEditorPtr, BlueprintEditorPtr->GetBlueprintObj()->SimpleConstructionScript);
 	}
 };
 
@@ -120,7 +122,9 @@ public:
 
 	virtual TSharedRef<SWidget> CreateTabBody(const FWorkflowTabSpawnInfo& Info) const OVERRIDE
 	{
-		return SNew(SUMGEditorWidgetTemplates, BlueprintEditor.Pin(), BlueprintEditor.Pin()->GetBlueprintObj()->SimpleConstructionScript);
+		TSharedPtr<FWidgetBlueprintEditor> BlueprintEditorPtr = StaticCastSharedPtr<FWidgetBlueprintEditor>(BlueprintEditor.Pin());
+
+		return SNew(SUMGEditorWidgetTemplates, BlueprintEditorPtr, BlueprintEditorPtr->GetBlueprintObj()->SimpleConstructionScript);
 	}
 };
 
diff --git a/Engine/Source/Editor/UMGEditor/Private/WidgetBlueprintEditor.cpp b/Engine/Source/Editor/UMGEditor/Private/WidgetBlueprintEditor.cpp
--- a/Engine/Source/Editor/UMGEditor/Private/WidgetBlueprintEditor.cpp
+++ b/Engine/Source/Editor/UMGEditor/Private/WidgetBlueprintEditor.cpp
@@ -2,6 +2,7 @@
 
 #include "UMGEditorPrivatePCH.h"
 
+#include "SKismetInspector.h"
 #include "WidgetBlueprintEditor.h"
 #include "MovieScene.h"
 #include "Editor/Sequencer/Public/ISequencerModule.h"
@@ -24,6 +25,30 @@ FWidgetBlueprintEditor::~FWidgetBlueprintEditor()
 	}
 }
 
+void FWidgetBlueprintEditor::SelectWidgets(TArray<UWidget*> Widgets)
+{
+	// Convert the selection set to an array of UObject* pointers
+	FString InspectorTitle;
+	TArray<UObject*> InspectorObjects;
+	InspectorObjects.Empty(Widgets.Num());
+	for ( UWidget* Widget : Widgets )
+	{
+		//if ( NodePtr->CanEditDefaults() )
+		{
+			InspectorTitle = "Widget";// Widget->GetDisplayString();
+			InspectorObjects.Add(Widget);
+		}
+	}
+
+	UWidgetBlueprint* Blueprint = GetWidgetBlueprintObj();
+
+	// Update the details panel
+	SKismetInspector::FShowDetailsOptions Options(InspectorTitle, true);
+	GetInspector()->ShowDetailsForObjects(InspectorObjects, Options);
+
+	SelectedPreviewWidgets = Widgets;
+}
+
 void FWidgetBlueprintEditor::OnBlueprintChanged(UBlueprint* InBlueprint)
 {
 	if ( InBlueprint )
@@ -64,9 +89,64 @@ void FWidgetBlueprintEditor::Tick(float DeltaTime)
 	}
 }
 
-void FWidgetBlueprintEditor::NotifyPostChange(const FPropertyChangedEvent& PropertyChangedEvent, UProperty* PropertyThatChanged)
+static bool MigratePropertyValue(UObject* SourceObject, UObject* DestinationObject, UProperty* MemberProperty)
 {
-	UpdatePreview(GetWidgetBlueprintObj(), true);
+	FString SourceValue;
+
+	// Get the property addresses for the source and destination objects.
+	uint8* SourceAddr = MemberProperty->ContainerPtrToValuePtr<uint8>(SourceObject);
+	uint8* DestionationAddr = MemberProperty->ContainerPtrToValuePtr<uint8>(DestinationObject);
+
+	if ( SourceAddr == NULL || DestionationAddr == NULL )
+	{
+		return false;
+	}
+
+	// Get the current value from the source object.
+	MemberProperty->ExportText_Direct(SourceValue, SourceAddr, SourceAddr, NULL, PPF_Localized);
+
+	const bool bNotifyObjectOfChange = true;
+
+	if ( !DestinationObject->HasAnyFlags(RF_ClassDefaultObject) && bNotifyObjectOfChange )
+	{
+		DestinationObject->PreEditChange(MemberProperty);
+	}
+
+	// Set the value on the destination object.
+	MemberProperty->ImportText(*SourceValue, DestionationAddr, 0, DestinationObject);
+
+	if ( !DestinationObject->HasAnyFlags(RF_ClassDefaultObject) && bNotifyObjectOfChange )
+	{
+		FPropertyChangedEvent PropertyEvent(MemberProperty);
+		DestinationObject->PostEditChangeProperty(PropertyEvent);
+	}
+
+	return true;
+}
+
+static bool MigratePropertyValue(UObject* SourceObject, UObject* DestinationObject, FEditPropertyChain::TDoubleLinkedListNode* PropertyChainNode, UProperty* MemberProperty)
+{
+	UProperty* CurrentProperty = PropertyChainNode->GetValue();
+
+	if ( PropertyChainNode->GetNextNode() == NULL )
+	{
+		return MigratePropertyValue(SourceObject, DestinationObject, MemberProperty);
+	}
+	
+	if ( UObjectProperty* CurrentObjectProperty = Cast<UObjectProperty>(CurrentProperty) )
+	{
+		// Get the property addresses for the source and destination objects.
+		UObject* SourceObjectProperty = CurrentObjectProperty->GetObjectPropertyValue(CurrentObjectProperty->ContainerPtrToValuePtr<void>(SourceObject));
+		UObject* DestionationObjectProperty = CurrentObjectProperty->GetObjectPropertyValue(CurrentObjectProperty->ContainerPtrToValuePtr<void>(DestinationObject));
+
+		return MigratePropertyValue(SourceObjectProperty, DestionationObjectProperty, PropertyChainNode->GetNextNode(), PropertyChainNode->GetNextNode()->GetValue());
+	}
+	else if ( UArrayProperty* CurrentArrayProperty = Cast<UArrayProperty>(CurrentProperty) )
+	{
+		// Arrays!
+	}
+
+	return MigratePropertyValue(SourceObject, DestinationObject, PropertyChainNode->GetNextNode(), MemberProperty);
 }
 
 void FWidgetBlueprintEditor::AddReferencedObjects( FReferenceCollector& Collector )
@@ -78,6 +158,32 @@ void FWidgetBlueprintEditor::AddReferencedObjects( FReferenceCollector& Collecto
 	}
 }
 
+void FWidgetBlueprintEditor::NotifyPostChange(const FPropertyChangedEvent& PropertyChangedEvent, class FEditPropertyChain* PropertyThatChanged)
+{
+	//Super::NotifyPostChange(PropertyChangedEvent, PropertyThatChanged);
+
+	if ( PropertyChangedEvent.ChangeType == EPropertyChangeType::ValueSet )
+	{
+		UWidgetBlueprint* Blueprint = GetWidgetBlueprintObj();
+
+		UUserWidget* PreviewActor = GetPreview();
+		if ( PreviewActor != NULL )
+		{
+			for ( UWidget* PreviewWidget : SelectedPreviewWidgets )
+			{
+				FString PreviewWidgetName = PreviewWidget->GetName();
+				UWidget* TemplateWidget = Blueprint->WidgetTree->FindWidget(PreviewWidgetName);
+
+				if ( TemplateWidget )
+				{
+					FEditPropertyChain::TDoubleLinkedListNode* PropertyChainNode = PropertyThatChanged->GetHead();
+					MigratePropertyValue(PreviewWidget, TemplateWidget, PropertyChainNode, PropertyChainNode->GetValue());
+				}
+			}
+		}
+	}
+}
+
 UWidgetBlueprint* FWidgetBlueprintEditor::GetWidgetBlueprintObj() const
 {
 	return Cast<UWidgetBlueprint>(GetBlueprintObj());
diff --git a/Engine/Source/Editor/UMGEditor/Private/WidgetBlueprintEditor.h b/Engine/Source/Editor/UMGEditor/Private/WidgetBlueprintEditor.h
--- a/Engine/Source/Editor/UMGEditor/Private/WidgetBlueprintEditor.h
+++ b/Engine/Source/Editor/UMGEditor/Private/WidgetBlueprintEditor.h
@@ -21,7 +21,7 @@ public:
 	void InitWidgetBlueprintEditor(const EToolkitMode::Type Mode, const TSharedPtr< IToolkitHost >& InitToolkitHost, const TArray<UBlueprint*>& InBlueprints, bool bShouldOpenInDefaultsMode);
 
 	virtual void Tick(float DeltaTime) OVERRIDE;
-	virtual void NotifyPostChange(const FPropertyChangedEvent& PropertyChangedEvent, UProperty* PropertyThatChanged) OVERRIDE;
+	virtual void NotifyPostChange(const FPropertyChangedEvent& PropertyChangedEvent, class FEditPropertyChain* PropertyThatChanged) OVERRIDE;
 
 	/** FGCObjectInterface */
 	virtual void AddReferencedObjects( FReferenceCollector& Collector ) OVERRIDE;
@@ -33,6 +33,8 @@ public:
 	/** @return The sequencer used to create widget animations */
 	TSharedPtr<ISequencer>& GetSequencer();
 
+	void SelectWidgets(TArray<UWidget*> Widgets);
+
 private:
 	void OnBlueprintChanged(UBlueprint* InBlueprint);
 
@@ -57,5 +59,7 @@ private:
 	/** Default movie scene for new animations */
 	UMovieScene* DefaultMovieScene;
 
+	TArray<UWidget*> SelectedPreviewWidgets;
+
 	mutable TWeakObjectPtr<UUserWidget> PreviewWidgetActorPtr;
 };
diff --git a/Engine/Source/Editor/UMGEditor/Public/DesignerExtension.h b/Engine/Source/Editor/UMGEditor/Public/DesignerExtension.h
--- a/Engine/Source/Editor/UMGEditor/Public/DesignerExtension.h
+++ b/Engine/Source/Editor/UMGEditor/Public/DesignerExtension.h
@@ -6,21 +6,30 @@ class UWidget;
 class SWidget;
 class UWidgetBlueprint;
 
+/**
+ * The Selected widget is a useful way to hold onto the selection in a way that allows for up to date access to the current preview object.
+ * Because the designer could end up rebuilding the preview, it's best to hold onto an FSelectedWidget.
+ */
 struct FSelectedWidget
 {
 public:
-	FSelectedWidget()
-		: Preview(NULL)
-		, Template(NULL)
-	{ }
-
-	bool IsValid() const
-	{
-		return Template != NULL && Preview != NULL;
-	}
-
-	UWidget* Template;
-	UWidget* Preview;
+	static FSelectedWidget FromTemplate(TSharedPtr<class FWidgetBlueprintEditor> WidgetEditor, UWidget* TemplateWidget);
+	static FSelectedWidget FromPreview(TSharedPtr<class FWidgetBlueprintEditor> WidgetEditor, UWidget* PreviewWidget);
+
+	FSelectedWidget();
+
+	bool IsValid() const;
+
+	UWidget* GetTemplate() const;
+	UWidget* GetPreview() const;
+
+private:
+	FSelectedWidget(TSharedPtr<class FWidgetBlueprintEditor> WidgetEditor, UWidget* TemplateWidget);
+
+private:
+
+	TSharedPtr<class FWidgetBlueprintEditor> WidgetEditor;
+	UWidget* TemplateWidget;
 };
 
 class FDesignerExtension : public TSharedFromThis<FDesignerExtension>
diff --git a/Engine/Source/Runtime/UMG/Private/Components/Border.cpp b/Engine/Source/Runtime/UMG/Private/Components/Border.cpp
--- a/Engine/Source/Runtime/UMG/Private/Components/Border.cpp
+++ b/Engine/Source/Runtime/UMG/Private/Components/Border.cpp
@@ -18,10 +18,10 @@ UBorder::UBorder(const FPostConstructInitializeProperties& PCIP)
 	HorizontalAlignment = BorderDefaults._HAlign;
 	VerticalAlignment = BorderDefaults._VAlign;
 	ContentPadding = BorderDefaults._Padding.Get();
-	//BorderImage = *BorderDefaults._BorderImage.Get();
-	DesiredSizeScale = BorderDefaults._DesiredSizeScale.Get();
 	ContentScale = BorderDefaults._ContentScale.Get();
 	ContentColorAndOpacity = FLinearColor::White;
+
+	DesiredSizeScale = BorderDefaults._DesiredSizeScale.Get();
 	BorderColorAndOpacity = FLinearColor::White;
 	ForegroundColor = FLinearColor::Black;
 	bShowEffectWhenDisabled = BorderDefaults._ShowEffectWhenDisabled.Get();
@@ -29,49 +29,50 @@ UBorder::UBorder(const FPostConstructInitializeProperties& PCIP)
 
 TSharedRef<SWidget> UBorder::RebuildWidget()
 {
-	TSharedRef<SBorder> NewBorder = SNew(SBorder)
-		.HAlign(HorizontalAlignment)
-		.VAlign(VerticalAlignment)
-		.Padding(BIND_UOBJECT_ATTRIBUTE(FMargin, GetContentPadding))
+	MyBorder = SNew(SBorder);
+	MyBorder->SetContent(Content ? Content->GetWidget() : SNullWidget::NullWidget);
+
+	return MyBorder.ToSharedRef();
+}
+
+void UBorder::SyncronizeProperties()
+{
+	Super::SyncronizeProperties();
+
+	if ( MyBorder.IsValid() )
+	{
+		MyBorder->SetHAlign(HorizontalAlignment);
+		MyBorder->SetVAlign(VerticalAlignment);
+		MyBorder->SetPadding(ContentPadding);
+
+		MyBorder->SetBorderBackgroundColor(BorderColorAndOpacity);
+		MyBorder->SetColorAndOpacity(ContentColorAndOpacity);
+		MyBorder->SetForegroundColor(ForegroundColor);
+
+		MyBorder->SetContentScale(ContentScale);
+		MyBorder->SetDesiredSizeScale(DesiredSizeScale);
+
+		MyBorder->SetShowEffectWhenDisabled(bShowEffectWhenDisabled);
+
+		MyBorder->SetBorderImage(GetBorderBrush());
+
 		// SLATE_EVENT( FPointerEventHandler, OnMouseButtonDown )
 		// SLATE_EVENT( FPointerEventHandler, OnMouseButtonUp )
 		// SLATE_EVENT( FPointerEventHandler, OnMouseMove )
 		// SLATE_EVENT( FPointerEventHandler, OnMouseDoubleClick )
-		.BorderImage(BIND_UOBJECT_ATTRIBUTE(const FSlateBrush*, GetBorderBrush))
-		.ContentScale(ContentScale)
-		.DesiredSizeScale(DesiredSizeScale)
-		.ColorAndOpacity(BIND_UOBJECT_ATTRIBUTE(FLinearColor, GetContentColor))
-		.BorderBackgroundColor(BIND_UOBJECT_ATTRIBUTE(FSlateColor, GetBorderColor))
-		.ForegroundColor(BIND_UOBJECT_ATTRIBUTE(FSlateColor, GetForegroundColor))
-		.ShowEffectWhenDisabled(bShowEffectWhenDisabled);
-
-	NewBorder->SetContent(Content ? Content->GetWidget() : SNullWidget::NullWidget);
-
-	MyBorder = NewBorder;
-	return NewBorder;
+	}
 }
 
 void UBorder::SetContent(UWidget* Content)
 {
 	Super::SetContent(Content);
 
-	TSharedPtr<SBorder> Border = MyBorder.Pin();
-	if ( Border.IsValid() )
+	if ( MyBorder.IsValid() )
 	{
-		Border->SetContent(Content ? Content->GetWidget() : SNullWidget::NullWidget);
+		MyBorder->SetContent(Content ? Content->GetWidget() : SNullWidget::NullWidget);
 	}
 }
 
-FMargin UBorder::GetContentPadding() const
-{
-	return ContentPadding;
-}
-
-FLinearColor UBorder::GetContentColor() const
-{
-	return ContentColorAndOpacity;
-}
-
 const FSlateBrush* UBorder::GetBorderBrush() const
 {
 	if ( BorderBrush == NULL )
@@ -83,16 +84,6 @@ const FSlateBrush* UBorder::GetBorderBrush() const
 	return &BorderBrush->Brush;
 }
 
-FSlateColor UBorder::GetBorderColor() const
-{
-	return BorderColorAndOpacity;
-}
-
-FSlateColor UBorder::GetForegroundColor() const
-{
-	return ForegroundColor;
-}
-
 /////////////////////////////////////////////////////
 
 #undef LOCTEXT_NAMESPACE
diff --git a/Engine/Source/Runtime/UMG/Private/Components/Button.cpp b/Engine/Source/Runtime/UMG/Private/Components/Button.cpp
--- a/Engine/Source/Runtime/UMG/Private/Components/Button.cpp
+++ b/Engine/Source/Runtime/UMG/Private/Components/Button.cpp
@@ -27,10 +27,20 @@ UButton::UButton(const FPostConstructInitializeProperties& PCIP)
 
 TSharedRef<SWidget> UButton::RebuildWidget()
 {
+	MyButton = SNew(SButton);
+	MyButton->SetContent(Content ? Content->GetWidget() : SNullWidget::NullWidget);
+	
+	return MyButton.ToSharedRef();
+}
+
+void UButton::SyncronizeProperties()
+{
+	Super::SyncronizeProperties();
+
 	SButton::FArguments ButtonDefaults;
 
 	const FButtonStyle* StylePtr = ( Style != NULL ) ? Style->GetStyle<FButtonStyle>() : NULL;
-	if (StylePtr == NULL)
+	if ( StylePtr == NULL )
 	{
 		StylePtr = ButtonDefaults._ButtonStyle;
 	}
@@ -47,63 +57,46 @@ TSharedRef<SWidget> UButton::RebuildWidget()
 		OptionalHoveredSound = HoveredSound;
 	}
 
-	TSharedRef<SButton> NewButton = SNew(SButton)
-		.ButtonStyle(StylePtr)
-		.HAlign(HorizontalAlignment)
-		.VAlign(VerticalAlignment)
-		.ContentPadding(ContentPadding)
-		.DesiredSizeScale(DesiredSizeScale)
-		.ContentScale(ContentScale)
-		.ButtonColorAndOpacity(ButtonColorAndOpacity)
-		.ForegroundColor(ForegroundColor)
-		.PressedSoundOverride(OptionalPressedSound)
-		.HoveredSoundOverride(OptionalHoveredSound)
-		.OnClicked(BIND_UOBJECT_DELEGATE(FOnClicked, HandleOnClicked))
-		;
-
-	NewButton->SetContent(Content ? Content->GetWidget() : SNullWidget::NullWidget);
-	
-	return NewButton;
+	if ( MyButton.IsValid() )
+	{
+		MyButton->SetButtonStyle(StylePtr);
+		MyButton->SetHAlign(HorizontalAlignment);
+		MyButton->SetVAlign(VerticalAlignment);
+		MyButton->SetContentPadding(ContentPadding);
+		MyButton->SetForegroundColor(ForegroundColor);
+		//MyButton->SetColorAndOpacity(ButtonColorAndOpacity);
+		MyButton->SetDesiredSizeScale(DesiredSizeScale);
+		MyButton->SetContentScale(ContentScale);
+		MyButton->SetBorderBackgroundColor(ButtonColorAndOpacity);
+		MyButton->SetPressedSound(OptionalPressedSound);
+		MyButton->SetHoveredSound(OptionalHoveredSound);
+		MyButton->SetOnClicked(BIND_UOBJECT_DELEGATE(FOnClicked, HandleOnClicked));
+	}
 }
 
 void UButton::SetContent(UWidget* InContent)
 {
 	Super::SetContent(InContent);
 
-	if ( ButtonWidget().IsValid() )
+	if ( MyButton.IsValid() )
 	{
-		ButtonWidget()->SetContent(InContent ? InContent->GetWidget() : SNullWidget::NullWidget);
+		MyButton->SetContent(InContent ? InContent->GetWidget() : SNullWidget::NullWidget);
 	}
 }
 
-FLinearColor UButton::GetButtonColorAndOpacity()
-{
-	return ButtonWidget()->GetBorderBackgroundColor().GetSpecifiedColor();
-}
-
 void UButton::SetButtonColorAndOpacity(FLinearColor InButtonColorAndOpacity)
 {
-	ButtonWidget()->SetBorderBackgroundColor(InButtonColorAndOpacity);
-}
-
-FLinearColor UButton::GetForegroundColor()
-{
-	return ButtonWidget()->GetForegroundColor().GetSpecifiedColor();
+	MyButton->SetBorderBackgroundColor(InButtonColorAndOpacity);
 }
 
 void UButton::SetForegroundColor(FLinearColor InForegroundColor)
 {
-	ButtonWidget()->SetForegroundColor(InForegroundColor);
-}
-
-FMargin UButton::GetContentPadding()
-{
-	return ButtonWidget()->GetContentPadding();
+	MyButton->SetForegroundColor(InForegroundColor);
 }
 
 void UButton::SetContentPadding(FMargin InContentPadding)
 {
-	ButtonWidget()->SetContentPadding(InContentPadding);
+	MyButton->SetContentPadding(InContentPadding);
 }
 
 FReply UButton::HandleOnClicked()
diff --git a/Engine/Source/Runtime/UMG/Private/Components/CanvasPanel.cpp b/Engine/Source/Runtime/UMG/Private/Components/CanvasPanel.cpp
--- a/Engine/Source/Runtime/UMG/Private/Components/CanvasPanel.cpp
+++ b/Engine/Source/Runtime/UMG/Private/Components/CanvasPanel.cpp
@@ -51,6 +51,11 @@ UWidget* UCanvasPanel::GetChildAt(int32 Index) const
 	return Slots[Index]->Content;
 }
 
+TSharedPtr<SFixedSizeCanvas> UCanvasPanel::GetCanvasWidget() const
+{
+	return MyCanvas.Pin();
+}
+
 TSharedRef<SWidget> UCanvasPanel::RebuildWidget()
 {
 	TSharedRef<SFixedSizeCanvas> NewCanvas = SNew(SFixedSizeCanvas, DesiredCanvasSize);
@@ -94,11 +99,37 @@ UCanvasPanelSlot* UCanvasPanel::AddSlot(UWidget* Content)
 bool UCanvasPanel::AddChild(UWidget* Child, FVector2D Position)
 {
 	UCanvasPanelSlot* Slot = AddSlot(Child);
-	Slot->Offset = FMargin(Position.X, Position.Y, 100, 25);
+	Slot->LayoutData.Offsets = FMargin(Position.X, Position.Y, 100, 25);
 
 	return true;
 }
 
+bool UCanvasPanel::GetGeometryForSlot(UCanvasPanelSlot* Slot, FGeometry& ArrangedGeometry) const
+{
+	if ( Slot->Content == NULL )
+	{
+		return false;
+	}
+
+	TSharedPtr<SFixedSizeCanvas> Canvas = GetCanvasWidget();
+	if ( Canvas.IsValid() )
+	{
+		FArrangedChildren ArrangedChildren(EVisibility::All);
+		Canvas->ArrangeChildren(Canvas->GetCachedGeometry(), ArrangedChildren);
+
+		for ( int32 ChildIndex = 0; ChildIndex < ArrangedChildren.Num(); ChildIndex++ )
+		{
+			if ( ArrangedChildren(ChildIndex).Widget == Slot->Content->GetWidget() )
+			{
+				ArrangedGeometry = ArrangedChildren(ChildIndex).Geometry;
+				return true;
+			}
+		}
+	}
+
+	return false;
+}
+
 #if WITH_EDITOR
 
 void UCanvasPanel::ConnectEditorData()
diff --git a/Engine/Source/Runtime/UMG/Private/Components/CanvasPanelSlot.cpp b/Engine/Source/Runtime/UMG/Private/Components/CanvasPanelSlot.cpp
--- a/Engine/Source/Runtime/UMG/Private/Components/CanvasPanelSlot.cpp
+++ b/Engine/Source/Runtime/UMG/Private/Components/CanvasPanelSlot.cpp
@@ -9,17 +9,17 @@ UCanvasPanelSlot::UCanvasPanelSlot(const FPostConstructInitializeProperties& PCI
 	: Super(PCIP)
 	, Slot(NULL)
 {
-	Offset = FMargin(0, 0, 1, 1);
-	Anchors = FAnchors(0.5f, 0.5f);
-	Alignment = FVector2D(0.5f, 0.5f);
+	LayoutData.Offsets = FMargin(0, 0, 1, 1);
+	LayoutData.Anchors = FAnchors(0.0f, 0.0f);
+	LayoutData.Alignment = FVector2D(0.5f, 0.5f);
 }
 
 void UCanvasPanelSlot::BuildSlot(TSharedRef<SConstraintCanvas> Canvas)
 {
 	Slot = &Canvas->AddSlot()
-		.Offset(Offset)
-		.Anchors(Anchors)
-		.Alignment(Alignment)
+		.Offset(LayoutData.Offsets)
+		.Anchors(LayoutData.Anchors)
+		.Alignment(LayoutData.Alignment)
 		[
 			Content == NULL ? SNullWidget::NullWidget : Content->GetWidget()
 		];
@@ -29,26 +29,26 @@ void UCanvasPanelSlot::Resize(const FVector2D& Direction, const FVector2D& Amoun
 {
 	if ( Direction.X < 0 )
 	{
-		Offset.Left -= Amount.X * Direction.X;
-		Offset.Right += Amount.X * Direction.X;
+		LayoutData.Offsets.Left -= Amount.X * Direction.X;
+		LayoutData.Offsets.Right += Amount.X * Direction.X;
 	}
 	if ( Direction.Y < 0 )
 	{
-		Offset.Top -= Amount.Y * Direction.Y;
-		Offset.Bottom += Amount.Y * Direction.Y;
+		LayoutData.Offsets.Top -= Amount.Y * Direction.Y;
+		LayoutData.Offsets.Bottom += Amount.Y * Direction.Y;
 	}
 	if ( Direction.X > 0 )
 	{
-		Offset.Right += Amount.X * Direction.X;
+		LayoutData.Offsets.Right += Amount.X * Direction.X;
 	}
 	if ( Direction.Y > 0 )
 	{
-		Offset.Bottom += Amount.Y * Direction.Y;
+		LayoutData.Offsets.Bottom += Amount.Y * Direction.Y;
 	}
 
 	if ( Slot )
 	{
-		Slot->Offset(Offset);
+		Slot->Offset(LayoutData.Offsets);
 	}
 }
 
@@ -59,7 +59,7 @@ bool UCanvasPanelSlot::CanResize(const FVector2D& Direction) const
 
 void UCanvasPanelSlot::SetOffset(FMargin InOffset)
 {
-	Offset = InOffset;
+	LayoutData.Offsets = InOffset;
 	if ( Slot )
 	{
 		Slot->Offset(InOffset);
@@ -68,7 +68,7 @@ void UCanvasPanelSlot::SetOffset(FMargin InOffset)
 
 void UCanvasPanelSlot::SetAnchors(FAnchors InAnchors)
 {
-	Anchors = InAnchors;
+	LayoutData.Anchors = InAnchors;
 	if ( Slot )
 	{
 		Slot->Anchors(InAnchors);
@@ -77,25 +77,111 @@ void UCanvasPanelSlot::SetAnchors(FAnchors InAnchors)
 
 void UCanvasPanelSlot::SetAlignment(FVector2D InAlignment)
 {
-	Alignment = InAlignment;
+	LayoutData.Alignment = InAlignment;
 	if ( Slot )
 	{
 		Slot->Alignment(InAlignment);
 	}
 }
 
+void UCanvasPanelSlot::Refresh()
+{
+	SetOffset(LayoutData.Offsets);
+	SetAnchors(LayoutData.Anchors);
+	SetAlignment(LayoutData.Alignment);
+}
+
 #if WITH_EDITOR
+
 void UCanvasPanelSlot::PreEditChange(UProperty* PropertyAboutToChange)
 {
-	//static FName AnchorsProperty(TEXT("Anchors"));
-
-	//if ( PropertyChangedEvent.MemberProperty && PropertyChangedEvent.MemberProperty->GetFName() == AnchorsProperty )
-	//{
+	Super::PreEditChange(PropertyAboutToChange);
 
-	//}
+	// Get the current location
+	if ( UCanvasPanel* Canvas = Cast<UCanvasPanel>(Parent) )
+	{
+		FGeometry Geometry;
+		if ( Canvas->GetGeometryForSlot(this, Geometry) )
+		{
+			PreEditGeometry = Geometry;
+			PreEditLayoutData = LayoutData;
+		}
+	}
 }
 
-void UCanvasPanelSlot::PostEditChangeProperty(struct FPropertyChangedEvent& PropertyChangedEvent)
+void UCanvasPanelSlot::PostEditChangeChainProperty(struct FPropertyChangedChainEvent& PropertyChangedEvent)
 {
+	Refresh();
+
+	static FName AnchorsProperty(TEXT("Anchors"));
+
+	if ( PropertyChangedEvent.MemberProperty && PropertyChangedEvent.MemberProperty->GetFName() == AnchorsProperty )
+	{
+		// Ensure we have a parent canvas
+		if ( UCanvasPanel* Canvas = Cast<UCanvasPanel>(Parent) )
+		{
+			FGeometry Geometry;
+			if ( Canvas->GetGeometryForSlot(this, Geometry) )
+			{
+				// Calculate the default anchor offset, ie where would this control be laid out if no offset were provided.
+				FVector2D CanvasSize = Canvas->GetCanvasWidget()->GetCachedGeometry().Size;
+				FMargin AnchorPositions = FMargin(
+					LayoutData.Anchors.Minimum.X * CanvasSize.X,
+					LayoutData.Anchors.Minimum.Y * CanvasSize.Y,
+					LayoutData.Anchors.Maximum.X * CanvasSize.X,
+					LayoutData.Anchors.Maximum.Y * CanvasSize.Y);
+				FVector2D DefaultAnchorPosition = FVector2D(AnchorPositions.Left, AnchorPositions.Top);
+
+				// Determine the amount that would be offset from the anchor position if alignment was applied.
+				FVector2D AlignmentOffset = LayoutData.Alignment * PreEditGeometry.Size;
+
+				// Determine where the widget's new position needs to be to maintain a stable location when the anchors change.
+				FVector2D PositionDelta = PreEditGeometry.Position - DefaultAnchorPosition;// ( Geometry.Position - DefaultAnchorOffset );
+
+				// Adjust the size to remain constant
+				if ( LayoutData.Anchors.IsStretchedHorizontal() && !PreEditLayoutData.Anchors.IsStretchedHorizontal() )
+				{
+					// Adjust the position to remain constant
+					LayoutData.Offsets.Left = PositionDelta.X;
+					LayoutData.Offsets.Right = AnchorPositions.Right - ( PositionDelta.X + PreEditGeometry.Size.X );
+				}
+				else if ( !LayoutData.Anchors.IsStretchedHorizontal() && PreEditLayoutData.Anchors.IsStretchedHorizontal() )
+				{
+					// Adjust the position to remain constant
+					LayoutData.Offsets.Left = PositionDelta.X + AlignmentOffset.X;
+					LayoutData.Offsets.Right = PreEditGeometry.Size.X;
+				}
+				else
+				{
+					// Adjust the position to remain constant
+					LayoutData.Offsets.Left = PositionDelta.X + AlignmentOffset.X;
+				}
+
+				if ( LayoutData.Anchors.IsStretchedVertical() && !PreEditLayoutData.Anchors.IsStretchedVertical() )
+				{
+					// Adjust the position to remain constant
+					LayoutData.Offsets.Top = PositionDelta.Y;
+					LayoutData.Offsets.Bottom = AnchorPositions.Bottom - ( PositionDelta.Y + PreEditGeometry.Size.Y );
+				}
+				else if ( !LayoutData.Anchors.IsStretchedVertical() && PreEditLayoutData.Anchors.IsStretchedVertical() )
+				{
+					// Adjust the position to remain constant
+					LayoutData.Offsets.Top = PositionDelta.Y + AlignmentOffset.Y;
+					LayoutData.Offsets.Bottom = PreEditGeometry.Size.Y;
+				}
+				else
+				{
+					// Adjust the position to remain constant
+					LayoutData.Offsets.Top = PositionDelta.Y + AlignmentOffset.Y;
+				}
+			}
+
+			// Apply the changes to the properties.
+			Refresh();
+		}
+	}
+
+	Super::PostEditChangeProperty(PropertyChangedEvent);
 }
+
 #endif
\ No newline at end of file
diff --git a/Engine/Source/Runtime/UMG/Private/Components/HorizontalBox.cpp b/Engine/Source/Runtime/UMG/Private/Components/HorizontalBox.cpp
--- a/Engine/Source/Runtime/UMG/Private/Components/HorizontalBox.cpp
+++ b/Engine/Source/Runtime/UMG/Private/Components/HorizontalBox.cpp
@@ -79,35 +79,22 @@ void UHorizontalBox::InsertChildAt(int32 Index, UWidget* Content)
 
 TSharedRef<SWidget> UHorizontalBox::RebuildWidget()
 {
-	TSharedRef<SHorizontalBox> NewCanvas = SNew(SHorizontalBox);
-	MyHorizontalBox = NewCanvas;
+	TSharedRef<SHorizontalBox> NewWidget = SNew(SHorizontalBox);
+	MyHorizontalBox = NewWidget;
 
-	TSharedPtr<SHorizontalBox> Canvas = MyHorizontalBox.Pin();
-	if ( Canvas.IsValid() )
+	for ( int32 SlotIndex = 0; SlotIndex < Slots.Num(); ++SlotIndex )
 	{
-		Canvas->ClearChildren();
-
-		for ( int32 SlotIndex = 0; SlotIndex < Slots.Num(); ++SlotIndex )
+		UHorizontalBoxSlot* Slot = Slots[SlotIndex];
+		if ( Slot == NULL )
 		{
-			UHorizontalBoxSlot* Slot = Slots[SlotIndex];
-			if ( Slot == NULL )
-			{
-				Slots[SlotIndex] = Slot = ConstructObject<UHorizontalBoxSlot>(UHorizontalBoxSlot::StaticClass(), this);
-			}
-
-			auto& NewSlot = Canvas->AddSlot()
-				.Padding(Slot->Padding)
-				.HAlign(Slot->HorizontalAlignment)
-				.VAlign(Slot->VerticalAlignment)
-				[
-					Slot->Content == NULL ? SNullWidget::NullWidget : Slot->Content->GetWidget()
-				];
-
-			NewSlot.SizeParam = UWidget::ConvertSerializedSizeParamToRuntime(Slot->Size);
+			Slots[SlotIndex] = Slot = ConstructObject<UHorizontalBoxSlot>(UHorizontalBoxSlot::StaticClass(), this);
 		}
+
+		Slot->Parent = this;
+		Slot->BuildSlot(NewWidget);
 	}
 
-	return NewCanvas;
+	return NewWidget;
 }
 
 UHorizontalBoxSlot* UHorizontalBox::AddSlot(UWidget* Content)
diff --git a/Engine/Source/Runtime/UMG/Private/Components/HorizontalBoxSlot.cpp b/Engine/Source/Runtime/UMG/Private/Components/HorizontalBoxSlot.cpp
--- a/Engine/Source/Runtime/UMG/Private/Components/HorizontalBoxSlot.cpp
+++ b/Engine/Source/Runtime/UMG/Private/Components/HorizontalBoxSlot.cpp
@@ -11,3 +11,60 @@ UHorizontalBoxSlot::UHorizontalBoxSlot(const FPostConstructInitializeProperties&
 	HorizontalAlignment = HAlign_Left;
 	VerticalAlignment = VAlign_Top;
 }
+
+void UHorizontalBoxSlot::BuildSlot(TSharedRef<SHorizontalBox> HorizontalBox)
+{
+	Slot = &HorizontalBox->AddSlot()
+		.HAlign(HorizontalAlignment)
+		.VAlign(VerticalAlignment)
+		.Padding(Padding)
+		[
+			Content == NULL ? SNullWidget::NullWidget : Content->GetWidget()
+		];
+
+	Slot->SizeParam = UWidget::ConvertSerializedSizeParamToRuntime(Size);
+}
+
+void UHorizontalBoxSlot::SetPadding(FMargin InPadding)
+{
+	Padding = InPadding;
+	if ( Slot )
+	{
+		Slot->Padding(InPadding);
+	}
+}
+
+void UHorizontalBoxSlot::SetSize(FSlateChildSize InSize)
+{
+	Size = InSize;
+	if ( Slot )
+	{
+		Slot->SizeParam = UWidget::ConvertSerializedSizeParamToRuntime(InSize);
+	}
+}
+
+void UHorizontalBoxSlot::SetHorizontalAlignment(EHorizontalAlignment InHorizontalAlignment)
+{
+	HorizontalAlignment = InHorizontalAlignment;
+	if ( Slot )
+	{
+		Slot->HAlign(InHorizontalAlignment);
+	}
+}
+
+void UHorizontalBoxSlot::SetVerticalAlignment(EVerticalAlignment InVerticalAlignment)
+{
+	VerticalAlignment = InVerticalAlignment;
+	if ( Slot )
+	{
+		Slot->VAlign(InVerticalAlignment);
+	}
+}
+
+void UHorizontalBoxSlot::Refresh()
+{
+	SetPadding(Padding);
+	SetSize(Size);
+	SetHorizontalAlignment(HorizontalAlignment);
+	SetVerticalAlignment(VerticalAlignment);
+}
diff --git a/Engine/Source/Runtime/UMG/Private/Components/Overlay.cpp b/Engine/Source/Runtime/UMG/Private/Components/Overlay.cpp
--- a/Engine/Source/Runtime/UMG/Private/Components/Overlay.cpp
+++ b/Engine/Source/Runtime/UMG/Private/Components/Overlay.cpp
@@ -55,33 +55,21 @@ void UOverlay::ReplaceChildAt(int32 Index, UWidget* Content)
 
 TSharedRef<SWidget> UOverlay::RebuildWidget()
 {
-	TSharedRef<SOverlay> NewCanvas = SNew(SOverlay);
-	MyVerticalBox = NewCanvas;
+	TSharedRef<SOverlay> NewWidget = SNew(SOverlay);
+	MyVerticalBox = NewWidget;
 
-	TSharedPtr<SOverlay> Canvas = MyVerticalBox.Pin();
-	if ( Canvas.IsValid() )
+	for ( int32 SlotIndex = 0; SlotIndex < Slots.Num(); ++SlotIndex )
 	{
-		Canvas->ClearChildren();
-
-		for ( int32 SlotIndex = 0; SlotIndex < Slots.Num(); ++SlotIndex )
+		UOverlaySlot* Slot = Slots[SlotIndex];
+		if ( Slot == NULL )
 		{
-			UOverlaySlot* Slot = Slots[SlotIndex];
-			if ( Slot == NULL )
-			{
-				Slots[SlotIndex] = Slot = ConstructObject<UOverlaySlot>(UOverlaySlot::StaticClass(), this);
-			}
-
-			auto& NewSlot = Canvas->AddSlot()
-				.Padding(Slot->Padding)
-				.HAlign(Slot->HorizontalAlignment)
-				.VAlign(Slot->VerticalAlignment)
-				[
-					Slot->Content == NULL ? SNullWidget::NullWidget : Slot->Content->GetWidget()
-				];
+			Slots[SlotIndex] = Slot = ConstructObject<UOverlaySlot>(UOverlaySlot::StaticClass(), this);
 		}
+
+		Slot->BuildSlot(NewWidget);
 	}
 
-	return NewCanvas;
+	return NewWidget;
 }
 
 UOverlaySlot* UOverlay::AddSlot(UWidget* Content)
diff --git a/Engine/Source/Runtime/UMG/Private/Components/OverlaySlot.cpp b/Engine/Source/Runtime/UMG/Private/Components/OverlaySlot.cpp
--- a/Engine/Source/Runtime/UMG/Private/Components/OverlaySlot.cpp
+++ b/Engine/Source/Runtime/UMG/Private/Components/OverlaySlot.cpp
@@ -11,3 +11,48 @@ UOverlaySlot::UOverlaySlot(const FPostConstructInitializeProperties& PCIP)
 	HorizontalAlignment = HAlign_Left;
 	VerticalAlignment = VAlign_Top;
 }
+
+void UOverlaySlot::BuildSlot(TSharedRef<SOverlay> Overlay)
+{
+	Slot = &Overlay->AddSlot()
+		.Padding(Padding)
+		.HAlign(HorizontalAlignment)
+		.VAlign(VerticalAlignment)
+		[
+			Content == NULL ? SNullWidget::NullWidget : Content->GetWidget()
+		];
+}
+
+void UOverlaySlot::SetPadding(FMargin InPadding)
+{
+	Padding = InPadding;
+	if ( Slot )
+	{
+		Slot->Padding(InPadding);
+	}
+}
+
+void UOverlaySlot::SetHorizontalAlignment(EHorizontalAlignment InHorizontalAlignment)
+{
+	HorizontalAlignment = InHorizontalAlignment;
+	if ( Slot )
+	{
+		Slot->HAlign(InHorizontalAlignment);
+	}
+}
+
+void UOverlaySlot::SetVerticalAlignment(EVerticalAlignment InVerticalAlignment)
+{
+	VerticalAlignment = InVerticalAlignment;
+	if ( Slot )
+	{
+		Slot->VAlign(InVerticalAlignment);
+	}
+}
+
+void UOverlaySlot::Refresh()
+{
+	SetPadding(Padding);
+	SetHorizontalAlignment(HorizontalAlignment);
+	SetVerticalAlignment(VerticalAlignment);
+}
diff --git a/Engine/Source/Runtime/UMG/Private/Components/ScrollBoxSlot.cpp b/Engine/Source/Runtime/UMG/Private/Components/ScrollBoxSlot.cpp
--- a/Engine/Source/Runtime/UMG/Private/Components/ScrollBoxSlot.cpp
+++ b/Engine/Source/Runtime/UMG/Private/Components/ScrollBoxSlot.cpp
@@ -21,3 +21,27 @@ void UScrollBoxSlot::BuildSlot(TSharedRef<SScrollBox> ScrollBox)
 			Content == NULL ? SNullWidget::NullWidget : Content->GetWidget()
 		];
 }
+
+void UScrollBoxSlot::SetPadding(FMargin InPadding)
+{
+	Padding = InPadding;
+	if ( Slot )
+	{
+		Slot->Padding(InPadding);
+	}
+}
+
+void UScrollBoxSlot::SetHorizontalAlignment(EHorizontalAlignment InHorizontalAlignment)
+{
+	HorizontalAlignment = InHorizontalAlignment;
+	if ( Slot )
+	{
+		Slot->HAlign(InHorizontalAlignment);
+	}
+}
+
+void UScrollBoxSlot::Refresh()
+{
+	SetPadding(Padding);
+	SetHorizontalAlignment(HorizontalAlignment);
+}
diff --git a/Engine/Source/Runtime/UMG/Private/Components/Spacer.cpp b/Engine/Source/Runtime/UMG/Private/Components/Spacer.cpp
--- a/Engine/Source/Runtime/UMG/Private/Components/Spacer.cpp
+++ b/Engine/Source/Runtime/UMG/Private/Components/Spacer.cpp
@@ -14,10 +14,22 @@ USpacer::USpacer(const FPostConstructInitializeProperties& PCIP)
 	bIsVariable = false;
 }
 
+void USpacer::SetSize(FVector2D InSize)
+{
+	Size = InSize;
+
+	if ( MySpacer.IsValid() )
+	{
+		MySpacer->SetSize(InSize);
+	}
+}
+
 TSharedRef<SWidget> USpacer::RebuildWidget()
 {
-	return SNew(SSpacer)
+	MySpacer = SNew(SSpacer)
 		.Size(Size);
+
+	return MySpacer.ToSharedRef();
 }
 
 /////////////////////////////////////////////////////
diff --git a/Engine/Source/Runtime/UMG/Private/Components/UniformGridSlot.cpp b/Engine/Source/Runtime/UMG/Private/Components/UniformGridSlot.cpp
--- a/Engine/Source/Runtime/UMG/Private/Components/UniformGridSlot.cpp
+++ b/Engine/Source/Runtime/UMG/Private/Components/UniformGridSlot.cpp
@@ -40,3 +40,9 @@ void UUniformGridSlot::SetColumn(int32 InColumn)
 		Slot->Column = InColumn;
 	}
 }
+
+void UUniformGridSlot::Refresh()
+{
+	SetRow(Row);
+	SetColumn(Column);
+}
diff --git a/Engine/Source/Runtime/UMG/Private/Components/VerticalBox.cpp b/Engine/Source/Runtime/UMG/Private/Components/VerticalBox.cpp
--- a/Engine/Source/Runtime/UMG/Private/Components/VerticalBox.cpp
+++ b/Engine/Source/Runtime/UMG/Private/Components/VerticalBox.cpp
@@ -79,28 +79,22 @@ void UVerticalBox::InsertChildAt(int32 Index, UWidget* Content)
 
 TSharedRef<SWidget> UVerticalBox::RebuildWidget()
 {
-	TSharedRef<SVerticalBox> NewCanvas = SNew(SVerticalBox);
-	MyVerticalBox = NewCanvas;
+	TSharedRef<SVerticalBox> NewWidget = SNew(SVerticalBox);
+	MyVerticalBox = NewWidget;
 
-	TSharedPtr<SVerticalBox> Canvas = MyVerticalBox.Pin();
-	if ( Canvas.IsValid() )
+	for ( int32 SlotIndex = 0; SlotIndex < Slots.Num(); ++SlotIndex )
 	{
-		Canvas->ClearChildren();
-
-		for ( int32 SlotIndex = 0; SlotIndex < Slots.Num(); ++SlotIndex )
+		UVerticalBoxSlot* Slot = Slots[SlotIndex];
+		if ( Slot == NULL )
 		{
-			UVerticalBoxSlot* Slot = Slots[SlotIndex];
-			if ( Slot == NULL )
-			{
-				Slots[SlotIndex] = Slot = ConstructObject<UVerticalBoxSlot>(UVerticalBoxSlot::StaticClass(), this);
-			}
-
-			Slot->Parent = this;
-			Slot->BuildSlot(NewCanvas);
+			Slots[SlotIndex] = Slot = ConstructObject<UVerticalBoxSlot>(UVerticalBoxSlot::StaticClass(), this);
 		}
+
+		Slot->Parent = this;
+		Slot->BuildSlot(NewWidget);
 	}
 
-	return NewCanvas;
+	return NewWidget;
 }
 
 UVerticalBoxSlot* UVerticalBox::AddSlot(UWidget* Content)
diff --git a/Engine/Source/Runtime/UMG/Private/Components/VerticalBoxSlot.cpp b/Engine/Source/Runtime/UMG/Private/Components/VerticalBoxSlot.cpp
--- a/Engine/Source/Runtime/UMG/Private/Components/VerticalBoxSlot.cpp
+++ b/Engine/Source/Runtime/UMG/Private/Components/VerticalBoxSlot.cpp
@@ -25,3 +25,47 @@ void UVerticalBoxSlot::BuildSlot(TSharedRef<SVerticalBox> VerticalBox)
 
 	Slot->SizeParam = UWidget::ConvertSerializedSizeParamToRuntime(Size);
 }
+
+void UVerticalBoxSlot::SetPadding(FMargin InPadding)
+{
+	Padding = InPadding;
+	if ( Slot )
+	{
+		Slot->Padding(InPadding);
+	}
+}
+
+void UVerticalBoxSlot::SetSize(FSlateChildSize InSize)
+{
+	Size = InSize;
+	if ( Slot )
+	{
+		Slot->SizeParam = UWidget::ConvertSerializedSizeParamToRuntime(InSize);
+	}
+}
+
+void UVerticalBoxSlot::SetHorizontalAlignment(EHorizontalAlignment InHorizontalAlignment)
+{
+	HorizontalAlignment = InHorizontalAlignment;
+	if ( Slot )
+	{
+		Slot->HAlign(InHorizontalAlignment);
+	}
+}
+
+void UVerticalBoxSlot::SetVerticalAlignment(EVerticalAlignment InVerticalAlignment)
+{
+	VerticalAlignment = InVerticalAlignment;
+	if ( Slot )
+	{
+		Slot->VAlign(InVerticalAlignment);
+	}
+}
+
+void UVerticalBoxSlot::Refresh()
+{
+	SetPadding(Padding);
+	SetSize(Size);
+	SetHorizontalAlignment(HorizontalAlignment);
+	SetVerticalAlignment(VerticalAlignment);
+}
diff --git a/Engine/Source/Runtime/UMG/Private/Components/Widget.cpp b/Engine/Source/Runtime/UMG/Private/Components/Widget.cpp
--- a/Engine/Source/Runtime/UMG/Private/Components/Widget.cpp
+++ b/Engine/Source/Runtime/UMG/Private/Components/Widget.cpp
@@ -29,22 +29,37 @@ void UWidget::SetIsEnabled(bool bInIsEnabled)
 
 TEnumAsByte<ESlateVisibility::Type> UWidget::GetVisibility()
 {
-	return UWidget::ConvertRuntimeToSerializedVisiblity(MyWidget->GetVisibility());
+	if ( MyWidget.IsValid() )
+	{
+		return UWidget::ConvertRuntimeToSerializedVisiblity(MyWidget->GetVisibility());
+	}
+
+	return Visiblity;
 }
 
 void UWidget::SetVisibility(TEnumAsByte<ESlateVisibility::Type> InVisibility)
 {
-	MyWidget->SetVisibility(UWidget::ConvertSerializedVisibilityToRuntime(InVisibility));
+	Visiblity = InVisibility;
+
+	if ( MyWidget.IsValid() )
+	{
+		MyWidget->SetVisibility(UWidget::ConvertSerializedVisibilityToRuntime(InVisibility));
+	}
 }
 
 void UWidget::SetToolTipText(const FText& InToolTipText)
 {
-	MyWidget->SetToolTipText(ToolTipText);
+	ToolTipText = InToolTipText;
+
+	if ( MyWidget.IsValid() )
+	{
+		MyWidget->SetToolTipText(InToolTipText);
+	}
 }
 
-void UWidget::PostEditChangeProperty(struct FPropertyChangedEvent& PropertyChangedEvent)
+bool UWidget::IsHovered() const
 {
-	// TODO Rename?
+	return MyWidget->IsHovered();
 }
 
 TSharedRef<SWidget> UWidget::GetWidget()
@@ -52,14 +67,7 @@ TSharedRef<SWidget> UWidget::GetWidget()
 	if ( !MyWidget.IsValid() )
 	{
 		MyWidget = RebuildWidget();
-		MyWidget->SetEnabled(OPTIONAL_BINDING(bool, bIsEnabled));
-		MyWidget->SetVisibility(OPTIONAL_BINDING_CONVERT(ESlateVisibility::Type, Visiblity, EVisibility, ConvertVisibility));
-		//MyWidget->SetCursor(OPTIONAL_BINDING(EMouseCursor)
-
-		if ( !ToolTipText.IsEmpty() )
-		{
-			MyWidget->SetToolTipText(OPTIONAL_BINDING(FText, ToolTipText));
-		}
+		SyncronizeProperties();
 	}
 
 	return MyWidget.ToSharedRef();
@@ -71,6 +79,21 @@ TSharedRef<SWidget> UWidget::RebuildWidget()
 	return SNew(SSpacer);
 }
 
+void UWidget::SyncronizeProperties()
+{
+	if ( MyWidget.IsValid() )
+	{
+		MyWidget->SetEnabled(OPTIONAL_BINDING(bool, bIsEnabled));
+		MyWidget->SetVisibility(OPTIONAL_BINDING_CONVERT(ESlateVisibility::Type, Visiblity, EVisibility, ConvertVisibility));
+		//MyWidget->SetCursor(OPTIONAL_BINDING(EMouseCursor)
+
+		if ( !ToolTipText.IsEmpty() )
+		{
+			MyWidget->SetToolTipText(OPTIONAL_BINDING(FText, ToolTipText));
+		}
+	}
+}
+
 EVisibility UWidget::ConvertSerializedVisibilityToRuntime(TEnumAsByte<ESlateVisibility::Type> Input)
 {
 	switch ( Input )
@@ -130,4 +153,4 @@ FSizeParam UWidget::ConvertSerializedSizeParamToRuntime(const FSlateChildSize& I
 	case ESlateSizeRule::Fill:
 		return FStretch(Input.Value);
 	}
-}
\ No newline at end of file
+}
diff --git a/Engine/Source/Runtime/UMG/Private/Slate/SConstraintCanvas.cpp b/Engine/Source/Runtime/UMG/Private/Slate/SConstraintCanvas.cpp
--- a/Engine/Source/Runtime/UMG/Private/Slate/SConstraintCanvas.cpp
+++ b/Engine/Source/Runtime/UMG/Private/Slate/SConstraintCanvas.cpp
@@ -79,7 +79,7 @@ void SConstraintCanvas::OnArrangeChildren( const FGeometry& AllottedGeometry, FA
 			if ( bIsHorizontalStretch )
 			{
 				LocalPosition.X = AnchorPixels.Left + Offset.Left;
-				LocalSize.X = AnchorPixels.Right - AnchorPixels.Left - Offset.Right;
+				LocalSize.X = AnchorPixels.Right - LocalPosition.X - Offset.Right;
 			}
 			else
 			{
@@ -91,7 +91,7 @@ void SConstraintCanvas::OnArrangeChildren( const FGeometry& AllottedGeometry, FA
 			if ( bIsVerticalStretch )
 			{
 				LocalPosition.Y = AnchorPixels.Top + Offset.Top;
-				LocalSize.Y = AnchorPixels.Bottom - AnchorPixels.Top - Offset.Bottom;
+				LocalSize.Y = AnchorPixels.Bottom - LocalPosition.Y - Offset.Bottom;
 			}
 			else
 			{
@@ -115,6 +115,8 @@ void SConstraintCanvas::OnArrangeChildren( const FGeometry& AllottedGeometry, FA
 
 int32 SConstraintCanvas::OnPaint( const FGeometry& AllottedGeometry, const FSlateRect& MyClippingRect, FSlateWindowElementList& OutDrawElements, int32 LayerId, const FWidgetStyle& InWidgetStyle, bool bParentEnabled ) const
 {
+	CachedGeometry = AllottedGeometry;
+
 	FArrangedChildren ArrangedChildren(EVisibility::Visible);
 	this->ArrangeChildren(AllottedGeometry, ArrangedChildren);
 
diff --git a/Engine/Source/Runtime/UMG/Private/WidgetBlueprintGeneratedClass.cpp b/Engine/Source/Runtime/UMG/Private/WidgetBlueprintGeneratedClass.cpp
--- a/Engine/Source/Runtime/UMG/Private/WidgetBlueprintGeneratedClass.cpp
+++ b/Engine/Source/Runtime/UMG/Private/WidgetBlueprintGeneratedClass.cpp
@@ -33,18 +33,18 @@ void UWidgetBlueprintGeneratedClass::Link(FArchive& Ar, bool bRelinkExistingProp
 	}
 }
 
-void UWidgetBlueprintGeneratedClass::InitializeWidget(UUserWidget* Actor) const
+void UWidgetBlueprintGeneratedClass::InitializeWidget(UUserWidget* UserWidget) const
 {
 	// Duplicate the graph, keeping track of what was duplicated
 	TMap<UObject*, UObject*> DuplicatedObjectList;
 
-	FObjectDuplicationParameters Parameters(const_cast<UWidgetTree*>( WidgetTree ), Actor);
+	FObjectDuplicationParameters Parameters(const_cast<UWidgetTree*>( WidgetTree ), UserWidget);
 	Parameters.CreatedObjects = &DuplicatedObjectList;
 
 	UWidgetTree* ClonedTree = CastChecked<UWidgetTree>(StaticDuplicateObjectEx(Parameters));
 
-	UUserWidget* WidgetActor = CastChecked<UUserWidget>(Actor);
-	UClass* ActorClass = Actor->GetClass();
+	UUserWidget* WidgetActor = CastChecked<UUserWidget>(UserWidget);
+	UClass* ActorClass = UserWidget->GetClass();
 
 	for ( UWidget* Widget : ClonedTree->WidgetTemplates )
 	{
@@ -54,19 +54,19 @@ void UWidgetBlueprintGeneratedClass::InitializeWidget(UUserWidget* Actor) const
 			continue;
 		}
 
-		FName NewName = *( FString::Printf(TEXT("WidgetComponent__%d"), Actor->Components.Num()) );
+		FName NewName = *( FString::Printf(TEXT("WidgetComponent__%d"), UserWidget->Components.Num()) );
 
 		FString VariableName = Widget->GetName();
 
 		Widget->bCreatedByConstructionScript = true; // Indicate it comes from a blueprint so it gets cleared when we rerun construction scripts
-		Actor->Components.Add(Widget); // Add to array so it gets saved
+		UserWidget->Components.Add(Widget); // Add to array so it gets saved
 //		Widget->SetNetAddressable();	// This component has a stable name that can be referenced for replication
 
 		// Find property with the same name as the template and assign the new Timeline to it
 		UObjectPropertyBase* Prop = FindField<UObjectPropertyBase>(ActorClass, *VariableName);
 		if ( Prop )
 		{
-			Prop->SetObjectPropertyValue_InContainer(Actor, Widget);
+			Prop->SetObjectPropertyValue_InContainer(UserWidget, Widget);
 		}
 
 		// Perform binding
diff --git a/Engine/Source/Runtime/UMG/Public/Components/Border.h b/Engine/Source/Runtime/UMG/Public/Components/Border.h
--- a/Engine/Source/Runtime/UMG/Public/Components/Border.h
+++ b/Engine/Source/Runtime/UMG/Public/Components/Border.h
@@ -10,23 +10,31 @@ class USlateBrushAsset;
  * A border is a container widget that can contain one child widget, providing an opportunity 
  * to surround it with a border image and adjustable padding.
  */
-UCLASS(meta=( BlueprintSpawnableComponent, Category="Common" ), ClassGroup=UserInterface)
+UCLASS(meta=( Category="Common" ), ClassGroup=UserInterface)
 class UMG_API UBorder : public UContentWidget
 {
 	GENERATED_UCLASS_BODY()
 
 	/** Horizontal positioning of the content within the border */
-	UPROPERTY(EditAnywhere, BlueprintReadOnly, Category=Appearance)
+	UPROPERTY(EditAnywhere, BlueprintReadOnly, Category="Content Layout")
 	TEnumAsByte<EHorizontalAlignment> HorizontalAlignment;
 
 	/** Vertical positioning of the content within the border */
-	UPROPERTY(EditAnywhere, BlueprintReadOnly, Category=Appearance)
+	UPROPERTY(EditAnywhere, BlueprintReadOnly, Category="Content Layout")
 	TEnumAsByte<EVerticalAlignment> VerticalAlignment;
 
 	/** The padding to add around the inner content. */
-	UPROPERTY(EditAnywhere, BlueprintReadWrite, Category=Appearance)
+	UPROPERTY(EditAnywhere, BlueprintReadWrite, Category="Content Layout")
 	FMargin ContentPadding;
 
+	/** The scaling factor for the border content */
+	UPROPERTY(EditAnywhere, BlueprintReadOnly, Category="Content Layout")
+	FVector2D ContentScale;
+
+	/** Color and opacity multiplier of content in the border */
+	UPROPERTY(EditAnywhere, BlueprintReadWrite, Category="Content Layout")
+	FLinearColor ContentColorAndOpacity;
+
 	//@TODO: Should these be exposed?
 	// SLATE_EVENT( FPointerEventHandler, OnMouseButtonDown )
 	// SLATE_EVENT( FPointerEventHandler, OnMouseButtonUp )
@@ -41,21 +49,13 @@ class UMG_API UBorder : public UContentWidget
 	UPROPERTY(EditAnywhere, BlueprintReadOnly, Category=Appearance)
 	FVector2D DesiredSizeScale;
 
-	/** The scaling factor for the border content */
-	UPROPERTY(EditAnywhere, BlueprintReadOnly, Category=Appearance)
-	FVector2D ContentScale;
-
-	/** Color and opacity multiplier of content in the border */
-	UPROPERTY(EditAnywhere, BlueprintReadWrite, Category=Appearance)
-	FLinearColor ContentColorAndOpacity;
-
 	/** Color and opacity of the actual border image */
 	UPROPERTY(EditAnywhere, BlueprintReadWrite, Category=Appearance)
-	FSlateColor BorderColorAndOpacity;
+	FLinearColor BorderColorAndOpacity;
 
 	/** The foreground color of text and some glyphs that appear as the border's content. */
 	UPROPERTY(EditAnywhere, BlueprintReadWrite, Category=Appearance)
-	FSlateColor ForegroundColor;
+	FLinearColor ForegroundColor;
 
 	/** Whether or not to show the disabled effect when this border is disabled */
 	UPROPERTY(EditAnywhere, BlueprintReadOnly, Category=Appearance)
@@ -65,17 +65,17 @@ class UMG_API UBorder : public UContentWidget
 	virtual void SetContent(UWidget* Content) OVERRIDE;
 	// End UContentWidget interface
 
+	// UWidget interface
+	void SyncronizeProperties() OVERRIDE;
+	// End of UWidget interface
+
 protected:
-	TWeakPtr<class SBorder> MyBorder;
+	TSharedPtr<SBorder> MyBorder;
 
 protected:
 	// UWidget interface
 	virtual TSharedRef<SWidget> RebuildWidget() OVERRIDE;
 	// End of UWidget interface
 
-	FMargin GetContentPadding() const;
-	FLinearColor GetContentColor() const;
 	const FSlateBrush* GetBorderBrush() const;
-	FSlateColor GetBorderColor() const;
-	FSlateColor GetForegroundColor() const;
 };
diff --git a/Engine/Source/Runtime/UMG/Public/Components/Button.h b/Engine/Source/Runtime/UMG/Public/Components/Button.h
--- a/Engine/Source/Runtime/UMG/Public/Components/Button.h
+++ b/Engine/Source/Runtime/UMG/Public/Components/Button.h
@@ -7,7 +7,7 @@
 DECLARE_DYNAMIC_MULTICAST_DELEGATE(FOnButtonClickedEvent);
 
 /** Buttons are clickable widgets */
-UCLASS(meta=( BlueprintSpawnableComponent, Category="Common" ), ClassGroup=UserInterface)
+UCLASS(meta=( Category="Common" ), ClassGroup=UserInterface)
 class UMG_API UButton : public UContentWidget
 {
 	GENERATED_UCLASS_BODY()
@@ -18,15 +18,15 @@ public:
 	USlateWidgetStyleAsset* Style;
 
 	/** Horizontal positioning of the content within the button */
-	UPROPERTY(EditDefaultsOnly, Category=Appearance)
+	UPROPERTY(EditDefaultsOnly, Category="Content Layout")
 	TEnumAsByte<EHorizontalAlignment> HorizontalAlignment;
 
 	/** Vertical positioning of the content within the button */
-	UPROPERTY(EditDefaultsOnly, Category=Appearance)
+	UPROPERTY(EditDefaultsOnly, Category="Content Layout")
 	TEnumAsByte<EVerticalAlignment> VerticalAlignment;
 
 	/** The padding to add around the button content. */
-	UPROPERTY(EditDefaultsOnly, Category=Appearance)
+	UPROPERTY(EditDefaultsOnly, Category="Content Layout")
 	FMargin ContentPadding;
 
 	/** The scaling factor for the button border */
@@ -55,52 +55,33 @@ public:
 	UPROPERTY(EditDefaultsOnly, Category=Events)
 	FOnReply OnClickedEvent;
 	
-	/*  */
-	UFUNCTION(BlueprintPure, Category="Appearance")
-	FLinearColor GetButtonColorAndOpacity();
-	
 	/*  */
 	UFUNCTION(BlueprintCallable, Category="Appearance")
 	void SetButtonColorAndOpacity(FLinearColor InButtonColorAndOpacity);
 
-	/*  */
-	UFUNCTION(BlueprintPure, Category="Appearance")
-	FLinearColor GetForegroundColor();
-
 	/*  */
 	UFUNCTION(BlueprintCallable, Category="Appearance")
 	void SetForegroundColor(FLinearColor InForegroundColor);
-	
-	/*  */
-	UFUNCTION(BlueprintPure, Category="Appearance")
-	FMargin GetContentPadding();
 
 	/*  */
 	UFUNCTION(BlueprintCallable, Category="Appearance")
 	void SetContentPadding(FMargin InContentPadding);
 
-
 	// UContentWidget interface
 	virtual void SetContent(UWidget* Content) OVERRIDE;
 	// End UContentWidget interface
 
-protected:
-	TSharedPtr<class SButton> ButtonWidget() const
-	{
-		if ( MyWidget.IsValid() )
-		{
-			return StaticCastSharedRef<SButton>(MyWidget.ToSharedRef());
-		}
-
-		return TSharedPtr<class SButton>();
-	}
+	// UWidget interface
+	void SyncronizeProperties() OVERRIDE;
+	// End of UWidget interface
 
 protected:
 	// UWidget interface
 	virtual TSharedRef<SWidget> RebuildWidget() OVERRIDE;
 	// End of UWidget interface
 
-	FMargin GetContentPadding() const;
-
 	FReply HandleOnClicked();
+
+protected:
+	TSharedPtr<SButton> MyButton;
 };
diff --git a/Engine/Source/Runtime/UMG/Public/Components/CanvasPanel.h b/Engine/Source/Runtime/UMG/Public/Components/CanvasPanel.h
--- a/Engine/Source/Runtime/UMG/Public/Components/CanvasPanel.h
+++ b/Engine/Source/Runtime/UMG/Public/Components/CanvasPanel.h
@@ -4,13 +4,13 @@
 
 #include "CanvasPanel.generated.h"
 
-UCLASS(meta=( BlueprintSpawnableComponent, Category="Panel" ), ClassGroup=UserInterface)
+UCLASS(meta=( Category="Panel" ), ClassGroup=UserInterface)
 class UMG_API UCanvasPanel : public UPanelWidget
 {
 	GENERATED_UCLASS_BODY()
 
 	/** The items placed on the canvas */
-	UPROPERTY(EditAnywhere, EditInline, Category=Slots)
+	UPROPERTY()
 	TArray<UCanvasPanelSlot*> Slots;
 
 	/** The desired size of the canvas */
@@ -25,6 +25,12 @@ class UMG_API UCanvasPanel : public UPanelWidget
 	virtual UWidget* GetChildAt(int32 Index) const OVERRIDE;
 	// End UPanelWidget
 
+	/** Gets the underlying native canvas widget if it has been constructed */
+	TSharedPtr<class SFixedSizeCanvas> GetCanvasWidget() const;
+
+	/** Computes the geometry for a particular slot based on the current geometry of the canvas. */
+	bool GetGeometryForSlot(UCanvasPanelSlot* Slot, FGeometry& ArrangedGeometry) const;
+
 #if WITH_EDITOR
 	// UObject interface
 	virtual void PostEditChangeProperty(struct FPropertyChangedEvent& PropertyChangedEvent) OVERRIDE;
diff --git a/Engine/Source/Runtime/UMG/Public/Components/CanvasPanelSlot.h b/Engine/Source/Runtime/UMG/Public/Components/CanvasPanelSlot.h
--- a/Engine/Source/Runtime/UMG/Public/Components/CanvasPanelSlot.h
+++ b/Engine/Source/Runtime/UMG/Public/Components/CanvasPanelSlot.h
@@ -6,22 +6,33 @@
 
 #include "CanvasPanelSlot.generated.h"
 
-UCLASS()
-class UMG_API UCanvasPanelSlot : public UPanelSlot
+USTRUCT()
+struct FAnchorData
 {
-	GENERATED_UCLASS_BODY()
+public:
+	GENERATED_USTRUCT_BODY()
 
 	/** Offset. */
-	UPROPERTY(EditDefaultsOnly, Category=Appearance)
-	FMargin Offset;
+	UPROPERTY(EditAnywhere, Category=AnchorData)
+	FMargin Offsets;
 	
 	/** Anchors. */
-	UPROPERTY(EditDefaultsOnly, Category=Appearance)
+	UPROPERTY(EditAnywhere, Category=AnchorData)
 	FAnchors Anchors;
 
 	/** Alignment. */
-	UPROPERTY(EditDefaultsOnly, Category=Appearance)
+	UPROPERTY(EditAnywhere, Category=AnchorData)
 	FVector2D Alignment;
+};
+
+UCLASS()
+class UMG_API UCanvasPanelSlot : public UPanelSlot
+{
+	GENERATED_UCLASS_BODY()
+
+	/** The anchoring information for the slot */
+	UPROPERTY(EditDefaultsOnly, Category=Appearance)
+	FAnchorData LayoutData;
 
 	void BuildSlot(TSharedRef<SConstraintCanvas> Canvas);
 
@@ -41,13 +52,22 @@ class UMG_API UCanvasPanelSlot : public UPanelSlot
 	UFUNCTION(BlueprintCallable, Category="Appearance")
 	void SetAlignment(FVector2D InAlignment);
 
+	// UPanelSlot interface
+	virtual void Refresh() OVERRIDE;
+	// End of UPanelSlot interface
+
 #if WITH_EDITOR
 	// UObject interface
 	virtual void PreEditChange(UProperty* PropertyAboutToChange) OVERRIDE;
-	virtual void PostEditChangeProperty(struct FPropertyChangedEvent& PropertyChangedEvent) OVERRIDE;
+	virtual void PostEditChangeChainProperty(struct FPropertyChangedChainEvent& PropertyChangedEvent);
 	// End of UObject interface
 #endif
 
 private:
 	SConstraintCanvas::FSlot* Slot;
+
+#if WITH_EDITOR
+	FGeometry PreEditGeometry;
+	FAnchorData PreEditLayoutData;
+#endif
 };
\ No newline at end of file
diff --git a/Engine/Source/Runtime/UMG/Public/Components/CheckBox.h b/Engine/Source/Runtime/UMG/Public/Components/CheckBox.h
--- a/Engine/Source/Runtime/UMG/Public/Components/CheckBox.h
+++ b/Engine/Source/Runtime/UMG/Public/Components/CheckBox.h
@@ -7,7 +7,7 @@
 DECLARE_DYNAMIC_MULTICAST_DELEGATE_OneParam( FOnCheckBoxComponentStateChanged, bool, bIsChecked );
 
 /** Check box widget */
-UCLASS(meta=( BlueprintSpawnableComponent, Category="Common" ), ClassGroup=UserInterface)
+UCLASS(meta=( Category="Common" ), ClassGroup=UserInterface)
 class UMG_API UCheckBox : public UWidget
 {
 	GENERATED_UCLASS_BODY()
diff --git a/Engine/Source/Runtime/UMG/Public/Components/CircularThrobberComponent.h b/Engine/Source/Runtime/UMG/Public/Components/CircularThrobberComponent.h
--- a/Engine/Source/Runtime/UMG/Public/Components/CircularThrobberComponent.h
+++ b/Engine/Source/Runtime/UMG/Public/Components/CircularThrobberComponent.h
@@ -7,7 +7,7 @@
 class USlateBrushAsset;
 
 /**  A throbber widget that orients images in a spinning circle. */
-UCLASS(meta=(BlueprintSpawnableComponent), ClassGroup=UserInterface)
+UCLASS(meta=( Category="Misc" ), ClassGroup=UserInterface)
 class UMG_API UCircularThrobberComponent : public UWidget
 {
 	GENERATED_UCLASS_BODY()
diff --git a/Engine/Source/Runtime/UMG/Public/Components/ComboBox.h b/Engine/Source/Runtime/UMG/Public/Components/ComboBox.h
--- a/Engine/Source/Runtime/UMG/Public/Components/ComboBox.h
+++ b/Engine/Source/Runtime/UMG/Public/Components/ComboBox.h
@@ -4,13 +4,15 @@
 
 #include "ComboBox.generated.h"
 
-UCLASS(meta=(BlueprintSpawnableComponent), ClassGroup=UserInterface)
+/** The combobox allows you to display a list of options to the user in a dropdown menu for them to select one. */
+UCLASS(meta=( Category="Misc" ), ClassGroup=UserInterface)
 class UMG_API UComboBox : public UWidget
 {
 	GENERATED_UCLASS_BODY()
 
 public:
 
+	/** The list of items to be displayed on the combobox. */
 	UPROPERTY(EditAnywhere, BlueprintReadWrite, Category=Content)
 	TArray<UObject*> Items;
 
@@ -18,22 +20,10 @@ public:
 	UPROPERTY(EditDefaultsOnly, Category=Events)
 	FGenerateWidgetUObject OnGenerateWidget;
 	
-protected:
-	TSharedPtr< SComboBox<UObject*> > ComboBoxWidget() const
-	{
-		if ( MyWidget.IsValid() )
-		{
-			return StaticCastSharedRef< SComboBox<UObject*> >(MyWidget.ToSharedRef());
-		}
-
-		return TSharedPtr< SComboBox<UObject*> >();
-	}
-
 protected:
 	TSharedRef<SWidget> HandleGenerateWidget(UObject* Item) const;
 
 	// UWidget interface
-
 	virtual TSharedRef<SWidget> RebuildWidget() OVERRIDE;
 	// End of UWidget interface
 };
diff --git a/Engine/Source/Runtime/UMG/Public/Components/EditableText.h b/Engine/Source/Runtime/UMG/Public/Components/EditableText.h
--- a/Engine/Source/Runtime/UMG/Public/Components/EditableText.h
+++ b/Engine/Source/Runtime/UMG/Public/Components/EditableText.h
@@ -8,7 +8,7 @@ DECLARE_DYNAMIC_MULTICAST_DELEGATE_OneParam(FOnEditableTextChangedEvent, const F
 DECLARE_DYNAMIC_MULTICAST_DELEGATE_TwoParams(FOnEditableTextCommittedEvent, const FText&, Text, ETextCommit::Type, CommitMethod);
 
 /** Editable text box widget */
-UCLASS(meta=( BlueprintSpawnableComponent, Category="Common" ), ClassGroup=UserInterface)
+UCLASS(meta=( Category="Common" ), ClassGroup=UserInterface)
 class UMG_API UEditableText : public UWidget
 {
 	GENERATED_UCLASS_BODY()
diff --git a/Engine/Source/Runtime/UMG/Public/Components/EditableTextBlock.h b/Engine/Source/Runtime/UMG/Public/Components/EditableTextBlock.h
--- a/Engine/Source/Runtime/UMG/Public/Components/EditableTextBlock.h
+++ b/Engine/Source/Runtime/UMG/Public/Components/EditableTextBlock.h
@@ -8,7 +8,7 @@ DECLARE_DYNAMIC_MULTICAST_DELEGATE_OneParam(FOnEditableTextBlock_TextChangedEven
 //DECLARE_DYNAMIC_MULTICAST_DELEGATE_TwoParams(FOnEditableTextBlock_TextCommittedEvent, const FText&, Text, ETextCommit::Type, Type);
 
 /** Editable text box widget */
-UCLASS(meta=( BlueprintSpawnableComponent, Category="Common" ), ClassGroup=UserInterface)
+UCLASS(meta=( Category="Common" ), ClassGroup=UserInterface)
 class UMG_API UEditableTextBlock : public UWidget
 {
 	GENERATED_UCLASS_BODY()
diff --git a/Engine/Source/Runtime/UMG/Public/Components/HorizontalBox.h b/Engine/Source/Runtime/UMG/Public/Components/HorizontalBox.h
--- a/Engine/Source/Runtime/UMG/Public/Components/HorizontalBox.h
+++ b/Engine/Source/Runtime/UMG/Public/Components/HorizontalBox.h
@@ -4,13 +4,13 @@
 
 #include "HorizontalBox.generated.h"
 
-UCLASS(meta=( BlueprintSpawnableComponent, Category="Panel" ), ClassGroup=UserInterface)
+UCLASS(meta=( Category="Panel" ), ClassGroup=UserInterface)
 class UMG_API UHorizontalBox : public UPanelWidget
 {
 	GENERATED_UCLASS_BODY()
 
 	/** The items placed on the canvas */
-	UPROPERTY(EditAnywhere, EditInline, Category=Slots)
+	UPROPERTY()
 	TArray<UHorizontalBoxSlot*> Slots;
 
 	UHorizontalBoxSlot* AddSlot(UWidget* Content);
diff --git a/Engine/Source/Runtime/UMG/Public/Components/HorizontalBoxSlot.h b/Engine/Source/Runtime/UMG/Public/Components/HorizontalBoxSlot.h
--- a/Engine/Source/Runtime/UMG/Public/Components/HorizontalBoxSlot.h
+++ b/Engine/Source/Runtime/UMG/Public/Components/HorizontalBoxSlot.h
@@ -11,16 +11,38 @@ class UMG_API UHorizontalBoxSlot : public UPanelSlot
 {
 	GENERATED_UCLASS_BODY()
 	
-	UPROPERTY(EditAnywhere, Category=Appearance)
+	/** The amount of padding between the slots parent and the content. */
+	UPROPERTY(EditDefaultsOnly, Category=Appearance)
 	FMargin Padding;
 
 	/** How much space this slot should occupy in the direction of the panel. */
-	UPROPERTY(EditAnywhere, Category=Appearance)
+	UPROPERTY(EditDefaultsOnly, Category=Appearance)
 	FSlateChildSize Size;
 
-	UPROPERTY(EditAnywhere, Category=Appearance)
+	UPROPERTY(EditDefaultsOnly, Category=Appearance)
 	TEnumAsByte<EHorizontalAlignment> HorizontalAlignment;
 
-	UPROPERTY(EditAnywhere, Category=Appearance)
+	UPROPERTY(EditDefaultsOnly, Category=Appearance)
 	TEnumAsByte<EVerticalAlignment> VerticalAlignment;
+
+	UFUNCTION(BlueprintCallable, Category="Appearance")
+	void SetPadding(FMargin InPadding);
+
+	UFUNCTION(BlueprintCallable, Category="Appearance")
+	void SetSize(FSlateChildSize InSize);
+
+	UFUNCTION(BlueprintCallable, Category="Appearance")
+	void SetHorizontalAlignment(EHorizontalAlignment InHorizontalAlignment);
+
+	UFUNCTION(BlueprintCallable, Category="Appearance")
+	void SetVerticalAlignment(EVerticalAlignment InVerticalAlignment);
+
+	void BuildSlot(TSharedRef<SHorizontalBox> HorizontalBox);
+
+	// UPanelSlot interface
+	virtual void Refresh() OVERRIDE;
+	// End of UPanelSlot interface
+
+private:
+	SHorizontalBox::FSlot* Slot;
 };
\ No newline at end of file
diff --git a/Engine/Source/Runtime/UMG/Public/Components/Image.h b/Engine/Source/Runtime/UMG/Public/Components/Image.h
--- a/Engine/Source/Runtime/UMG/Public/Components/Image.h
+++ b/Engine/Source/Runtime/UMG/Public/Components/Image.h
@@ -4,11 +4,10 @@
 
 #include "Image.generated.h"
 
-
 class USlateBrushAsset;
 
 /** Image widget */
-UCLASS(meta=( BlueprintSpawnableComponent, Category="Common" ), ClassGroup=UserInterface)
+UCLASS(meta=( Category="Common" ), ClassGroup=UserInterface)
 class UMG_API UImage : public UWidget
 {
 	GENERATED_UCLASS_BODY()
diff --git a/Engine/Source/Runtime/UMG/Public/Components/ListView.h b/Engine/Source/Runtime/UMG/Public/Components/ListView.h
--- a/Engine/Source/Runtime/UMG/Public/Components/ListView.h
+++ b/Engine/Source/Runtime/UMG/Public/Components/ListView.h
@@ -5,7 +5,7 @@
 #include "ListView.generated.h"
 
 /** Allows thousands of items to be displayed in a list.  Generates widgets dynamically for each item. */
-UCLASS(meta=( BlueprintSpawnableComponent ), ClassGroup=UserInterface)
+UCLASS(meta=( Category="Misc" ), ClassGroup=UserInterface)
 class UMG_API UListView : public UTableViewBase
 {
 	GENERATED_UCLASS_BODY()
diff --git a/Engine/Source/Runtime/UMG/Public/Components/Overlay.h b/Engine/Source/Runtime/UMG/Public/Components/Overlay.h
--- a/Engine/Source/Runtime/UMG/Public/Components/Overlay.h
+++ b/Engine/Source/Runtime/UMG/Public/Components/Overlay.h
@@ -4,13 +4,13 @@
 
 #include "Overlay.generated.h"
 
-UCLASS(meta=( BlueprintSpawnableComponent, Category="Panel" ), ClassGroup=UserInterface)
+UCLASS(meta=( Category="Panel" ), ClassGroup=UserInterface)
 class UMG_API UOverlay : public UPanelWidget
 {
 	GENERATED_UCLASS_BODY()
 
 	/** The items placed on the canvas */
-	UPROPERTY(EditAnywhere, EditInline, Category=Slots)
+	UPROPERTY()
 	TArray<UOverlaySlot*> Slots;
 
 	UOverlaySlot* AddSlot(UWidget* Content);
diff --git a/Engine/Source/Runtime/UMG/Public/Components/OverlaySlot.h b/Engine/Source/Runtime/UMG/Public/Components/OverlaySlot.h
--- a/Engine/Source/Runtime/UMG/Public/Components/OverlaySlot.h
+++ b/Engine/Source/Runtime/UMG/Public/Components/OverlaySlot.h
@@ -12,12 +12,39 @@ class UMG_API UOverlaySlot : public UPanelSlot
 {
 	GENERATED_UCLASS_BODY()
 	
-	UPROPERTY(EditAnywhere, Category=Appearance)
+	/** The padding area between the slot and the content it contains. */
+	UPROPERTY(EditDefaultsOnly, Category=Layout)
 	FMargin Padding;
 
-	UPROPERTY(EditAnywhere, Category=Appearance)
+	/** The alignment of the object horizontally. */
+	UPROPERTY(EditDefaultsOnly, Category=Layout)
 	TEnumAsByte<EHorizontalAlignment> HorizontalAlignment;
 
-	UPROPERTY(EditAnywhere, Category=Appearance)
+	/** The alignment of the object vertically. */
+	UPROPERTY(EditDefaultsOnly, Category=Layout)
 	TEnumAsByte<EVerticalAlignment> VerticalAlignment;
+
+	UFUNCTION(BlueprintCallable, Category="Appearance")
+	void SetPadding(FMargin InPadding);
+
+	UFUNCTION(BlueprintCallable, Category="Appearance")
+	void SetHorizontalAlignment(EHorizontalAlignment InHorizontalAlignment);
+
+	UFUNCTION(BlueprintCallable, Category="Appearance")
+	void SetVerticalAlignment(EVerticalAlignment InVerticalAlignment);
+
+	void BuildSlot(TSharedRef<SHorizontalBox> HorizontalBox);
+
+	// UPanelSlot interface
+	virtual void Refresh() OVERRIDE;
+	// End of UPanelSlot interface
+
+	/** Builds the underlying FSlot for the Slate layout panel. */
+	void BuildSlot(TSharedRef<SOverlay> InOverlay);
+
+private:
+	//TODO UMG Slots should hold weak or shared refs to slots.
+
+	/** A raw pointer to the slot to allow us to adjust the size, padding...etc at runtime. */
+	SOverlay::FOverlaySlot* Slot;
 };
diff --git a/Engine/Source/Runtime/UMG/Public/Components/PanelSlot.h b/Engine/Source/Runtime/UMG/Public/Components/PanelSlot.h
--- a/Engine/Source/Runtime/UMG/Public/Components/PanelSlot.h
+++ b/Engine/Source/Runtime/UMG/Public/Components/PanelSlot.h
@@ -25,4 +25,20 @@ public:
 	virtual void MoveTo(const FVector2D& Location);
 
 	virtual bool CanMove() const;
+
+	/** Applies all properties to the live slot if possible. */
+	virtual void Refresh()
+	{
+	}
+
+#if WITH_EDITOR
+
+	virtual void PostEditChangeProperty(struct FPropertyChangedEvent& PropertyChangedEvent) OVERRIDE
+	{
+		Super::PostEditChangeProperty(PropertyChangedEvent);
+
+		Refresh();
+	}
+
+#endif
 };
diff --git a/Engine/Source/Runtime/UMG/Public/Components/ProgressBar.h b/Engine/Source/Runtime/UMG/Public/Components/ProgressBar.h
--- a/Engine/Source/Runtime/UMG/Public/Components/ProgressBar.h
+++ b/Engine/Source/Runtime/UMG/Public/Components/ProgressBar.h
@@ -5,7 +5,7 @@
 #include "ProgressBar.generated.h"
 
 /** ProgressBar widget */
-UCLASS(meta=(BlueprintSpawnableComponent), ClassGroup=UserInterface)
+UCLASS(meta=( Category="Misc" ), ClassGroup=UserInterface)
 class UMG_API UProgressBar : public UWidget
 {
 	GENERATED_UCLASS_BODY()
diff --git a/Engine/Source/Runtime/UMG/Public/Components/ScrollBox.h b/Engine/Source/Runtime/UMG/Public/Components/ScrollBox.h
--- a/Engine/Source/Runtime/UMG/Public/Components/ScrollBox.h
+++ b/Engine/Source/Runtime/UMG/Public/Components/ScrollBox.h
@@ -5,7 +5,7 @@
 #include "ScrollBox.generated.h"
 
 /** An arbitrary scrollable collection of widgets.  Great for presenting 10-100 widgets in a list.  Doesn't support virtualization. */
-UCLASS(meta=( BlueprintSpawnableComponent, Category="Panel" ), ClassGroup=UserInterface)
+UCLASS(meta=( Category="Panel" ), ClassGroup=UserInterface)
 class UMG_API UScrollBox : public UPanelWidget
 {
 	GENERATED_UCLASS_BODY()
diff --git a/Engine/Source/Runtime/UMG/Public/Components/ScrollBoxSlot.h b/Engine/Source/Runtime/UMG/Public/Components/ScrollBoxSlot.h
--- a/Engine/Source/Runtime/UMG/Public/Components/ScrollBoxSlot.h
+++ b/Engine/Source/Runtime/UMG/Public/Components/ScrollBoxSlot.h
@@ -20,6 +20,16 @@ class UMG_API UScrollBoxSlot : public UPanelSlot
 	UPROPERTY(EditDefaultsOnly, Category=Layout)
 	TEnumAsByte<EHorizontalAlignment> HorizontalAlignment;
 
+	UFUNCTION(BlueprintCallable, Category="Appearance")
+	void SetPadding(FMargin InPadding);
+
+	UFUNCTION(BlueprintCallable, Category="Appearance")
+	void SetHorizontalAlignment(EHorizontalAlignment InHorizontalAlignment);
+
+	// UPanelSlot interface
+	virtual void Refresh() OVERRIDE;
+	// End of UPanelSlot interface
+
 	/** Builds the underlying FSlot for the Slate layout panel. */
 	void BuildSlot(TSharedRef<SScrollBox> ScrollBox);
 
diff --git a/Engine/Source/Runtime/UMG/Public/Components/Slider.h b/Engine/Source/Runtime/UMG/Public/Components/Slider.h
--- a/Engine/Source/Runtime/UMG/Public/Components/Slider.h
+++ b/Engine/Source/Runtime/UMG/Public/Components/Slider.h
@@ -9,7 +9,7 @@ DECLARE_DYNAMIC_MULTICAST_DELEGATE(FOnMouseCaptureEndEvent);
 DECLARE_DYNAMIC_MULTICAST_DELEGATE_OneParam(FOnFloatValueChangedEvent, float, Value);
 
 /** Slider widget */
-UCLASS(meta=(BlueprintSpawnableComponent), ClassGroup=UserInterface)
+UCLASS(meta=( Category="Common" ), ClassGroup=UserInterface)
 class UMG_API USlider : public UWidget
 {
 	GENERATED_UCLASS_BODY()
diff --git a/Engine/Source/Runtime/UMG/Public/Components/Spacer.h b/Engine/Source/Runtime/UMG/Public/Components/Spacer.h
--- a/Engine/Source/Runtime/UMG/Public/Components/Spacer.h
+++ b/Engine/Source/Runtime/UMG/Public/Components/Spacer.h
@@ -5,7 +5,7 @@
 #include "Spacer.generated.h"
 
 /** A spacer widget; it does not have a visual representation, and just provides padding between other widgets. */
-UCLASS(meta=(BlueprintSpawnableComponent), ClassGroup=UserInterface)
+UCLASS(meta=( Category="Misc" ), ClassGroup=UserInterface)
 class UMG_API USpacer : public UWidget
 {
 	GENERATED_UCLASS_BODY()
@@ -14,10 +14,15 @@ class UMG_API USpacer : public UWidget
 	UPROPERTY(EditDefaultsOnly, Category=Appearance)
 	FVector2D Size;
 
-	//TODO UMG UFunction SetSize
+	/** Sets the size of the spacer */
+	UFUNCTION(BlueprintCallable, Category="Widget")
+	void SetSize(FVector2D InSize);
 
 protected:
 	// UWidget interface
 	virtual TSharedRef<SWidget> RebuildWidget() OVERRIDE;
 	// End of UWidget interface
+
+protected:
+	TSharedPtr<SSpacer> MySpacer;
 };
diff --git a/Engine/Source/Runtime/UMG/Public/Components/TableViewBase.h b/Engine/Source/Runtime/UMG/Public/Components/TableViewBase.h
--- a/Engine/Source/Runtime/UMG/Public/Components/TableViewBase.h
+++ b/Engine/Source/Runtime/UMG/Public/Components/TableViewBase.h
@@ -5,7 +5,7 @@
 #include "TableViewBase.generated.h"
 
 /** The base class for all wrapped table views */
-UCLASS(Abstract, meta=(BlueprintSpawnableComponent), ClassGroup=UserInterface)
+UCLASS(Abstract)
 class UMG_API UTableViewBase : public UWidget
 {
 	GENERATED_UCLASS_BODY()
diff --git a/Engine/Source/Runtime/UMG/Public/Components/TextBlock.h b/Engine/Source/Runtime/UMG/Public/Components/TextBlock.h
--- a/Engine/Source/Runtime/UMG/Public/Components/TextBlock.h
+++ b/Engine/Source/Runtime/UMG/Public/Components/TextBlock.h
@@ -5,7 +5,7 @@
 #include "TextBlock.generated.h"
 
 /** A simple static text widget */
-UCLASS(meta=( BlueprintSpawnableComponent, DisplayName="Text", Category="Common" ), ClassGroup=UserInterface )
+UCLASS(meta=( Category="Common" ), ClassGroup=UserInterface )
 class UMG_API UTextBlock : public UWidget
 {
 	GENERATED_UCLASS_BODY()
diff --git a/Engine/Source/Runtime/UMG/Public/Components/ThrobberComponent.h b/Engine/Source/Runtime/UMG/Public/Components/ThrobberComponent.h
--- a/Engine/Source/Runtime/UMG/Public/Components/ThrobberComponent.h
+++ b/Engine/Source/Runtime/UMG/Public/Components/ThrobberComponent.h
@@ -7,7 +7,7 @@
 class USlateBrushAsset;
 
 /** A Throbber widget that shows several zooming circles in a row. */
-UCLASS(meta=(BlueprintSpawnableComponent), ClassGroup=UserInterface)
+UCLASS(meta=( Category="Misc" ), ClassGroup=UserInterface)
 class UMG_API UThrobberComponent : public UWidget
 {
 	GENERATED_UCLASS_BODY()
diff --git a/Engine/Source/Runtime/UMG/Public/Components/TileView.h b/Engine/Source/Runtime/UMG/Public/Components/TileView.h
--- a/Engine/Source/Runtime/UMG/Public/Components/TileView.h
+++ b/Engine/Source/Runtime/UMG/Public/Components/TileView.h
@@ -5,7 +5,7 @@
 #include "TileView.generated.h"
 
 /** A flow panel that presents the contents as a set of tiles all uniformly sized. */
-UCLASS(meta=(BlueprintSpawnableComponent), ClassGroup=UserInterface)
+UCLASS(meta=( Category="Misc" ), ClassGroup=UserInterface)
 class UMG_API UTileView : public UTableViewBase
 {
 	GENERATED_UCLASS_BODY()
diff --git a/Engine/Source/Runtime/UMG/Public/Components/UniformGridPanel.h b/Engine/Source/Runtime/UMG/Public/Components/UniformGridPanel.h
--- a/Engine/Source/Runtime/UMG/Public/Components/UniformGridPanel.h
+++ b/Engine/Source/Runtime/UMG/Public/Components/UniformGridPanel.h
@@ -5,13 +5,13 @@
 #include "UniformGridPanel.generated.h"
 
 /** A panel that evenly divides up available space between all of its children. */
-UCLASS(meta=( BlueprintSpawnableComponent, Category="Panel"), ClassGroup=UserInterface)
+UCLASS(meta=( Category="Panel" ), ClassGroup=UserInterface)
 class UMG_API UUniformGridPanel : public UPanelWidget
 {
 	GENERATED_UCLASS_BODY()
 
-	/** The items placed on the canvas */
-	UPROPERTY(EditDefaultsOnly, EditInline, Category=Slots)
+	/** The items placed on the grid */
+	UPROPERTY()
 	TArray<UUniformGridSlot*> Slots;
 
 	/** Padding given to each slot */
diff --git a/Engine/Source/Runtime/UMG/Public/Components/UniformGridSlot.h b/Engine/Source/Runtime/UMG/Public/Components/UniformGridSlot.h
--- a/Engine/Source/Runtime/UMG/Public/Components/UniformGridSlot.h
+++ b/Engine/Source/Runtime/UMG/Public/Components/UniformGridSlot.h
@@ -41,6 +41,10 @@ public:
 	UFUNCTION(BlueprintCallable, Category=Layout)
 	void SetColumn(int32 InColumn);
 
+	// UPanelSlot interface
+	virtual void Refresh() OVERRIDE;
+	// End of UPanelSlot interface
+
 	/** Builds the underlying FSlot for the Slate layout panel. */
 	void BuildSlot(TSharedRef<SUniformGridPanel> GridPanel);
 
diff --git a/Engine/Source/Runtime/UMG/Public/Components/VerticalBox.h b/Engine/Source/Runtime/UMG/Public/Components/VerticalBox.h
--- a/Engine/Source/Runtime/UMG/Public/Components/VerticalBox.h
+++ b/Engine/Source/Runtime/UMG/Public/Components/VerticalBox.h
@@ -8,13 +8,13 @@
  * A vertical box widget is a layout panel allowing child widgets to be automatically laid out
  * vertically.
  */
-UCLASS(meta=( BlueprintSpawnableComponent, Category="Panel" ), ClassGroup=UserInterface)
+UCLASS(meta=( Category="Panel" ), ClassGroup=UserInterface)
 class UMG_API UVerticalBox : public UPanelWidget
 {
 	GENERATED_UCLASS_BODY()
 
 	/** The slots containing the widgets that are flowed vertically. */
-	UPROPERTY(EditDefaultsOnly, EditInline, Category=Slots)
+	UPROPERTY()
 	TArray<UVerticalBoxSlot*> Slots;
 
 	//TODO UMG Add ways to make adding slots callable by blueprints.
diff --git a/Engine/Source/Runtime/UMG/Public/Components/VerticalBoxSlot.h b/Engine/Source/Runtime/UMG/Public/Components/VerticalBoxSlot.h
--- a/Engine/Source/Runtime/UMG/Public/Components/VerticalBoxSlot.h
+++ b/Engine/Source/Runtime/UMG/Public/Components/VerticalBoxSlot.h
@@ -28,6 +28,22 @@ class UMG_API UVerticalBoxSlot : public UPanelSlot
 	UPROPERTY(EditDefaultsOnly, Category=Layout)
 	TEnumAsByte<EVerticalAlignment> VerticalAlignment;
 
+	UFUNCTION(BlueprintCallable, Category="Appearance")
+	void SetPadding(FMargin InPadding);
+
+	UFUNCTION(BlueprintCallable, Category="Appearance")
+	void SetSize(FSlateChildSize InSize);
+
+	UFUNCTION(BlueprintCallable, Category="Appearance")
+	void SetHorizontalAlignment(EHorizontalAlignment InHorizontalAlignment);
+
+	UFUNCTION(BlueprintCallable, Category="Appearance")
+	void SetVerticalAlignment(EVerticalAlignment InVerticalAlignment);
+
+	// UPanelSlot interface
+	virtual void Refresh() OVERRIDE;
+	// End of UPanelSlot interface
+
 	/** Builds the underlying FSlot for the Slate layout panel. */
 	void BuildSlot(TSharedRef<SVerticalBox> InVerticalBox);
 
diff --git a/Engine/Source/Runtime/UMG/Public/Components/Widget.h b/Engine/Source/Runtime/UMG/Public/Components/Widget.h
--- a/Engine/Source/Runtime/UMG/Public/Components/Widget.h
+++ b/Engine/Source/Runtime/UMG/Public/Components/Widget.h
@@ -111,19 +111,35 @@ public:
 	UFUNCTION(BlueprintCallable, Category="Widget")
 	TEnumAsByte<ESlateVisibility::Type> GetVisibility();
 
-	/** Sets the visiblity of the widget. */
+	/** Sets the visibility of the widget. */
 	UFUNCTION(BlueprintCallable, Category="Widget")
 	void SetVisibility(TEnumAsByte<ESlateVisibility::Type> InVisibility);
 
+	/** Gets if the button is currently being hovered by the mouse */
+	UFUNCTION(BlueprintCallable, Category="Widget")
+	bool IsHovered() const;
+
 	/** Gets the underlying Slate SWidget for this UWidget. */
 	TSharedRef<SWidget> GetWidget();
 
+	/**
+	 * Applies all properties to the native widget if possible.  This is called after a widget is constructed.
+	 * It can also be called by the editor to update modified state, so ensure all initialization to a widgets
+	 * properties are performed here, or the property and visual state may become unsynced.
+	 */
+	virtual void SyncronizeProperties();
+
 #if WITH_EDITOR
 	/** Allows general fixups and connections only used at editor time. */
 	virtual void ConnectEditorData() { }
 
 	// UObject interface
-	virtual void PostEditChangeProperty(struct FPropertyChangedEvent& PropertyChangedEvent) OVERRIDE;
+	virtual void PostEditChangeProperty(struct FPropertyChangedEvent& PropertyChangedEvent) OVERRIDE
+	{
+		Super::PostEditChangeProperty(PropertyChangedEvent);
+
+		SyncronizeProperties();
+	}
 	// End of UObject interface
 #endif
 
diff --git a/Engine/Source/Runtime/UMG/Public/Slate/SConstraintCanvas.h b/Engine/Source/Runtime/UMG/Public/Slate/SConstraintCanvas.h
--- a/Engine/Source/Runtime/UMG/Public/Slate/SConstraintCanvas.h
+++ b/Engine/Source/Runtime/UMG/Public/Slate/SConstraintCanvas.h
@@ -116,6 +116,9 @@ public:
 	 */
 	void ClearChildren();
 
+	/** Get the cached geometry of the canvas. */
+	FGeometry GetCachedGeometry() const { return CachedGeometry; }
+
 public:
 
 	// Begin SWidget overrides
@@ -134,4 +137,6 @@ protected:
 
 	/** The ConstraintCanvas widget's children. */
 	TPanelChildren< FSlot > Children;
+
+	mutable FGeometry CachedGeometry;
 };
