UMG - Fix startup crash due to UBlueprint::OnChanged function being bound twice to the same function.  FWidgetBlueprintEditor::InitWidgetBLueprintEditor and superclass FBlueprintEditor::InitBlueprintEditor both bind it.

diff --git a/Engine/Source/Editor/UMGEditor/Private/WidgetBlueprintEditor.cpp b/Engine/Source/Editor/UMGEditor/Private/WidgetBlueprintEditor.cpp
--- a/Engine/Source/Editor/UMGEditor/Private/WidgetBlueprintEditor.cpp
+++ b/Engine/Source/Editor/UMGEditor/Private/WidgetBlueprintEditor.cpp
@@ -61,9 +61,6 @@ void FWidgetBlueprintEditor::InitWidgetBlueprintEditor(const EToolkitMode::Type
 {
 	InitBlueprintEditor(Mode, InitToolkitHost, InBlueprints, bShouldOpenInDefaultsMode);
 
-	UWidgetBlueprint* Blueprint = GetWidgetBlueprintObj();
-	Blueprint->OnChanged().AddSP(this, &FWidgetBlueprintEditor::OnBlueprintChanged);
-
 	UpdatePreview(GetWidgetBlueprintObj(), true);
 }
 
diff --git a/Engine/Source/Editor/UMGEditor/Private/WidgetBlueprintEditor.h b/Engine/Source/Editor/UMGEditor/Private/WidgetBlueprintEditor.h
--- a/Engine/Source/Editor/UMGEditor/Private/WidgetBlueprintEditor.h
+++ b/Engine/Source/Editor/UMGEditor/Private/WidgetBlueprintEditor.h
@@ -36,7 +36,7 @@ public:
 	void SelectWidgets(TArray<UWidget*> Widgets);
 
 private:
-	void OnBlueprintChanged(UBlueprint* InBlueprint);
+	virtual void OnBlueprintChanged(UBlueprint* InBlueprint) OVERRIDE;
 
 	void DestroyPreview();
 	void UpdatePreview(UBlueprint* InBlueprint, bool bInForceFullUpdate);
