UMG - Drop Previews, Design Time variable display of widgets (dashed border around invisible widgets), capable of dynamically adding widgets to several container objects.  Adding events to the UMG Border.  Having more widgets implement SyncronizeProperties so that we can live update the preview without rebuilding the entire blueprint.  Fixing more anchoring issues with the new canvas layout panel.

diff --git a/Engine/Source/Editor/EditorStyle/Private/SlateEditorStyle.cpp b/Engine/Source/Editor/EditorStyle/Private/SlateEditorStyle.cpp
--- a/Engine/Source/Editor/EditorStyle/Private/SlateEditorStyle.cpp
+++ b/Engine/Source/Editor/EditorStyle/Private/SlateEditorStyle.cpp
@@ -5212,6 +5212,22 @@ void FSlateEditorStyle::FStyle::SetupUMGEditorStyles()
 	Set("UMGEditor.FillCenter", new IMAGE_BRUSH("Icons/UMG/umg_anchors_fill_center", Icon64x64));
 	Set("UMGEditor.FillRight", new IMAGE_BRUSH("Icons/UMG/umg_anchors_fill_right", Icon64x64));
 	Set("UMGEditor.FillFill", new IMAGE_BRUSH("Icons/UMG/umg_anchors_fill_fill", Icon64x64));
+	
+	Set("UMGEditor.ToolboxPreviewWidget", new IMAGE_BRUSH("Icons/UMG/umg_anchors_fill_fill", Icon16x16));
+	
+	Set("UMGEditor.ToolboxHeader", FTableRowStyle()
+		.SetEvenRowBackgroundBrush( BOX_BRUSH( "PropertyView/DetailCategoryTop", FMargin( 4/16.0f, 8.0f/16.0f, 4/16.0f, 4/16.0f ) ) )
+		.SetEvenRowBackgroundHoveredBrush( BOX_BRUSH( "PropertyView/DetailCategoryTop", FMargin( 4/16.0f, 8.0f/16.0f, 4/16.0f, 4/16.0f ), FLinearColor(0.5f,0.5f,0.5f,1.0f) ) )
+		.SetOddRowBackgroundBrush( BOX_BRUSH( "PropertyView/DetailCategoryTop", FMargin( 4/16.0f, 8.0f/16.0f, 4/16.0f, 4/16.0f ) ) )
+		.SetOddRowBackgroundHoveredBrush( BOX_BRUSH( "PropertyView/DetailCategoryTop", FMargin( 4/16.0f, 8.0f/16.0f, 4/16.0f, 4/16.0f ), FLinearColor(0.5f,0.5f,0.5f,1.0f) ) )
+		.SetSelectorFocusedBrush( BORDER_BRUSH( "Common/Selector", FMargin(4.f/16.f), SelectorColor ) )
+		.SetActiveBrush( IMAGE_BRUSH( "Common/Selection", Icon8x8, SelectionColor ) )
+		.SetActiveHoveredBrush( IMAGE_BRUSH( "Common/Selection", Icon8x8, SelectionColor ) )
+		.SetInactiveBrush( IMAGE_BRUSH( "Common/Selection", Icon8x8, SelectionColor_Inactive ) )
+		.SetInactiveHoveredBrush( IMAGE_BRUSH( "Common/Selection", Icon8x8, SelectionColor_Inactive ) )
+		.SetTextColor( DefaultForeground )
+		.SetSelectedTextColor( InvertedForeground )
+		);
 }
 
 void FSlateEditorStyle::FStyle::SetupTranslationEditorStyles()
diff --git a/Engine/Source/Editor/UMGEditor/Private/SUMGDesigner.cpp b/Engine/Source/Editor/UMGEditor/Private/SUMGDesigner.cpp
--- a/Engine/Source/Editor/UMGEditor/Private/SUMGDesigner.cpp
+++ b/Engine/Source/Editor/UMGEditor/Private/SUMGDesigner.cpp
@@ -31,7 +31,8 @@ static bool LocateWidgetsUnderCursor_Helper(FArrangedWidget& Candidate, FVector2
 
 		// Check to see if we were asked to still allow children to be hit test visible
 		bool bHitChildWidget = false;
-		//if ( ( Candidate.Widget->GetVisibility().AreChildrenHitTestVisible() ) != 0 || OutWidgetsUnderCursor. )
+		
+		if ( Candidate.Widget->GetVisibility().AreChildrenHitTestVisible() )//!= 0 || OutWidgetsUnderCursor. )
 		{
 			FArrangedChildren ArrangedChildren(OutWidgetsUnderCursor.GetFilter());
 			Candidate.Widget->ArrangeChildren(Candidate.Geometry, ArrangedChildren);
@@ -45,7 +46,9 @@ static bool LocateWidgetsUnderCursor_Helper(FArrangedWidget& Candidate, FVector2
 		}
 
 		// If we hit a child widget or we hit our candidate widget then we'll append our widgets
-		const bool bHitCandidateWidget = OutWidgetsUnderCursor.Accepts(Candidate.Widget->GetVisibility());
+		const bool bHitCandidateWidget = OutWidgetsUnderCursor.Accepts(Candidate.Widget->GetVisibility()) &&
+			Candidate.Widget->GetVisibility().AreChildrenHitTestVisible();
+		
 		bHitAnyWidget = bHitChildWidget || bHitCandidateWidget;
 		if ( !bHitAnyWidget )
 		{
@@ -64,7 +67,9 @@ static bool LocateWidgetsUnderCursor_Helper(FArrangedWidget& Candidate, FVector2
 
 void SUMGDesigner::Construct(const FArguments& InArgs, TSharedPtr<FWidgetBlueprintEditor> InBlueprintEditor)
 {
-	PreviewWidgetActor = NULL;
+	PreviewWidgetObject = NULL;
+    DropPreviewWidget = NULL;
+    DropPreviewParent = NULL;
 	BlueprintEditor = InBlueprintEditor;
 	CurrentHandle = DH_NONE;
 
@@ -182,6 +187,14 @@ FSelectedWidget SUMGDesigner::GetWidgetAtCursor(const FGeometry& MyGeometry, con
 		{
 			FArrangedWidget& Child = Children.GetInternalArray()[ChildIndex];
 			Preview = WidgetActor->GetWidgetHandle(Child.Widget);
+            
+            // Ignore the drop preview widget when doing widget picking
+            if (Preview == DropPreviewWidget)
+            {
+                Preview = NULL;
+                continue;
+            }
+            
 			if ( Preview )
 			{
 				ArrangedWidget = Child;
@@ -607,14 +620,14 @@ FCursorReply SUMGDesigner::OnCursorQuery(const FGeometry& MyGeometry, const FPoi
 
 void SUMGDesigner::Tick(const FGeometry& AllottedGeometry, const double InCurrentTime, const float InDeltaTime)
 {
-	UUserWidget* WidgetActor = BlueprintEditor.Pin()->GetPreview();
+	UUserWidget* LatestPreviewWidget = BlueprintEditor.Pin()->GetPreview();
 
-	if ( WidgetActor != PreviewWidgetActor )
+	if ( LatestPreviewWidget != PreviewWidgetObject )
 	{
-		PreviewWidgetActor = WidgetActor;
-		if ( PreviewWidgetActor )
+		PreviewWidgetObject = LatestPreviewWidget;
+		if ( PreviewWidgetObject )
 		{
-			TSharedRef<SWidget> CurrentWidget = PreviewWidgetActor->MakeFullScreenWidget();
+			TSharedRef<SWidget> CurrentWidget = PreviewWidgetObject->MakeFullScreenWidget();
 
 			if ( CurrentWidget != PreviewWidget.Pin() )
 			{
@@ -641,13 +654,13 @@ void SUMGDesigner::Tick(const FGeometry& AllottedGeometry, const double InCurren
 	// Update the selected widget to match the selected template.
 	if ( CurrentSelection.IsValid() )
 	{
-		if ( WidgetActor )
+		if ( LatestPreviewWidget )
 		{
 			// Set the selected widget so that we can draw the highlight
 			//@TODO UMG Don't store the transient Widget as the selected object, store the template component and lookup the widget from that.
 			//SelectedWidget = ArrangedWidget.Widget;
 
-			SelectedWidget = WidgetActor->GetWidgetFromName(CurrentSelection.GetTemplate()->GetName());
+			SelectedWidget = LatestPreviewWidget->GetWidgetFromName(CurrentSelection.GetTemplate()->GetName());
 		}
 	}
 
@@ -668,60 +681,155 @@ void SUMGDesigner::OnDragLeave(const FDragDropEvent& DragDropEvent)
 	{
 		DragDropOp->ResetToDefaultToolTip();
 	}
+    
+    if (DropPreviewWidget)
+    {
+        UWidgetBlueprint* BP = GetBlueprint();
+        BP->WidgetTree->RemoveWidget(DropPreviewWidget);
+        DropPreviewWidget = NULL;
+    }
 }
 
 FReply SUMGDesigner::OnDragOver(const FGeometry& MyGeometry, const FDragDropEvent& DragDropEvent)
 {
-	TSharedPtr<FWidgetTemplateDragDropOp> DragDropOp = DragDropEvent.GetOperationAs<FWidgetTemplateDragDropOp>();
-	if ( DragDropOp.IsValid() )
-	{
-		//@TODO UMG Drop Feedback
-		return FReply::Handled();
-	}
+    UWidgetBlueprint* BP = GetBlueprint();
+    
+	if (DropPreviewWidget)
+    {
+        if (DropPreviewParent)
+        {
+            DropPreviewParent->RemoveChild(DropPreviewWidget);
+        }
+        
+        BP->WidgetTree->RemoveWidget(DropPreviewWidget);
+        DropPreviewWidget = NULL;
+    }
+    
+    DropPreviewWidget = AddPreview(MyGeometry, DragDropEvent);
+    if ( DropPreviewWidget )
+    {
+        //@TODO UMG Drop Feedback
+        return FReply::Handled();
+    }
 
 	return FReply::Unhandled();
 }
 
-FReply SUMGDesigner::OnDrop(const FGeometry& MyGeometry, const FDragDropEvent& DragDropEvent)
+UWidget* SUMGDesigner::AddPreview(const FGeometry& MyGeometry, const FDragDropEvent& DragDropEvent)
 {
-	TSharedPtr<FWidgetTemplateDragDropOp> DragDropOp = DragDropEvent.GetOperationAs<FWidgetTemplateDragDropOp>();
+    TSharedPtr<FWidgetTemplateDragDropOp> DragDropOp = DragDropEvent.GetOperationAs<FWidgetTemplateDragDropOp>();
 	if ( DragDropOp.IsValid() )
 	{
-		FArrangedWidget ArrangedWidget(SNullWidget::NullWidget, FGeometry());
-		FSelectedWidget Selection = GetWidgetAtCursor(MyGeometry, DragDropEvent, ArrangedWidget);
-
-		UWidgetBlueprint* BP = GetBlueprint();
-
-		if ( Selection.IsValid() && Selection.GetTemplate()->IsA(UPanelWidget::StaticClass()) )
-		{
-			UPanelWidget* Parent = Cast<UPanelWidget>(Selection.GetTemplate());
-
-			UWidget* Widget = DragDropOp->Template->Create(BP->WidgetTree);
-
-			FVector2D LocalPosition = ArrangedWidget.Geometry.AbsoluteToLocal(DragDropEvent.GetScreenSpacePosition());
-			Parent->AddChild(Widget, LocalPosition);
-			//@TODO UMG When we add a child blindly we need to default the slot size to the preferred size of the widget if the container supports such things.
-			//@TODO UMG We may need a desired size canvas, where the slots have no size, they only give you position, alternatively, maybe slots that don't clip, so center is still easy.
-
-			// Update the selected template to be the newly created one.
-			//SelectedTemplate = Widget;
-
-			FBlueprintEditorUtils::MarkBlueprintAsStructurallyModified(BP);
-
-			return FReply::Handled();
-		}
-		else if ( BP->WidgetTree->WidgetTemplates.Num() == 0 )
-		{
-			// No existing templates so just create it and make it the root widget.
-			UWidget* Widget = DragDropOp->Template->Create(BP->WidgetTree);
-			//SelectedTemplate = Widget;
+        FArrangedWidget ArrangedWidget(SNullWidget::NullWidget, FGeometry());
+        FSelectedWidget Selection = GetWidgetAtCursor(MyGeometry, DragDropEvent, ArrangedWidget);
+        
+        UWidgetBlueprint* BP = GetBlueprint();
+        
+        if ( Selection.IsValid() )
+        {
+            UWidget* Target = Selection.GetPreview();
+            if ( Target->IsA(UPanelWidget::StaticClass()) )
+            {
+                UPanelWidget* Parent = Cast<UPanelWidget>(Target);
+                
+                UWidget* Widget = DragDropOp->Template->Create(BP->WidgetTree);
+				Widget->IsDesignTime(true);
+            
+                FVector2D LocalPosition = ArrangedWidget.Geometry.AbsoluteToLocal(DragDropEvent.GetScreenSpacePosition());
+                Parent->AddChild(Widget, LocalPosition);
+                //@TODO UMG When we add a child blindly we need to default the slot size to the preferred size of the widget if the container supports such things.
+                //@TODO UMG We may need a desired size canvas, where the slots have no size, they only give you position, alternatively, maybe slots that don't clip, so center is still easy.
+            
+                DropPreviewParent = Parent;
+            
+                return Widget;
+            }
+            else if ( BP->WidgetTree->WidgetTemplates.Num() == 1 )
+            {
+                UWidget* Widget = DragDropOp->Template->Create(BP->WidgetTree);
+                Widget->IsDesignTime(true);
+				
+                DropPreviewParent = NULL;
+                
+                return Widget;
+            }
+        }
+    }
+    
+    return NULL;
+}
 
-			FBlueprintEditorUtils::MarkBlueprintAsStructurallyModified(BP);
+bool SUMGDesigner::AddToTemplate(const FGeometry& MyGeometry, const FDragDropEvent& DragDropEvent)
+{
+    TSharedPtr<FWidgetTemplateDragDropOp> DragDropOp = DragDropEvent.GetOperationAs<FWidgetTemplateDragDropOp>();
+	if ( DragDropOp.IsValid() )
+	{
+        FArrangedWidget ArrangedWidget(SNullWidget::NullWidget, FGeometry());
+        FSelectedWidget Selection = GetWidgetAtCursor(MyGeometry, DragDropEvent, ArrangedWidget);
+        
+        UWidgetBlueprint* BP = GetBlueprint();
+        
+        if ( Selection.IsValid() )
+        {
+            UWidget* Target = Selection.GetTemplate();
+            if ( Target->IsA(UPanelWidget::StaticClass()) )
+            {
+                UPanelWidget* Parent = Cast<UPanelWidget>(Target);
+                
+                UWidget* Widget = DragDropOp->Template->Create(BP->WidgetTree);
+				Widget->IsDesignTime(true);
+                
+                FVector2D LocalPosition = ArrangedWidget.Geometry.AbsoluteToLocal(DragDropEvent.GetScreenSpacePosition());
+                Parent->AddChild(Widget, LocalPosition);
+                //@TODO UMG When we add a child blindly we need to default the slot size to the preferred size of the widget if the container supports such things.
+                //@TODO UMG We may need a desired size canvas, where the slots have no size, they only give you position, alternatively, maybe slots that don't clip, so center is still easy.
+                
+                // Update the selected template to be the newly created one.
+                CurrentSelection = FSelectedWidget::FromTemplate(BlueprintEditor.Pin(), Widget);
+                
+                return true;
+            }
+        }
+        
+        if ( BP->WidgetTree->WidgetTemplates.Num() == 0 )
+        {
+            // TODO UMG This method isn't great, maybe the user widget should just be a canvas.
+            
+            // Add it to the root if there are no other widgets to add it to.
+            UWidget* Widget = DragDropOp->Template->Create(BP->WidgetTree);
+			Widget->IsDesignTime(true);
+            
+            CurrentSelection = FSelectedWidget::FromTemplate(BlueprintEditor.Pin(), Widget);
+            
+            return true;
+        }
+    }
+    
+    return false;
+}
 
-			return FReply::Handled();
-		}
-	}
+FReply SUMGDesigner::OnDrop(const FGeometry& MyGeometry, const FDragDropEvent& DragDropEvent)
+{
+    UWidgetBlueprint* BP = GetBlueprint();
+    
+    if (DropPreviewWidget)
+    {
+        if (DropPreviewParent)
+        {
+            DropPreviewParent->RemoveChild(DropPreviewWidget);
+        }
+        
+        BP->WidgetTree->RemoveWidget(DropPreviewWidget);
+        DropPreviewWidget = NULL;
+    }
+    
+    if ( AddToTemplate(MyGeometry, DragDropEvent) )
+    {
+        FBlueprintEditorUtils::MarkBlueprintAsStructurallyModified(BP);
 
+		return FReply::Handled();
+    }
+    
 	return FReply::Unhandled();
 }
 
diff --git a/Engine/Source/Editor/UMGEditor/Private/SUMGDesigner.h b/Engine/Source/Editor/UMGEditor/Private/SUMGDesigner.h
--- a/Engine/Source/Editor/UMGEditor/Private/SUMGDesigner.h
+++ b/Engine/Source/Editor/UMGEditor/Private/SUMGDesigner.h
@@ -82,15 +82,22 @@ private:
 private:
 	void DrawDragHandles(const FPaintGeometry& SelectionGeometry, const FSlateRect& MyClippingRect, FSlateWindowElementList& OutDrawElements, int32 LayerId, const FWidgetStyle& InWidgetStyle) const;
 	DragHandle HitTestDragHandles(const FGeometry& AllottedGeometry, const FPointerEvent& PointerEvent) const;
-
+    
+    UWidget* AddPreview(const FGeometry& MyGeometry, const FDragDropEvent& DragDropEvent);
+    
+    bool AddToTemplate(const FGeometry& MyGeometry, const FDragDropEvent& DragDropEvent);
+    
 private:
 	TWeakPtr<FWidgetBlueprintEditor> BlueprintEditor;
 
 	TWeakPtr<SWidget> PreviewWidget;
-	UUserWidget* PreviewWidgetActor;
+	UUserWidget* PreviewWidgetObject;
 
 	FSelectedWidget CurrentSelection;
 	TWeakPtr<SWidget> SelectedWidget;
+    
+    UWidget* DropPreviewWidget;
+    UPanelWidget* DropPreviewParent;
 
 	TSharedPtr<SBorder> PreviewSurface;
 	TSharedPtr<SCanvas> ExtensionWidgetCanvas;
diff --git a/Engine/Source/Editor/UMGEditor/Private/SUMGEditorTree.cpp b/Engine/Source/Editor/UMGEditor/Private/SUMGEditorTree.cpp
--- a/Engine/Source/Editor/UMGEditor/Private/SUMGEditorTree.cpp
+++ b/Engine/Source/Editor/UMGEditor/Private/SUMGEditorTree.cpp
@@ -30,17 +30,6 @@ void SUMGEditorTree::Construct(const FArguments& InArgs, TSharedPtr<FWidgetBluep
 	[
 		SNew(SVerticalBox)
 
-		+ SVerticalBox::Slot()
-		.AutoHeight()
-		[
-			SNew(SButton)
-			.OnClicked(this, &SUMGEditorTree::CreateTestUI)
-			[
-				SNew(STextBlock)
-				.Text(NSLOCTEXT("SUMGEditorTree", "CreateTestUI", "Create Test UI"))
-			]
-		]
-
 		+ SVerticalBox::Slot()
 		.FillHeight(1.0f)
 		[
@@ -220,49 +209,6 @@ void SUMGEditorTree::WrapSelectedWidgets(UClass* WidgetClass)
 	FBlueprintEditorUtils::MarkBlueprintAsStructurallyModified(BP);
 }
 
-FReply SUMGEditorTree::CreateTestUI()
-{
-	UWidgetBlueprint* BP = GetBlueprint();
-	TArray<UWidget*>& WidgetTemplates = BP->WidgetTree->WidgetTemplates;
-
-	if ( WidgetTemplates.Num() > 0 )
-	{
-		UVerticalBox* Vertical = CastChecked<UVerticalBox>(WidgetTemplates[2]);
-
-		UButton* NewButton = BP->WidgetTree->ConstructWidget<UButton>(UButton::StaticClass());
-
-		Vertical->AddChild(NewButton, FVector2D(0,0));
-	}
-	else
-	{
-		UCanvasPanel* Canvas = BP->WidgetTree->ConstructWidget<UCanvasPanel>(UCanvasPanel::StaticClass());
-		UBorder* Border = BP->WidgetTree->ConstructWidget<UBorder>(UBorder::StaticClass());
-		UVerticalBox* Vertical = BP->WidgetTree->ConstructWidget<UVerticalBox>(UVerticalBox::StaticClass());
-		UButton* Button1 = BP->WidgetTree->ConstructWidget<UButton>(UButton::StaticClass());
-		UButton* Button2 = BP->WidgetTree->ConstructWidget<UButton>(UButton::StaticClass());
-		UButton* Button3 = BP->WidgetTree->ConstructWidget<UButton>(UButton::StaticClass());
-		UButton* Button4 = BP->WidgetTree->ConstructWidget<UButton>(UButton::StaticClass());
-
-		UTextBlock* Text1 = BP->WidgetTree->ConstructWidget<UTextBlock>(UTextBlock::StaticClass());
-		Text1->Text = FText::FromString(TEXT("Button!"));
-		Button4->SetContent(Text1);
-
-		UCanvasPanelSlot* Slot = Canvas->AddSlot(Border);
-		Slot->LayoutData.Offsets = FMargin(200, 300, 20, 50);
-
-		Border->SetContent(Vertical);
-
-		Vertical->AddSlot(Button1);
-		Vertical->AddSlot(Button2);
-		Vertical->AddSlot(Button3);
-		Vertical->AddSlot(Button4);
-	}
-
-	FBlueprintEditorUtils::MarkBlueprintAsStructurallyModified(BP);
-
-	return FReply::Handled();
-}
-
 FReply SUMGEditorTree::HandleDeleteSelected()
 {
 	DeleteSelected();
diff --git a/Engine/Source/Editor/UMGEditor/Private/SUMGEditorTree.h b/Engine/Source/Editor/UMGEditor/Private/SUMGEditorTree.h
--- a/Engine/Source/Editor/UMGEditor/Private/SUMGEditorTree.h
+++ b/Engine/Source/Editor/UMGEditor/Private/SUMGEditorTree.h
@@ -29,7 +29,6 @@ private:
 	TSharedRef< ITableRow > WidgetHierarchy_OnGenerateRow(UWidget* InItem, const TSharedRef<STableViewBase>& OwnerTable);
 	void WidgetHierarchy_OnSelectionChanged(UWidget* SelectedItem, ESelectInfo::Type SelectInfo);
 
-	FReply CreateTestUI();
 	void RefreshTree();
 
 	FReply HandleDeleteSelected();
diff --git a/Engine/Source/Editor/UMGEditor/Private/SUMGEditorWidgetTemplates.cpp b/Engine/Source/Editor/UMGEditor/Private/SUMGEditorWidgetTemplates.cpp
--- a/Engine/Source/Editor/UMGEditor/Private/SUMGEditorWidgetTemplates.cpp
+++ b/Engine/Source/Editor/UMGEditor/Private/SUMGEditorWidgetTemplates.cpp
@@ -99,11 +99,12 @@ public:
 	virtual TSharedRef<ITableRow> BuildRow(const TSharedRef<STableViewBase>& OwnerTable) override
 	{
 		return SNew(STableRow< TSharedPtr<FWidgetViewModel> >, OwnerTable)
+        .Style( FEditorStyle::Get(), "UMGEditor.ToolboxHeader" )
 			.Padding(2.0f)
 			.ShowSelection(false)
 			[
-				SNew(STextBlock)
-				.Text(GroupName)
+                SNew(STextBlock)
+                .Text(GroupName)
 			];
 	}
 
diff --git a/Engine/Source/Editor/UMGEditor/Private/UMGEditorPrivatePCH.h b/Engine/Source/Editor/UMGEditor/Private/UMGEditorPrivatePCH.h
--- a/Engine/Source/Editor/UMGEditor/Private/UMGEditorPrivatePCH.h
+++ b/Engine/Source/Editor/UMGEditor/Private/UMGEditorPrivatePCH.h
@@ -23,3 +23,6 @@
 
 #include "DetailCustomizations.h"
 #include "CanvasSlotCustomization.h"
+
+#include "WidgetTemplate.h"
+#include "WidgetTemplateDragDropOp.h"
diff --git a/Engine/Source/Editor/UMGEditor/Private/WidgetBlueprintEditor.cpp b/Engine/Source/Editor/UMGEditor/Private/WidgetBlueprintEditor.cpp
--- a/Engine/Source/Editor/UMGEditor/Private/WidgetBlueprintEditor.cpp
+++ b/Engine/Source/Editor/UMGEditor/Private/WidgetBlueprintEditor.cpp
@@ -193,7 +193,7 @@ void FWidgetBlueprintEditor::NotifyPostChange(const FPropertyChangedEvent& Prope
 {
 	//Super::NotifyPostChange(PropertyChangedEvent, PropertyThatChanged);
 	
-	if ( PropertyChangedEvent.ChangeType == EPropertyChangeType::ValueSet )
+	if ( PropertyChangedEvent.ChangeType != EPropertyChangeType::Interactive )
 	{
 		MigrateFromChain(PropertyThatChanged, false);
 	}
@@ -286,6 +286,13 @@ void FWidgetBlueprintEditor::UpdatePreview(UBlueprint* InBlueprint, bool bInForc
 
 		PreviewActor = ConstructObject<UUserWidget>(PreviewBlueprint->GeneratedClass, PreviewScene.GetWorld()->GetCurrentLevel());
 		PreviewActor->SetFlags(RF_Standalone | RF_Transactional);
+        
+        // Configure all the widgets to be set to design time.
+        PreviewActor->IsDesignTime(true);
+        for(UWidget* SubPreviewWidget : PreviewActor->Components)
+        {
+            SubPreviewWidget->IsDesignTime(true);
+        }
 
 		PreviewWidgetActorPtr = PreviewActor;
 	}
diff --git a/Engine/Source/Runtime/UMG/Private/Components/Border.cpp b/Engine/Source/Runtime/UMG/Private/Components/Border.cpp
--- a/Engine/Source/Runtime/UMG/Private/Components/Border.cpp
+++ b/Engine/Source/Runtime/UMG/Private/Components/Border.cpp
@@ -39,28 +39,25 @@ void UBorder::SyncronizeProperties()
 {
 	Super::SyncronizeProperties();
 
-	if ( MyBorder.IsValid() )
-	{
-		MyBorder->SetHAlign(HorizontalAlignment);
-		MyBorder->SetVAlign(VerticalAlignment);
-		MyBorder->SetPadding(ContentPadding);
-
-		MyBorder->SetBorderBackgroundColor(BorderColorAndOpacity);
-		MyBorder->SetColorAndOpacity(ContentColorAndOpacity);
-		MyBorder->SetForegroundColor(ForegroundColor);
-
-		MyBorder->SetContentScale(ContentScale);
-		MyBorder->SetDesiredSizeScale(DesiredSizeScale);
-
-		MyBorder->SetShowEffectWhenDisabled(bShowEffectWhenDisabled);
-
-		MyBorder->SetBorderImage(GetBorderBrush());
-
-		// SLATE_EVENT( FPointerEventHandler, OnMouseButtonDown )
-		// SLATE_EVENT( FPointerEventHandler, OnMouseButtonUp )
-		// SLATE_EVENT( FPointerEventHandler, OnMouseMove )
-		// SLATE_EVENT( FPointerEventHandler, OnMouseDoubleClick )
-	}
+	MyBorder->SetHAlign(HorizontalAlignment);
+	MyBorder->SetVAlign(VerticalAlignment);
+	MyBorder->SetPadding(ContentPadding);
+	
+	MyBorder->SetBorderBackgroundColor(BorderColorAndOpacity);
+	MyBorder->SetColorAndOpacity(ContentColorAndOpacity);
+	MyBorder->SetForegroundColor(ForegroundColor);
+	
+	MyBorder->SetContentScale(ContentScale);
+	MyBorder->SetDesiredSizeScale(DesiredSizeScale);
+	
+	MyBorder->SetShowEffectWhenDisabled(bShowEffectWhenDisabled);
+	
+	MyBorder->SetBorderImage(GetBorderBrush());
+
+	MyBorder->SetOnMouseButtonDown(BIND_UOBJECT_DELEGATE(FPointerEventHandler, HandleMouseButtonDown));
+	MyBorder->SetOnMouseButtonUp(BIND_UOBJECT_DELEGATE(FPointerEventHandler, HandleMouseButtonUp));
+	MyBorder->SetOnMouseMove(BIND_UOBJECT_DELEGATE(FPointerEventHandler, HandleMouseMove));
+	MyBorder->SetOnMouseDoubleClick(BIND_UOBJECT_DELEGATE(FPointerEventHandler, HandleMouseDoubleClick));
 }
 
 void UBorder::SetContent(UWidget* Content)
@@ -84,6 +81,46 @@ const FSlateBrush* UBorder::GetBorderBrush() const
 	return &BorderBrush->Brush;
 }
 
+FReply UBorder::HandleMouseButtonDown(const FGeometry& Geometry, const FPointerEvent& MouseEvent)
+{
+	if ( OnMouseButtonDownEvent.IsBound() )
+	{
+		return OnMouseButtonDownEvent.Execute(Geometry, MouseEvent).ToReply( MyBorder.ToSharedRef() );
+	}
+
+	return FReply::Unhandled();
+}
+
+FReply UBorder::HandleMouseButtonUp(const FGeometry& Geometry, const FPointerEvent& MouseEvent)
+{
+	if ( OnMouseButtonUpEvent.IsBound() )
+	{
+		return OnMouseButtonUpEvent.Execute(Geometry, MouseEvent).ToReply( MyBorder.ToSharedRef() );
+	}
+
+	return FReply::Unhandled();
+}
+
+FReply UBorder::HandleMouseMove(const FGeometry& Geometry, const FPointerEvent& MouseEvent)
+{
+	if ( OnMouseMoveEvent.IsBound() )
+	{
+		return OnMouseMoveEvent.Execute(Geometry, MouseEvent).ToReply( MyBorder.ToSharedRef() );
+	}
+
+	return FReply::Unhandled();
+}
+
+FReply UBorder::HandleMouseDoubleClick(const FGeometry& Geometry, const FPointerEvent& MouseEvent)
+{
+	if ( OnMouseDoubleClickEvent.IsBound() )
+	{
+		return OnMouseDoubleClickEvent.Execute(Geometry, MouseEvent).ToReply( MyBorder.ToSharedRef() );
+	}
+
+	return FReply::Unhandled();
+}
+
 /////////////////////////////////////////////////////
 
 #undef LOCTEXT_NAMESPACE
diff --git a/Engine/Source/Runtime/UMG/Private/Components/Button.cpp b/Engine/Source/Runtime/UMG/Private/Components/Button.cpp
--- a/Engine/Source/Runtime/UMG/Private/Components/Button.cpp
+++ b/Engine/Source/Runtime/UMG/Private/Components/Button.cpp
@@ -21,7 +21,8 @@ UButton::UButton(const FPostConstructInitializeProperties& PCIP)
 	DesiredSizeScale = ButtonDefaults._DesiredSizeScale.Get();
 	ContentScale = ButtonDefaults._ContentScale.Get();
 
-	ButtonColorAndOpacity = FLinearColor::White;
+	ColorAndOpacity = FLinearColor::White;
+	BackgroundColor = FLinearColor::White;
 	ForegroundColor = FLinearColor::Black;
 }
 
@@ -57,21 +58,21 @@ void UButton::SyncronizeProperties()
 		OptionalHoveredSound = HoveredSound;
 	}
 
-	if ( MyButton.IsValid() )
-	{
-		MyButton->SetButtonStyle(StylePtr);
-		MyButton->SetHAlign(HorizontalAlignment);
-		MyButton->SetVAlign(VerticalAlignment);
-		MyButton->SetContentPadding(ContentPadding);
-		MyButton->SetForegroundColor(ForegroundColor);
-		//MyButton->SetColorAndOpacity(ButtonColorAndOpacity);
-		MyButton->SetDesiredSizeScale(DesiredSizeScale);
-		MyButton->SetContentScale(ContentScale);
-		MyButton->SetBorderBackgroundColor(ButtonColorAndOpacity);
-		MyButton->SetPressedSound(OptionalPressedSound);
-		MyButton->SetHoveredSound(OptionalHoveredSound);
-		MyButton->SetOnClicked(BIND_UOBJECT_DELEGATE(FOnClicked, HandleOnClicked));
-	}
+	MyButton->SetButtonStyle(StylePtr);
+	MyButton->SetHAlign(HorizontalAlignment);
+	MyButton->SetVAlign(VerticalAlignment);
+	MyButton->SetContentPadding(ContentPadding);
+	
+	MyButton->SetForegroundColor( ForegroundColor );
+	MyButton->SetColorAndOpacity( ColorAndOpacity );
+	MyButton->SetBorderBackgroundColor( BackgroundColor );
+	
+	MyButton->SetDesiredSizeScale(DesiredSizeScale);
+	MyButton->SetContentScale(ContentScale);
+	MyButton->SetPressedSound(OptionalPressedSound);
+	MyButton->SetHoveredSound(OptionalHoveredSound);
+
+	MyButton->SetOnClicked(BIND_UOBJECT_DELEGATE(FOnClicked, HandleOnClicked));
 }
 
 void UButton::SetContent(UWidget* InContent)
@@ -103,12 +104,17 @@ FReply UButton::HandleOnClicked()
 {
 	if ( OnClickedEvent.IsBound() )
 	{
-		return OnClickedEvent.Execute().ToReply();
+		return OnClickedEvent.Execute().ToReply( MyButton.ToSharedRef() );
 	}
 	
 	return FReply::Unhandled();
 }
 
+bool UButton::IsPressed() const
+{
+	return MyButton->IsPressed();
+}
+
 
 /////////////////////////////////////////////////////
 
diff --git a/Engine/Source/Runtime/UMG/Private/Components/CanvasPanel.cpp b/Engine/Source/Runtime/UMG/Private/Components/CanvasPanel.cpp
--- a/Engine/Source/Runtime/UMG/Private/Components/CanvasPanel.cpp
+++ b/Engine/Source/Runtime/UMG/Private/Components/CanvasPanel.cpp
@@ -75,7 +75,13 @@ bool UCanvasPanel::AddChild(UWidget* Child, FVector2D Position)
 {
 	UCanvasPanelSlot* Slot = AddSlot(Child);
 	Slot->LayoutData.Offsets = FMargin(Position.X, Position.Y, 100, 25);
-
+	
+	// Add the child to the live canvas if it already exists
+	if (MyCanvas.IsValid())
+	{
+		Slot->BuildSlot(MyCanvas.ToSharedRef());
+	}
+	
 	return true;
 }
 
@@ -85,6 +91,13 @@ bool UCanvasPanel::RemoveChild(UWidget* Child)
 	if ( SlotIndex != -1 )
 	{
 		Slots.RemoveAt(SlotIndex);
+		
+		// Remove the widget from the live slot if it exists.
+		if (MyCanvas.IsValid())
+		{
+			MyCanvas->RemoveSlot(Child->GetWidget());
+		}
+		
 		return true;
 	}
 
@@ -99,6 +112,8 @@ void UCanvasPanel::ReplaceChildAt(int32 Index, UWidget* Content)
 #if WITH_EDITOR
 	Content->Slot = Slot;
 #endif
+	
+	//TODO UMG How do we handle this swap for live controls?  It doesn't matter for canvas, but matters for flow controls it will matter
 }
 
 void UCanvasPanel::InsertChildAt(int32 Index, UWidget* Content)
@@ -125,7 +140,7 @@ TSharedRef<SWidget> UCanvasPanel::RebuildWidget()
 		Slot->BuildSlot(MyCanvas.ToSharedRef());
 	}
 
-	return MyCanvas.ToSharedRef();
+	return BuildDesignTimeWidget( MyCanvas.ToSharedRef() );
 }
 
 UCanvasPanelSlot* UCanvasPanel::AddSlot(UWidget* Content)
@@ -144,8 +159,6 @@ UCanvasPanelSlot* UCanvasPanel::AddSlot(UWidget* Content)
 	return Slot;
 }
 
-
-
 bool UCanvasPanel::GetGeometryForSlot(UCanvasPanelSlot* Slot, FGeometry& ArrangedGeometry) const
 {
 	if ( Slot->Content == NULL )
@@ -186,4 +199,5 @@ void UCanvasPanel::ConnectEditorData()
 void UCanvasPanel::PostEditChangeProperty(struct FPropertyChangedEvent& PropertyChangedEvent)
 {
 }
+
 #endif
\ No newline at end of file
diff --git a/Engine/Source/Runtime/UMG/Private/Components/CanvasPanelSlot.cpp b/Engine/Source/Runtime/UMG/Private/Components/CanvasPanelSlot.cpp
--- a/Engine/Source/Runtime/UMG/Private/Components/CanvasPanelSlot.cpp
+++ b/Engine/Source/Runtime/UMG/Private/Components/CanvasPanelSlot.cpp
@@ -114,8 +114,23 @@ void UCanvasPanelSlot::PostEditChangeChainProperty(struct FPropertyChangedChainE
 	Refresh();
 
 	static FName AnchorsProperty(TEXT("Anchors"));
+	
+	FEditPropertyChain::TDoubleLinkedListNode* AnchorNode = PropertyChangedEvent.PropertyChain.GetHead()->GetNextNode();
+	if (!AnchorNode)
+	{
+		return;
+	}
+	
+	FEditPropertyChain::TDoubleLinkedListNode* LayoutDataNode = AnchorNode->GetNextNode();
+	
+	if (!LayoutDataNode)
+	{
+		return;
+	}
+	
+	UProperty* AnchorProperty = LayoutDataNode->GetValue();
 
-	if ( PropertyChangedEvent.MemberProperty && PropertyChangedEvent.MemberProperty->GetFName() == AnchorsProperty )
+	if ( AnchorProperty && AnchorProperty->GetFName() == AnchorsProperty )
 	{
 		// Ensure we have a parent canvas
 		if ( UCanvasPanel* Canvas = Cast<UCanvasPanel>(Parent) )
@@ -139,17 +154,17 @@ void UCanvasPanelSlot::PostEditChangeChainProperty(struct FPropertyChangedChainE
 				FVector2D PositionDelta = PreEditGeometry.Position - DefaultAnchorPosition;// ( Geometry.Position - DefaultAnchorOffset );
 
 				// Adjust the size to remain constant
-				if ( LayoutData.Anchors.IsStretchedHorizontal() && !PreEditLayoutData.Anchors.IsStretchedHorizontal() )
+				if ( !LayoutData.Anchors.IsStretchedHorizontal() && PreEditLayoutData.Anchors.IsStretchedHorizontal() )
 				{
 					// Adjust the position to remain constant
-					LayoutData.Offsets.Left = PositionDelta.X;
-					LayoutData.Offsets.Right = AnchorPositions.Right - ( PositionDelta.X + PreEditGeometry.Size.X );
+					LayoutData.Offsets.Left = PositionDelta.X + AlignmentOffset.X;
+					LayoutData.Offsets.Right = PreEditGeometry.Size.X;
 				}
-				else if ( !LayoutData.Anchors.IsStretchedHorizontal() && PreEditLayoutData.Anchors.IsStretchedHorizontal() )
+				else if ( LayoutData.Anchors.IsStretchedHorizontal() )
 				{
 					// Adjust the position to remain constant
-					LayoutData.Offsets.Left = PositionDelta.X + AlignmentOffset.X;
-					LayoutData.Offsets.Right = PreEditGeometry.Size.X;
+					LayoutData.Offsets.Left = PositionDelta.X;
+					LayoutData.Offsets.Right = AnchorPositions.Right - ( PositionDelta.X + PreEditGeometry.Size.X );
 				}
 				else
 				{
@@ -157,17 +172,17 @@ void UCanvasPanelSlot::PostEditChangeChainProperty(struct FPropertyChangedChainE
 					LayoutData.Offsets.Left = PositionDelta.X + AlignmentOffset.X;
 				}
 
-				if ( LayoutData.Anchors.IsStretchedVertical() && !PreEditLayoutData.Anchors.IsStretchedVertical() )
+				if ( !LayoutData.Anchors.IsStretchedVertical() && PreEditLayoutData.Anchors.IsStretchedVertical() )
 				{
 					// Adjust the position to remain constant
-					LayoutData.Offsets.Top = PositionDelta.Y;
-					LayoutData.Offsets.Bottom = AnchorPositions.Bottom - ( PositionDelta.Y + PreEditGeometry.Size.Y );
+					LayoutData.Offsets.Top = PositionDelta.Y + AlignmentOffset.Y;
+					LayoutData.Offsets.Bottom = PreEditGeometry.Size.Y;
 				}
-				else if ( !LayoutData.Anchors.IsStretchedVertical() && PreEditLayoutData.Anchors.IsStretchedVertical() )
+				else if ( LayoutData.Anchors.IsStretchedVertical() )
 				{
 					// Adjust the position to remain constant
-					LayoutData.Offsets.Top = PositionDelta.Y + AlignmentOffset.Y;
-					LayoutData.Offsets.Bottom = PreEditGeometry.Size.Y;
+					LayoutData.Offsets.Top = PositionDelta.Y;
+					LayoutData.Offsets.Bottom = AnchorPositions.Bottom - ( PositionDelta.Y + PreEditGeometry.Size.Y );
 				}
 				else
 				{
diff --git a/Engine/Source/Runtime/UMG/Private/Components/CheckBox.cpp b/Engine/Source/Runtime/UMG/Private/Components/CheckBox.cpp
--- a/Engine/Source/Runtime/UMG/Private/Components/CheckBox.cpp
+++ b/Engine/Source/Runtime/UMG/Private/Components/CheckBox.cpp
@@ -48,7 +48,7 @@ TSharedRef<SWidget> UCheckBox::RebuildWidget()
 		OptionalHoveredSound = HoveredSound;
 	}
 
-	return SNew(SCheckBox)
+	MyCheckbox = SNew(SCheckBox)
 		.Style(StylePtr)
 		.OnCheckStateChanged( BIND_UOBJECT_DELEGATE(FOnCheckStateChanged, SlateOnCheckStateChangedCallback) )
 		.IsChecked( BIND_UOBJECT_ATTRIBUTE(ESlateCheckBoxState::Type, GetCheckState) )
@@ -60,9 +60,35 @@ TSharedRef<SWidget> UCheckBox::RebuildWidget()
 		.UncheckedSoundOverride(OptionalUncheckedSound)
 		.HoveredSoundOverride(OptionalHoveredSound)
 		;
-// 		SLATE_ATTRIBUTE( bool, ReadOnly )
-// 		SLATE_ARGUMENT( bool, IsFocusable )
-// 		SLATE_EVENT( FOnGetContent, OnGetMenuContent )
+	
+	return MyCheckbox.ToSharedRef();
+}
+
+void UCheckBox::SyncronizeProperties()
+{
+	Super::SyncronizeProperties();
+	
+	MyCheckbox->SetUncheckedImage(UncheckedImage ? &UncheckedImage->Brush : nullptr);
+	MyCheckbox->SetUncheckedHoveredImage(UncheckedHoveredImage ? &UncheckedHoveredImage->Brush : nullptr);
+	MyCheckbox->SetUncheckedPressedImage(UncheckedPressedImage ? &UncheckedPressedImage->Brush : nullptr);
+	
+	MyCheckbox->SetCheckedImage(CheckedImage ? &CheckedImage->Brush : nullptr);
+	MyCheckbox->SetCheckedHoveredImage(CheckedHoveredImage ? &CheckedHoveredImage->Brush : nullptr);
+	MyCheckbox->SetCheckedPressedImage(CheckedPressedImage ? &CheckedPressedImage->Brush : nullptr);
+	
+	MyCheckbox->SetUndeterminedImage(UndeterminedImage ? &UndeterminedImage->Brush : nullptr);
+	MyCheckbox->SetUndeterminedHoveredImage(UndeterminedHoveredImage ? &UndeterminedHoveredImage->Brush : nullptr);
+	MyCheckbox->SetUndeterminedPressedImage(UndeterminedPressedImage ? &UndeterminedPressedImage->Brush : nullptr);
+}
+
+bool UCheckBox::IsPressed() const
+{
+	return MyCheckbox->IsPressed();
+}
+
+bool UCheckBox::IsChecked() const
+{
+	return MyCheckbox->IsChecked();
 }
 
 ESlateCheckBoxState::Type UCheckBox::GetCheckState() const
diff --git a/Engine/Source/Runtime/UMG/Private/Components/ContentWidget.cpp b/Engine/Source/Runtime/UMG/Private/Components/ContentWidget.cpp
--- a/Engine/Source/Runtime/UMG/Private/Components/ContentWidget.cpp
+++ b/Engine/Source/Runtime/UMG/Private/Components/ContentWidget.cpp
@@ -15,10 +15,6 @@
 UContentWidget::UContentWidget(const FPostConstructInitializeProperties& PCIP)
 	: Super(PCIP)
 {
-	//PrimaryComponentTick.bCanEverTick = true;
-
-	//bAutoActivate = true;
-	//bTickInEditor = true;
 }
 
 UWidget* UContentWidget::GetContent()
@@ -43,15 +39,20 @@ UWidget* UContentWidget::GetChildAt(int32 Index) const
 
 bool UContentWidget::AddChild(UWidget* InContent, FVector2D Position)
 {
-	Content = InContent;
-	return true;
+	if ( Content == NULL )
+	{
+		SetContent(InContent);
+		return true;
+	}
+
+	return false;
 }
 
 bool UContentWidget::RemoveChild(UWidget* Child)
 {
 	if ( Content == Child )
 	{
-		Content = NULL;
+		SetContent(NULL);
 		return true;
 	}
 
@@ -61,7 +62,7 @@ bool UContentWidget::RemoveChild(UWidget* Child)
 void UContentWidget::ReplaceChildAt(int32 Index, UWidget* Child)
 {
 	check(Index == 0);
-	Content = Child;
+	SetContent(Child);
 }
 
 /////////////////////////////////////////////////////
diff --git a/Engine/Source/Runtime/UMG/Private/Components/EditableText.cpp b/Engine/Source/Runtime/UMG/Private/Components/EditableText.cpp
--- a/Engine/Source/Runtime/UMG/Private/Components/EditableText.cpp
+++ b/Engine/Source/Runtime/UMG/Private/Components/EditableText.cpp
@@ -30,37 +30,40 @@ UEditableText::UEditableText(const FPostConstructInitializeProperties& PCIP)
 TSharedRef<SWidget> UEditableText::RebuildWidget()
 {
 	FString FontPath = FPaths::EngineContentDir() / Font.FontName.ToString();
-
+	
 	SEditableText::FArguments Defaults;
-
-	return SNew(SEditableText)
-		.Text(Text)
-		//.Style(Style ? Style->TextBlockStyle : Defaults._Style)
-		.HintText(HintText)
-		.Font(FSlateFontInfo(FontPath, Font.Size))
-		.ColorAndOpacity(ColorAndOpacity)
-		.IsReadOnly(IsReadOnly)
-		.IsPassword(IsPassword)
-		.MinDesiredWidth(MinimumDesiredWidth)
-		.IsCaretMovedWhenGainFocus(IsCaretMovedWhenGainFocus)
-		.SelectAllTextWhenFocused(SelectAllTextWhenFocused)
-		.RevertTextOnEscape(RevertTextOnEscape)
-		.ClearKeyboardFocusOnCommit(ClearKeyboardFocusOnCommit)
-		.SelectAllTextOnCommit(SelectAllTextOnCommit)
-		.OnTextChanged(BIND_UOBJECT_DELEGATE(FOnTextChanged, HandleOnTextChanged))
-		.OnTextCommitted(BIND_UOBJECT_DELEGATE(FOnTextCommitted, HandleOnTextCommitted))
-		;
+	
+	const FEditableTextStyle* StylePtr = ( Style != NULL ) ? Style->GetStyle<FEditableTextStyle>() : NULL;
+	if ( StylePtr == NULL )
+	{
+		StylePtr = Defaults._Style;
+	}
+	
+	MyEditableText = SNew(SEditableText)
+	.Style(StylePtr)
+	.Font(FSlateFontInfo(FontPath, Font.Size))
+	.MinDesiredWidth(MinimumDesiredWidth)
+	.IsCaretMovedWhenGainFocus(IsCaretMovedWhenGainFocus)
+	.SelectAllTextWhenFocused(SelectAllTextWhenFocused)
+	.RevertTextOnEscape(RevertTextOnEscape)
+	.ClearKeyboardFocusOnCommit(ClearKeyboardFocusOnCommit)
+	.SelectAllTextOnCommit(SelectAllTextOnCommit)
+	.OnTextChanged(BIND_UOBJECT_DELEGATE(FOnTextChanged, HandleOnTextChanged))
+	.OnTextCommitted(BIND_UOBJECT_DELEGATE(FOnTextCommitted, HandleOnTextCommitted))
+	;
+	
+	return MyEditableText.ToSharedRef();
 }
 
-#if WITH_EDITOR
-
-void UEditableText::PostEditChangeProperty(struct FPropertyChangedEvent& PropertyChangedEvent)
+void UEditableText::SyncronizeProperties()
 {
-	
+	MyEditableText->SetText(Text);
+	MyEditableText->SetHintText(HintText);
+	MyEditableText->SetIsReadOnly(IsReadOnly);
+	MyEditableText->SetIsPassword(IsPassword);
+	MyEditableText->SetColorAndOpacity(ColorAndOpacity);
 }
 
-#endif
-
 void UEditableText::HandleOnTextChanged(const FText& Text)
 {
 	OnTextChanged.Broadcast(Text);
diff --git a/Engine/Source/Runtime/UMG/Private/Components/HorizontalBox.cpp b/Engine/Source/Runtime/UMG/Private/Components/HorizontalBox.cpp
--- a/Engine/Source/Runtime/UMG/Private/Components/HorizontalBox.cpp
+++ b/Engine/Source/Runtime/UMG/Private/Components/HorizontalBox.cpp
@@ -38,7 +38,14 @@ int32 UHorizontalBox::GetChildIndex(UWidget* Content) const
 
 bool UHorizontalBox::AddChild(UWidget* Child, FVector2D Position)
 {
-	AddSlot(Child);
+	UHorizontalBoxSlot* Slot = AddSlot(Child);
+	
+	// Add the child to the live canvas if it already exists
+	if (MyHorizontalBox.IsValid())
+	{
+		Slot->BuildSlot(MyHorizontalBox.ToSharedRef());
+	}
+	
 	return true;
 }
 
@@ -48,6 +55,13 @@ bool UHorizontalBox::RemoveChild(UWidget* Child)
 	if ( SlotIndex != -1 )
 	{
 		Slots.RemoveAt(SlotIndex);
+		
+		// Remove the widget from the live slot if it exists.
+		if (MyHorizontalBox.IsValid())
+		{
+			MyHorizontalBox->RemoveSlot(Child->GetWidget());
+		}
+		
 		return true;
 	}
 
@@ -80,16 +94,15 @@ void UHorizontalBox::InsertChildAt(int32 Index, UWidget* Content)
 
 TSharedRef<SWidget> UHorizontalBox::RebuildWidget()
 {
-	TSharedRef<SHorizontalBox> NewWidget = SNew(SHorizontalBox);
-	MyHorizontalBox = NewWidget;
+	MyHorizontalBox = SNew(SHorizontalBox);
 
 	for ( auto Slot : Slots )
 	{
 		Slot->Parent = this;
-		Slot->BuildSlot(NewWidget);
+		Slot->BuildSlot( MyHorizontalBox.ToSharedRef() );
 	}
 
-	return NewWidget;
+	return BuildDesignTimeWidget( MyHorizontalBox.ToSharedRef() );
 }
 
 UHorizontalBoxSlot* UHorizontalBox::AddSlot(UWidget* Content)
diff --git a/Engine/Source/Runtime/UMG/Private/Components/Overlay.cpp b/Engine/Source/Runtime/UMG/Private/Components/Overlay.cpp
--- a/Engine/Source/Runtime/UMG/Private/Components/Overlay.cpp
+++ b/Engine/Source/Runtime/UMG/Private/Components/Overlay.cpp
@@ -23,7 +23,14 @@ UWidget* UOverlay::GetChildAt(int32 Index) const
 
 bool UOverlay::AddChild(UWidget* Child, FVector2D Position)
 {
-	AddSlot(Child);
+	UOverlaySlot* Slot = AddSlot(Child);
+	
+	// Add the child to the live canvas if it already exists
+	if (MyOverlay.IsValid())
+	{
+		Slot->BuildSlot(MyOverlay.ToSharedRef());
+	}
+	
 	return true;
 }
 
@@ -36,6 +43,13 @@ bool UOverlay::RemoveChild(UWidget* Child)
 		if ( Slot->Content == Child )
 		{
 			Slots.RemoveAt(SlotIndex);
+			
+			// Remove the widget from the live slot if it exists.
+			if (MyOverlay.IsValid())
+			{
+				MyOverlay->RemoveSlot(Child->GetWidget());
+			}
+			
 			return true;
 		}
 	}
@@ -55,16 +69,15 @@ void UOverlay::ReplaceChildAt(int32 Index, UWidget* Content)
 
 TSharedRef<SWidget> UOverlay::RebuildWidget()
 {
-	TSharedRef<SOverlay> NewWidget = SNew(SOverlay);
-	MyVerticalBox = NewWidget;
+	MyOverlay = SNew(SOverlay);
 
 	for ( auto Slot : Slots )
 	{
 		Slot->Parent = this;
-		Slot->BuildSlot(NewWidget);
+		Slot->BuildSlot(MyOverlay.ToSharedRef());
 	}
 
-	return NewWidget;
+	return BuildDesignTimeWidget( MyOverlay.ToSharedRef() );
 }
 
 UOverlaySlot* UOverlay::AddSlot(UWidget* Content)
diff --git a/Engine/Source/Runtime/UMG/Private/Components/PanelSlot.cpp b/Engine/Source/Runtime/UMG/Private/Components/PanelSlot.cpp
--- a/Engine/Source/Runtime/UMG/Private/Components/PanelSlot.cpp
+++ b/Engine/Source/Runtime/UMG/Private/Components/PanelSlot.cpp
@@ -29,3 +29,8 @@ bool UPanelSlot::CanMove() const
 {
 	return false;
 }
+
+bool UPanelSlot::IsDesignTime() const
+{
+	return Parent->IsDesignTime();
+}
diff --git a/Engine/Source/Runtime/UMG/Private/Components/ProcessBar.cpp b/Engine/Source/Runtime/UMG/Private/Components/ProcessBar.cpp
--- a/Engine/Source/Runtime/UMG/Private/Components/ProcessBar.cpp
+++ b/Engine/Source/Runtime/UMG/Private/Components/ProcessBar.cpp
@@ -18,15 +18,30 @@ UProgressBar::UProgressBar(const FPostConstructInitializeProperties& PCIP)
 
 TSharedRef<SWidget> UProgressBar::RebuildWidget()
 {
-	MyProgressBar = SNew(SProgressBar)
-		.BarFillType(BarFillType)
-		.Percent(Percent)
-		.FillColorAndOpacity(FillColorAndOpacity)
-		.BorderPadding(BorderPadding);
+	MyProgressBar = SNew(SProgressBar);
 
 	return MyProgressBar.ToSharedRef();
 }
 
+void UProgressBar::SyncronizeProperties()
+{
+	const FProgressBarStyle* StylePtr = ( Style != NULL ) ? Style->GetStyle<FProgressBarStyle>() : NULL;
+	if ( StylePtr )
+	{
+		MyProgressBar->SetStyle(StylePtr);
+	}
+	
+	MyProgressBar->SetPercent(bOverride_Percent ? Percent : TOptional<float>());
+	
+	MyProgressBar->SetBarFillType(BarFillType);
+	MyProgressBar->SetFillColorAndOpacity(FillColorAndOpacity);
+	MyProgressBar->SetBorderPadding(BorderPadding);
+	
+	MyProgressBar->SetBackgroundImage(BackgroundImage ? &BackgroundImage->Brush : nullptr);
+	MyProgressBar->SetFillImage(FillImage ? &FillImage->Brush : nullptr);
+	MyProgressBar->SetMarqueeImage(MarqueeImage ? &MarqueeImage->Brush : nullptr);
+}
+
 void UProgressBar::SetPercent(float InPercent)
 {
 	return MyProgressBar->SetPercent(InPercent);
diff --git a/Engine/Source/Runtime/UMG/Private/Components/ScrollBox.cpp b/Engine/Source/Runtime/UMG/Private/Components/ScrollBox.cpp
--- a/Engine/Source/Runtime/UMG/Private/Components/ScrollBox.cpp
+++ b/Engine/Source/Runtime/UMG/Private/Components/ScrollBox.cpp
@@ -40,7 +40,14 @@ int32 UScrollBox::GetChildIndex(UWidget* Content) const
 
 bool UScrollBox::AddChild(UWidget* Child, FVector2D Position)
 {
-	AddSlot(Child);
+	UScrollBoxSlot* Slot = AddSlot(Child);
+	
+	// Add the child to the live canvas if it already exists
+	if (MyScrollBox.IsValid())
+	{
+		Slot->BuildSlot(MyScrollBox.ToSharedRef());
+	}
+	
 	return true;
 }
 
@@ -50,6 +57,13 @@ bool UScrollBox::RemoveChild(UWidget* Child)
 	if ( SlotIndex != -1 )
 	{
 		Slots.RemoveAt(SlotIndex);
+		
+		// Remove the widget from the live slot if it exists.
+		if (MyScrollBox.IsValid())
+		{
+			MyScrollBox->RemoveSlot(Child->GetWidget());
+		}
+		
 		return true;
 	}
 
@@ -82,22 +96,15 @@ void UScrollBox::InsertChildAt(int32 Index, UWidget* Content)
 
 TSharedRef<SWidget> UScrollBox::RebuildWidget()
 {
-	TSharedRef<SScrollBox> NewCanvas = SNew(SScrollBox);
-	MyScrollBox = NewCanvas;
+	MyScrollBox = SNew(SScrollBox);
 
-	TSharedPtr<SScrollBox> Canvas = MyScrollBox.Pin();
-	if ( Canvas.IsValid() )
+	for ( auto Slot : Slots )
 	{
-		Canvas->ClearChildren();
-
-		for ( auto Slot : Slots )
-		{
-			Slot->Parent = this;
-			Slot->BuildSlot(NewCanvas);
-		}
+		Slot->Parent = this;
+		Slot->BuildSlot(MyScrollBox.ToSharedRef());
 	}
-
-	return NewCanvas;
+	
+	return BuildDesignTimeWidget( MyScrollBox.ToSharedRef() );
 }
 
 UScrollBoxSlot* UScrollBox::AddSlot(UWidget* Content)
diff --git a/Engine/Source/Runtime/UMG/Private/Components/Slider.cpp b/Engine/Source/Runtime/UMG/Private/Components/Slider.cpp
--- a/Engine/Source/Runtime/UMG/Private/Components/Slider.cpp
+++ b/Engine/Source/Runtime/UMG/Private/Components/Slider.cpp
@@ -17,15 +17,7 @@ USlider::USlider(const FPostConstructInitializeProperties& PCIP)
 
 TSharedRef<SWidget> USlider::RebuildWidget()
 {
-	// TODO UMG Bind these
-	//FOnMouseCaptureBeginEvent OnMouseCaptureBegin;
-	//FOnMouseCaptureEndEvent OnMouseCaptureEnd;
-
 	MySlider = SNew(SSlider)
-		.Orientation(Orientation)
-		.SliderBarColor(SliderBarColor)
-		.SliderHandleColor(SliderHandleColor)
-		.Value(Value)
 		.OnMouseCaptureBegin(BIND_UOBJECT_DELEGATE(FSimpleDelegate, HandleOnMouseCaptureBegin))
 		.OnMouseCaptureEnd(BIND_UOBJECT_DELEGATE(FSimpleDelegate, HandleOnMouseCaptureEnd))
 		.OnValueChanged(BIND_UOBJECT_DELEGATE(FOnFloatValueChanged, HandleOnValueChanged));
@@ -33,6 +25,16 @@ TSharedRef<SWidget> USlider::RebuildWidget()
 	return MySlider.ToSharedRef();
 }
 
+void USlider::SyncronizeProperties()
+{
+	Super::SyncronizeProperties();
+	
+	MySlider->SetOrientation(Orientation);
+	MySlider->SetSliderBarColor(SliderBarColor);
+	MySlider->SetSliderHandleColor(SliderHandleColor);
+	MySlider->SetValue(Value);
+}
+
 void USlider::HandleOnValueChanged(float InValue)
 {
 	OnValueChanged.Broadcast(InValue);
@@ -55,6 +57,7 @@ float USlider::GetValue()
 
 void USlider::SetValue(float InValue)
 {
+	Value = InValue;
 	return MySlider->SetValue(InValue);
 }
 
diff --git a/Engine/Source/Runtime/UMG/Private/Components/Spacer.cpp b/Engine/Source/Runtime/UMG/Private/Components/Spacer.cpp
--- a/Engine/Source/Runtime/UMG/Private/Components/Spacer.cpp
+++ b/Engine/Source/Runtime/UMG/Private/Components/Spacer.cpp
@@ -26,12 +26,21 @@ void USpacer::SetSize(FVector2D InSize)
 
 TSharedRef<SWidget> USpacer::RebuildWidget()
 {
-	MySpacer = SNew(SSpacer)
-		.Size(Size);
+	MySpacer = SNew(SSpacer);
 
+	//TODO UMG COnsider using a design time wrapper for spacer to show expandy arrows or some other
+	// indicator that there's a widget at work here.
+	
 	return MySpacer.ToSharedRef();
 }
 
+void USpacer::SyncronizeProperties()
+{
+	Super::SyncronizeProperties();
+	
+	MySpacer->SetSize(Size);
+}
+
 /////////////////////////////////////////////////////
 
 #undef LOCTEXT_NAMESPACE
diff --git a/Engine/Source/Runtime/UMG/Private/Components/UniformGridPanel.cpp b/Engine/Source/Runtime/UMG/Private/Components/UniformGridPanel.cpp
--- a/Engine/Source/Runtime/UMG/Private/Components/UniformGridPanel.cpp
+++ b/Engine/Source/Runtime/UMG/Private/Components/UniformGridPanel.cpp
@@ -23,7 +23,14 @@ UWidget* UUniformGridPanel::GetChildAt(int32 Index) const
 
 bool UUniformGridPanel::AddChild(UWidget* Child, FVector2D Position)
 {
-	AddSlot(Child);
+	UUniformGridSlot* Slot = AddSlot(Child);
+	
+	// Add the child to the live panel if it already exists
+	if (MyUniformGridPanel.IsValid())
+	{
+		Slot->BuildSlot(MyUniformGridPanel.ToSharedRef());
+	}
+	
 	return true;
 }
 
@@ -36,6 +43,13 @@ bool UUniformGridPanel::RemoveChild(UWidget* Child)
 		if ( Slot->Content == Child )
 		{
 			Slots.RemoveAt(SlotIndex);
+			
+			// Remove the widget from the live slot if it exists.
+			if (MyUniformGridPanel.IsValid())
+			{
+				MyUniformGridPanel->RemoveSlot(Child->GetWidget());
+			}
+			
 			return true;
 		}
 	}
@@ -55,21 +69,19 @@ void UUniformGridPanel::ReplaceChildAt(int32 Index, UWidget* Content)
 
 TSharedRef<SWidget> UUniformGridPanel::RebuildWidget()
 {
-	TSharedRef<SUniformGridPanel> NewPanel =
+	MyUniformGridPanel =
 		SNew(SUniformGridPanel)
 		.SlotPadding(SlotPadding)
 		.MinDesiredSlotWidth(MinDesiredSlotWidth)
 		.MinDesiredSlotHeight(MinDesiredSlotHeight);
-	
-	MyUniformGridPanel = NewPanel;
 
 	for ( auto Slot : Slots )
 	{
 		Slot->Parent = this;
-		Slot->BuildSlot(NewPanel);
+		Slot->BuildSlot(MyUniformGridPanel.ToSharedRef());
 	}
 
-	return NewPanel;
+	return BuildDesignTimeWidget( MyUniformGridPanel.ToSharedRef() );
 }
 
 UUniformGridSlot* UUniformGridPanel::AddSlot(UWidget* Content)
diff --git a/Engine/Source/Runtime/UMG/Private/Components/VerticalBox.cpp b/Engine/Source/Runtime/UMG/Private/Components/VerticalBox.cpp
--- a/Engine/Source/Runtime/UMG/Private/Components/VerticalBox.cpp
+++ b/Engine/Source/Runtime/UMG/Private/Components/VerticalBox.cpp
@@ -38,7 +38,14 @@ int32 UVerticalBox::GetChildIndex(UWidget* Content) const
 
 bool UVerticalBox::AddChild(UWidget* Child, FVector2D Position)
 {
-	AddSlot(Child);
+	UVerticalBoxSlot* Slot = AddSlot(Child);
+	
+	// Add the child to the live canvas if it already exists
+	if (MyVerticalBox.IsValid())
+	{
+		Slot->BuildSlot(MyVerticalBox.ToSharedRef());
+	}
+	
 	return true;
 }
 
@@ -48,6 +55,13 @@ bool UVerticalBox::RemoveChild(UWidget* Child)
 	if ( SlotIndex != -1 )
 	{
 		Slots.RemoveAt(SlotIndex);
+		
+		// Remove the widget from the live slot if it exists.
+		if (MyVerticalBox.IsValid())
+		{
+			MyVerticalBox->RemoveSlot(Child->GetWidget());
+		}
+		
 		return true;
 	}
 
@@ -80,16 +94,15 @@ void UVerticalBox::InsertChildAt(int32 Index, UWidget* Content)
 
 TSharedRef<SWidget> UVerticalBox::RebuildWidget()
 {
-	TSharedRef<SVerticalBox> NewWidget = SNew(SVerticalBox);
-	MyVerticalBox = NewWidget;
+	MyVerticalBox = SNew(SVerticalBox);
 
 	for ( auto Slot : Slots )
 	{
 		Slot->Parent = this;
-		Slot->BuildSlot(NewWidget);
+		Slot->BuildSlot(MyVerticalBox.ToSharedRef());
 	}
 
-	return NewWidget;
+	return BuildDesignTimeWidget( MyVerticalBox.ToSharedRef() );
 }
 
 UVerticalBoxSlot* UVerticalBox::AddSlot(UWidget* Content)
diff --git a/Engine/Source/Runtime/UMG/Private/Components/Widget.cpp b/Engine/Source/Runtime/UMG/Private/Components/Widget.cpp
--- a/Engine/Source/Runtime/UMG/Private/Components/Widget.cpp
+++ b/Engine/Source/Runtime/UMG/Private/Components/Widget.cpp
@@ -10,6 +10,7 @@ UWidget::UWidget(const FPostConstructInitializeProperties& PCIP)
 {
 	bIsEnabled = true;
 	bIsVariable = true;
+	bDesignTime = false;
 	Visiblity = ESlateVisibility::Visible;
 }
 
@@ -62,17 +63,71 @@ bool UWidget::IsHovered() const
 	return MyWidget->IsHovered();
 }
 
-TSharedRef<SWidget> UWidget::GetWidget()
+TSharedRef<SWidget> UWidget::GetWidget() const
 {
 	if ( !MyWidget.IsValid() )
 	{
-		MyWidget = RebuildWidget();
-		SyncronizeProperties();
+		// We lie a bit about this being a const function.  If this is the first call it's not really const
+		// and we need to construct and cache the widget for the first run.  But instead of forcing everyone
+		// downstream to make RebuildWidget const and force every implementation to make things mutable, we
+		// just blow away the const here.
+		UWidget* MutableThis = const_cast<UWidget*>(this);
+		
+		MyWidget = MutableThis->RebuildWidget();
+		MutableThis->SyncronizeProperties();
 	}
 
 	return MyWidget.ToSharedRef();
 }
 
+TSharedRef<SWidget> UWidget::BuildDesignTimeWidget(TSharedRef<SWidget> WrapWidget)
+{
+	if (IsDesignTime())
+	{
+		return SNew(SOverlay)
+		
+		+ SOverlay::Slot()
+		.HAlign(HAlign_Fill)
+		.VAlign(VAlign_Fill)
+		[
+			WrapWidget
+		]
+		
+		+ SOverlay::Slot()
+		.HAlign(HAlign_Fill)
+		.VAlign(VAlign_Fill)
+		[
+			SNew(SBorder)
+			.Visibility( EVisibility::HitTestInvisible )
+			.BorderImage( FCoreStyle::Get().GetBrush("FocusRectangle") )
+		];
+	}
+	else
+	{
+		return WrapWidget;
+	}
+}
+
+#if WITH_EDITOR
+
+TSharedRef<SWidget> UWidget::GetToolboxPreviewWidget() const
+{
+	return SNew(SImage);
+	//.Image( FEditorStyle::GetBrush("UMGEditor.ToolboxPreviewWidget") );
+}
+
+void UWidget::PostEditChangeProperty(FPropertyChangedEvent& PropertyChangedEvent)
+{
+	Super::PostEditChangeProperty(PropertyChangedEvent);
+
+	if ( MyWidget.IsValid() )
+	{
+		SyncronizeProperties();
+	}
+}
+
+#endif
+
 TSharedRef<SWidget> UWidget::RebuildWidget()
 {
 	ensureMsg(false, TEXT("You must implement RebuildWidget() in your child class"));
@@ -81,19 +136,26 @@ TSharedRef<SWidget> UWidget::RebuildWidget()
 
 void UWidget::SyncronizeProperties()
 {
-	if ( MyWidget.IsValid() )
-	{
-		MyWidget->SetEnabled(OPTIONAL_BINDING(bool, bIsEnabled));
-		MyWidget->SetVisibility(OPTIONAL_BINDING_CONVERT(ESlateVisibility::Type, Visiblity, EVisibility, ConvertVisibility));
-		//MyWidget->SetCursor(OPTIONAL_BINDING(EMouseCursor)
+	MyWidget->SetEnabled(OPTIONAL_BINDING(bool, bIsEnabled));
+	MyWidget->SetVisibility(OPTIONAL_BINDING_CONVERT(ESlateVisibility::Type, Visiblity, EVisibility, ConvertVisibility));
+	//MyWidget->SetCursor(OPTIONAL_BINDING(EMouseCursor)
 
-		if ( !ToolTipText.IsEmpty() )
-		{
+	if ( !ToolTipText.IsEmpty() )
+	{
 			MyWidget->SetToolTipText(OPTIONAL_BINDING(FText, ToolTipText));
-		}
 	}
 }
 
+bool UWidget::IsDesignTime() const
+{
+	return bDesignTime;
+}
+
+void UWidget::IsDesignTime(bool bInDesignTime)
+{
+	bDesignTime = bInDesignTime;
+}
+
 EVisibility UWidget::ConvertSerializedVisibilityToRuntime(TEnumAsByte<ESlateVisibility::Type> Input)
 {
 	switch ( Input )
diff --git a/Engine/Source/Runtime/UMG/Private/Slate/SObjectWidget.cpp b/Engine/Source/Runtime/UMG/Private/Slate/SObjectWidget.cpp
--- a/Engine/Source/Runtime/UMG/Private/Slate/SObjectWidget.cpp
+++ b/Engine/Source/Runtime/UMG/Private/Slate/SObjectWidget.cpp
@@ -21,16 +21,21 @@ int32 SObjectWidget::OnPaint(const FGeometry& AllottedGeometry, const FSlateRect
 {
 	int32 MaxLayer = SCompoundWidget::OnPaint(AllottedGeometry, MyClippingRect, OutDrawElements, LayerId, InWidgetStyle, bParentEnabled);
 
-	//const FGeometry& AllottedGeometry, const FSlateRect& MyClippingRect, FSlateWindowElementList& OutDrawElements, int32 LayerId, const FWidgetStyle& InWidgetStyle, bool bParentEnabled
-	FPaintContext Context(AllottedGeometry, MyClippingRect, OutDrawElements, LayerId, InWidgetStyle, bParentEnabled);
-	WidgetObject->OnPaint(Context);
-
-	return FMath::Max(MaxLayer, Context.MaxLayer);
+	if (WidgetObject)
+	{
+		//const FGeometry& AllottedGeometry, const FSlateRect& MyClippingRect, FSlateWindowElementList& OutDrawElements, int32 LayerId, const FWidgetStyle& InWidgetStyle, bool bParentEnabled
+		FPaintContext Context(AllottedGeometry, MyClippingRect, OutDrawElements, LayerId, InWidgetStyle, bParentEnabled);
+		WidgetObject->OnPaint(Context);
+		
+		return FMath::Max(MaxLayer, Context.MaxLayer);
+	}
+	
+	return MaxLayer;
 }
 
 FReply SObjectWidget::OnKeyboardFocusReceived(const FGeometry& MyGeometry, const FKeyboardFocusEvent& InKeyboardFocusEvent)
 {
-	return WidgetObject->OnKeyboardFocusReceived(MyGeometry, InKeyboardFocusEvent).ToReply();
+	return WidgetObject->OnKeyboardFocusReceived(MyGeometry, InKeyboardFocusEvent).ToReply(WidgetObject->GetWidget());
 }
 
 void SObjectWidget::OnKeyboardFocusLost(const FKeyboardFocusEvent& InKeyboardFocusEvent)
@@ -45,42 +50,49 @@ void SObjectWidget::OnKeyboardFocusChanging(const FWeakWidgetPath& PreviousFocus
 
 FReply SObjectWidget::OnKeyChar(const FGeometry& MyGeometry, const FCharacterEvent& InCharacterEvent)
 {
-	return WidgetObject->OnKeyChar(MyGeometry, InCharacterEvent).ToReply();
+	return WidgetObject->OnKeyChar(MyGeometry, InCharacterEvent).ToReply(WidgetObject->GetWidget());
 }
 
 FReply SObjectWidget::OnPreviewKeyDown(const FGeometry& MyGeometry, const FKeyboardEvent& InKeyboardEvent)
 {
-	return WidgetObject->OnPreviewKeyDown(MyGeometry, InKeyboardEvent).ToReply();
+	return WidgetObject->OnPreviewKeyDown(MyGeometry, InKeyboardEvent).ToReply(WidgetObject->GetWidget());
 }
 
 FReply SObjectWidget::OnKeyDown(const FGeometry& MyGeometry, const FKeyboardEvent& InKeyboardEvent)
 {
-	return WidgetObject->OnKeyDown(MyGeometry, InKeyboardEvent).ToReply();
+	return WidgetObject->OnKeyDown(MyGeometry, InKeyboardEvent).ToReply(WidgetObject->GetWidget());
 }
 
 FReply SObjectWidget::OnKeyUp(const FGeometry& MyGeometry, const FKeyboardEvent& InKeyboardEvent)
 {
-	return WidgetObject->OnKeyUp(MyGeometry, InKeyboardEvent).ToReply();
+	return WidgetObject->OnKeyUp(MyGeometry, InKeyboardEvent).ToReply(WidgetObject->GetWidget());
 }
 
 FReply SObjectWidget::OnMouseButtonDown(const FGeometry& MyGeometry, const FPointerEvent& MouseEvent)
 {
-	return WidgetObject->OnMouseButtonDown(MyGeometry, MouseEvent).ToReply();
+	FReply Reply = WidgetObject->OnMouseButtonDown(MyGeometry, MouseEvent).ToReply(WidgetObject->GetWidget());
+	//TODO UMG Figure out how to let the user more easily manage when drags start.
+	if ( MouseEvent.GetEffectingButton() == EKeys::LeftMouseButton )
+	{
+		return Reply.DetectDrag(WidgetObject->GetWidget(), EKeys::LeftMouseButton);
+	}
+
+	return Reply;
 }
 
 FReply SObjectWidget::OnPreviewMouseButtonDown(const FGeometry& MyGeometry, const FPointerEvent& MouseEvent)
 {
-	return WidgetObject->OnPreviewMouseButtonDown(MyGeometry, MouseEvent).ToReply();
+	return WidgetObject->OnPreviewMouseButtonDown(MyGeometry, MouseEvent).ToReply(WidgetObject->GetWidget());
 }
 
 FReply SObjectWidget::OnMouseButtonUp(const FGeometry& MyGeometry, const FPointerEvent& MouseEvent)
 {
-	return WidgetObject->OnMouseButtonUp(MyGeometry, MouseEvent).ToReply();
+	return WidgetObject->OnMouseButtonUp(MyGeometry, MouseEvent).ToReply(WidgetObject->GetWidget());
 }
 
 FReply SObjectWidget::OnMouseMove(const FGeometry& MyGeometry, const FPointerEvent& MouseEvent)
 {
-	return WidgetObject->OnMouseMove(MyGeometry, MouseEvent).ToReply();
+	return WidgetObject->OnMouseMove(MyGeometry, MouseEvent).ToReply(WidgetObject->GetWidget());
 }
 
 void SObjectWidget::OnMouseEnter(const FGeometry& MyGeometry, const FPointerEvent& MouseEvent)
@@ -95,22 +107,32 @@ void SObjectWidget::OnMouseLeave(const FPointerEvent& MouseEvent)
 
 FReply SObjectWidget::OnMouseWheel(const FGeometry& MyGeometry, const FPointerEvent& MouseEvent)
 {
-	return WidgetObject->OnMouseWheel(MyGeometry, MouseEvent).ToReply();
+	return WidgetObject->OnMouseWheel(MyGeometry, MouseEvent).ToReply(WidgetObject->GetWidget());
 }
 
 FCursorReply SObjectWidget::OnCursorQuery(const FGeometry& MyGeometry, const FPointerEvent& CursorEvent) const
 {
-	// TODO UMG
+	// TODO UMG  Allow conditional overriding of the cursor logic.
 	return FCursorReply::Unhandled();
 }
 
 FReply SObjectWidget::OnMouseButtonDoubleClick(const FGeometry& MyGeometry, const FPointerEvent& MouseEvent)
 {
-	return WidgetObject->OnMouseButtonDoubleClick(MyGeometry, MouseEvent).ToReply();
+	return WidgetObject->OnMouseButtonDoubleClick(MyGeometry, MouseEvent).ToReply(WidgetObject->GetWidget());
 }
 
 FReply SObjectWidget::OnDragDetected(const FGeometry& MyGeometry, const FPointerEvent& MouseEvent)
 {
+	WidgetObject->OnDragDetected(MyGeometry, MouseEvent);
+
+	//TODO UMG Need to support drag drop, instead of returning the standard FSReply, maybe something else?
+	// Maybe like drag data, optional uobject + string data.
+
+	//FUMGDragDropOp
+	//.BeginDragDrop(FTrackNodeDragDropOp::New(SharedThis(this), ScreenCursorPos, ScreenNodePosition));
+
+	//FTrackNodeDragDropOp
+
 	// TODO UMG
 	return FReply::Unhandled();
 }
@@ -139,40 +161,40 @@ FReply SObjectWidget::OnDrop(const FGeometry& MyGeometry, const FDragDropEvent&
 
 FReply SObjectWidget::OnControllerButtonPressed(const FGeometry& MyGeometry, const FControllerEvent& ControllerEvent)
 {
-	return WidgetObject->OnControllerButtonPressed(MyGeometry, ControllerEvent).ToReply();
+	return WidgetObject->OnControllerButtonPressed(MyGeometry, ControllerEvent).ToReply(WidgetObject->GetWidget());
 }
 
 FReply SObjectWidget::OnControllerButtonReleased(const FGeometry& MyGeometry, const FControllerEvent& ControllerEvent)
 {
-	return WidgetObject->OnControllerButtonReleased(MyGeometry, ControllerEvent).ToReply();
+	return WidgetObject->OnControllerButtonReleased(MyGeometry, ControllerEvent).ToReply(WidgetObject->GetWidget());
 }
 
 FReply SObjectWidget::OnControllerAnalogValueChanged(const FGeometry& MyGeometry, const FControllerEvent& ControllerEvent)
 {
-	return WidgetObject->OnControllerAnalogValueChanged(MyGeometry, ControllerEvent).ToReply();
+	return WidgetObject->OnControllerAnalogValueChanged(MyGeometry, ControllerEvent).ToReply(WidgetObject->GetWidget());
 }
 
 FReply SObjectWidget::OnTouchGesture(const FGeometry& MyGeometry, const FPointerEvent& GestureEvent)
 {
-	return WidgetObject->OnTouchGesture(MyGeometry, GestureEvent).ToReply();
+	return WidgetObject->OnTouchGesture(MyGeometry, GestureEvent).ToReply(WidgetObject->GetWidget());
 }
 
 FReply SObjectWidget::OnTouchStarted(const FGeometry& MyGeometry, const FPointerEvent& InTouchEvent)
 {
-	return WidgetObject->OnTouchStarted(MyGeometry, InTouchEvent).ToReply();
+	return WidgetObject->OnTouchStarted(MyGeometry, InTouchEvent).ToReply(WidgetObject->GetWidget());
 }
 
 FReply SObjectWidget::OnTouchMoved(const FGeometry& MyGeometry, const FPointerEvent& InTouchEvent)
 {
-	return WidgetObject->OnTouchMoved(MyGeometry, InTouchEvent).ToReply();
+	return WidgetObject->OnTouchMoved(MyGeometry, InTouchEvent).ToReply(WidgetObject->GetWidget());
 }
 
 FReply SObjectWidget::OnTouchEnded(const FGeometry& MyGeometry, const FPointerEvent& InTouchEvent)
 {
-	return WidgetObject->OnTouchEnded(MyGeometry, InTouchEvent).ToReply();
+	return WidgetObject->OnTouchEnded(MyGeometry, InTouchEvent).ToReply(WidgetObject->GetWidget());
 }
 
 FReply SObjectWidget::OnMotionDetected(const FGeometry& MyGeometry, const FMotionEvent& InMotionEvent)
 {
-	return WidgetObject->OnMotionDetected(MyGeometry, InMotionEvent).ToReply();
+	return WidgetObject->OnMotionDetected(MyGeometry, InMotionEvent).ToReply(WidgetObject->GetWidget());
 }
diff --git a/Engine/Source/Runtime/UMG/Private/UserWidget.cpp b/Engine/Source/Runtime/UMG/Private/UserWidget.cpp
--- a/Engine/Source/Runtime/UMG/Private/UserWidget.cpp
+++ b/Engine/Source/Runtime/UMG/Private/UserWidget.cpp
@@ -4,6 +4,29 @@
 
 #include "SceneViewport.h"
 
+TSharedRef<FUMGDragDropOp> FUMGDragDropOp::New()
+{
+	TSharedRef<FUMGDragDropOp> Operation = MakeShareable(new FUMGDragDropOp);
+	Operation->Construct();
+
+	return Operation;
+}
+
+void FUMGDragDropOp::OnDrop(bool bDropWasHandled, const FPointerEvent& MouseEvent)
+{
+	FDragDropOperation::OnDrop(bDropWasHandled, MouseEvent);
+}
+
+void FUMGDragDropOp::OnDragged(const class FDragDropEvent& DragDropEvent)
+{
+	FDragDropOperation::OnDragged(DragDropEvent);
+}
+
+TSharedPtr<SWidget> FUMGDragDropOp::GetDefaultDecorator() const
+{
+	return DecoratorWidget;
+}
+
 class SViewportWidgetHost : public SCompoundWidget
 {
 	SLATE_BEGIN_ARGS(SViewportWidgetHost)
@@ -132,7 +155,7 @@ void UUserWidget::PostInitProperties()
 		BGClass->InitializeWidget(this);
 	}
 
-	RebuildWidget();
+	//RebuildWidget();
 }
 
 UWorld* UUserWidget::GetWorld() const
@@ -154,6 +177,8 @@ UWorld* UUserWidget::GetWorld() const
 
 UWidget* UUserWidget::GetWidgetHandle(TSharedRef<SWidget> InWidget)
 {
+	GetWidget();
+	
 	return WidgetToComponent.FindRef(InWidget);
 }
 
@@ -199,6 +224,8 @@ TSharedRef<SWidget> UUserWidget::RebuildWidget()
 
 TSharedPtr<SWidget> UUserWidget::GetWidgetFromName(const FString& Name) const
 {
+	GetWidget();
+	
 	for ( auto& Entry : WidgetToComponent )
 	{
 		if ( Entry.Value->GetName().Equals(Name, ESearchCase::IgnoreCase) )
@@ -212,6 +239,8 @@ TSharedPtr<SWidget> UUserWidget::GetWidgetFromName(const FString& Name) const
 
 UWidget* UUserWidget::GetHandleFromName(const FString& Name) const
 {
+	GetWidget();
+	
 	for ( auto& Entry : WidgetToComponent )
 	{
 		if ( Entry.Value->GetName().Equals(Name, ESearchCase::IgnoreCase) )
@@ -233,6 +262,8 @@ TSharedRef<SWidget> UUserWidget::MakeWidget()
 
 TSharedRef<SWidget> UUserWidget::MakeFullScreenWidget()
 {
+	GetWidget();
+	
 	if ( bAbsoluteLayout )
 	{
 		return SNew(SCanvas)
@@ -266,6 +297,8 @@ TSharedRef<SWidget> UUserWidget::MakeFullScreenWidget()
 
 UWidget* UUserWidget::GetRootWidgetComponent()
 {
+	GetWidget();
+	
 	if ( Components.Num() > 0 )
 	{
 		return Components[0];
@@ -276,6 +309,8 @@ UWidget* UUserWidget::GetRootWidgetComponent()
 
 void UUserWidget::Show()
 {
+	GetWidget();
+	
 	if ( !FullScreenWidget.IsValid() )
 	{
 		TSharedRef<SWidget> RootWidget = MakeFullScreenWidget();
@@ -309,6 +344,8 @@ void UUserWidget::Show()
 
 void UUserWidget::Hide()
 {
+	GetWidget();
+	
 	if ( FullScreenWidget.IsValid() )
 	{
 		TSharedPtr<SWidget> RootWidget = FullScreenWidget.Pin();
@@ -337,11 +374,15 @@ void UUserWidget::Hide()
 
 bool UUserWidget::GetIsVisible()
 {
+	GetWidget();
+	
 	return FullScreenWidget.IsValid();
 }
 
 TEnumAsByte<ESlateVisibility::Type> UUserWidget::GetVisiblity()
 {
+	GetWidget();
+	
 	if ( FullScreenWidget.IsValid() )
 	{
 		TSharedPtr<SWidget> RootWidget = FullScreenWidget.Pin();
diff --git a/Engine/Source/Runtime/UMG/Private/WidgetTree.cpp b/Engine/Source/Runtime/UMG/Private/WidgetTree.cpp
--- a/Engine/Source/Runtime/UMG/Private/WidgetTree.cpp
+++ b/Engine/Source/Runtime/UMG/Private/WidgetTree.cpp
@@ -54,14 +54,12 @@ UPanelWidget* UWidgetTree::FindWidgetParent(UWidget* Widget, int32& OutChildInde
 	return NULL;
 }
 
-bool UWidgetTree::RemoveWidget(UWidget* InRemovedWidget)
+bool UWidgetTree::RemoveWidgetRecursive(UWidget* InRemovedWidget)
 {
 	UWidget* Parent = NULL;
-
-	bool bRemoved = false;
-
+	
 	//TODO UMG Make the Widget Tree actually a tree, instead of a list, it makes things like removal difficult.
-
+	
 	if ( UPanelWidget* InNonLeafRemovedWidget = Cast<UPanelWidget>(InRemovedWidget) )
 	{
 		while ( InNonLeafRemovedWidget->GetChildrenCount() > 0 )
@@ -69,7 +67,7 @@ bool UWidgetTree::RemoveWidget(UWidget* InRemovedWidget)
 			RemoveWidget(InNonLeafRemovedWidget->GetChildAt(0));
 		}
 	}
-
+	
 	for ( UWidget* Template : WidgetTemplates )
 	{
 		UPanelWidget* NonLeafTemplate = Cast<UPanelWidget>(Template);
@@ -77,16 +75,19 @@ bool UWidgetTree::RemoveWidget(UWidget* InRemovedWidget)
 		{
 			if ( NonLeafTemplate->RemoveChild(InRemovedWidget) )
 			{
-				bRemoved = true;
-				break;
+				return true;
 			}
 		}
 	}
+	
+	return false;
+}
 
-	if ( bRemoved )
-	{
-		WidgetTemplates.Remove(InRemovedWidget);
-	}
+bool UWidgetTree::RemoveWidget(UWidget* InRemovedWidget)
+{
+	bool bRemoved = RemoveWidgetRecursive(InRemovedWidget);
 
-	return bRemoved;
-}
\ No newline at end of file
+	int32 IndexRemoved = WidgetTemplates.Remove(InRemovedWidget);
+	
+	return bRemoved || IndexRemoved != -1;
+}
diff --git a/Engine/Source/Runtime/UMG/Public/Blueprint/UserWidget.h b/Engine/Source/Runtime/UMG/Public/Blueprint/UserWidget.h
--- a/Engine/Source/Runtime/UMG/Public/Blueprint/UserWidget.h
+++ b/Engine/Source/Runtime/UMG/Public/Blueprint/UserWidget.h
@@ -11,6 +11,23 @@ static FSlateRect NullRect;
 static FSlateWindowElementList NullElementList;
 static FWidgetStyle NullStyle;
 
+class FUMGDragDropOp : public FDragDropOperation
+{
+public:
+	DRAG_DROP_OPERATOR_TYPE(FUMGDragDropOp, FDragDropOperation)
+
+	static TSharedRef<FUMGDragDropOp> New();
+
+	virtual void OnDrop(bool bDropWasHandled, const FPointerEvent& MouseEvent) override;
+
+	virtual void OnDragged(const class FDragDropEvent& DragDropEvent) override;
+
+	virtual TSharedPtr<SWidget> GetDefaultDecorator() const override;
+
+private:
+	TSharedPtr<SWidget> DecoratorWidget;
+};
+
 /**
  * The state passed into OnPaint that we can expose as a single painting structure to blueprints to
  * allow script code to override OnPaint behavior.
@@ -200,8 +217,8 @@ class UMG_API UUserWidget : public UWidget
 	UFUNCTION(BlueprintImplementableEvent, Category="User Interface")
 	FSReply OnMouseButtonDoubleClick(FGeometry InMyGeometry, const FPointerEvent& InMouseEvent);
 
-	//UFUNCTION(BlueprintImplementableEvent, Category="User Interface")
-	//FSReply OnDragDetected(FGeometry MyGeometry, const FPointerEvent& MouseEvent);
+	UFUNCTION(BlueprintImplementableEvent, Category="User Interface")
+	FSReply OnDragDetected(FGeometry MyGeometry, const FPointerEvent& MouseEvent);
 	//UFUNCTION(BlueprintImplementableEvent, Category="User Interface")
 	//void OnDragEnter(FGeometry MyGeometry, FDragDropEvent DragDropEvent);
 	//UFUNCTION(BlueprintImplementableEvent, Category="User Interface")
diff --git a/Engine/Source/Runtime/UMG/Public/Blueprint/WidgetTree.h b/Engine/Source/Runtime/UMG/Public/Blueprint/WidgetTree.h
--- a/Engine/Source/Runtime/UMG/Public/Blueprint/WidgetTree.h
+++ b/Engine/Source/Runtime/UMG/Public/Blueprint/WidgetTree.h
@@ -32,4 +32,8 @@ class UMG_API UWidgetTree : public UObject
 
 		return (T*)Widget;
 	}
+	
+private:
+	bool RemoveWidgetRecursive(UWidget* InRemovedWidget);
+	
 };
diff --git a/Engine/Source/Runtime/UMG/Public/Components/Border.h b/Engine/Source/Runtime/UMG/Public/Components/Border.h
--- a/Engine/Source/Runtime/UMG/Public/Components/Border.h
+++ b/Engine/Source/Runtime/UMG/Public/Components/Border.h
@@ -16,51 +16,57 @@ class UMG_API UBorder : public UContentWidget
 	GENERATED_UCLASS_BODY()
 
 	/** Horizontal positioning of the content within the border */
-	UPROPERTY(EditAnywhere, BlueprintReadOnly, Category="Content Layout")
+	UPROPERTY(EditDefaultsOnly, Category="Content Layout")
 	TEnumAsByte<EHorizontalAlignment> HorizontalAlignment;
 
 	/** Vertical positioning of the content within the border */
-	UPROPERTY(EditAnywhere, BlueprintReadOnly, Category="Content Layout")
+	UPROPERTY(EditDefaultsOnly, Category="Content Layout")
 	TEnumAsByte<EVerticalAlignment> VerticalAlignment;
 
 	/** The padding to add around the inner content. */
-	UPROPERTY(EditAnywhere, BlueprintReadWrite, Category="Content Layout")
+	UPROPERTY(EditDefaultsOnly, Category="Content Layout")
 	FMargin ContentPadding;
 
 	/** The scaling factor for the border content */
-	UPROPERTY(EditAnywhere, BlueprintReadOnly, Category="Content Layout")
+	UPROPERTY(EditDefaultsOnly, Category="Content Layout")
 	FVector2D ContentScale;
 
 	/** Color and opacity multiplier of content in the border */
-	UPROPERTY(EditAnywhere, BlueprintReadWrite, Category="Content Layout")
+	UPROPERTY(EditDefaultsOnly, Category="Content Layout")
 	FLinearColor ContentColorAndOpacity;
 
-	//@TODO: Should these be exposed?
-	// SLATE_EVENT( FPointerEventHandler, OnMouseButtonDown )
-	// SLATE_EVENT( FPointerEventHandler, OnMouseButtonUp )
-	// SLATE_EVENT( FPointerEventHandler, OnMouseMove )
-	// SLATE_EVENT( FPointerEventHandler, OnMouseDoubleClick )
-
 	/** Image to use for the border */
-	UPROPERTY(EditAnywhere, BlueprintReadWrite, Category=Appearance, meta=( DisplayThumbnail = "true" ))
+	UPROPERTY(EditDefaultsOnly, Category=Appearance, meta=( DisplayThumbnail = "true" ))
 	USlateBrushAsset* BorderBrush;
 
 	/** The scaling factor for the border */
-	UPROPERTY(EditAnywhere, BlueprintReadOnly, Category=Appearance)
+	UPROPERTY(EditDefaultsOnly, Category=Appearance)
 	FVector2D DesiredSizeScale;
 
 	/** Color and opacity of the actual border image */
-	UPROPERTY(EditAnywhere, BlueprintReadWrite, Category=Appearance)
+	UPROPERTY(EditDefaultsOnly, Category=Appearance)
 	FLinearColor BorderColorAndOpacity;
 
 	/** The foreground color of text and some glyphs that appear as the border's content. */
-	UPROPERTY(EditAnywhere, BlueprintReadWrite, Category=Appearance)
+	UPROPERTY(EditDefaultsOnly, Category=Appearance)
 	FLinearColor ForegroundColor;
 
 	/** Whether or not to show the disabled effect when this border is disabled */
-	UPROPERTY(EditAnywhere, BlueprintReadOnly, Category=Appearance)
+	UPROPERTY(EditDefaultsOnly, Category=Appearance)
 	bool bShowEffectWhenDisabled;
 
+	UPROPERTY(EditDefaultsOnly, Category=Events)
+	FOnPointerEvent OnMouseButtonDownEvent;
+
+	UPROPERTY(EditDefaultsOnly, Category=Events)
+	FOnPointerEvent OnMouseButtonUpEvent;
+
+	UPROPERTY(EditDefaultsOnly, Category=Events)
+	FOnPointerEvent OnMouseMoveEvent;
+
+	UPROPERTY(EditDefaultsOnly, Category=Events)
+	FOnPointerEvent OnMouseDoubleClickEvent;
+
 	// UContentWidget interface
 	virtual void SetContent(UWidget* Content) override;
 	// End UContentWidget interface
@@ -77,5 +83,10 @@ protected:
 	virtual TSharedRef<SWidget> RebuildWidget() override;
 	// End of UWidget interface
 
+	FReply HandleMouseButtonDown(const FGeometry& Geometry, const FPointerEvent& MouseEvent);
+	FReply HandleMouseButtonUp(const FGeometry& Geometry, const FPointerEvent& MouseEvent);
+	FReply HandleMouseMove(const FGeometry& Geometry, const FPointerEvent& MouseEvent);
+	FReply HandleMouseDoubleClick(const FGeometry& Geometry, const FPointerEvent& MouseEvent);
+
 	const FSlateBrush* GetBorderBrush() const;
 };
diff --git a/Engine/Source/Runtime/UMG/Public/Components/Button.h b/Engine/Source/Runtime/UMG/Public/Components/Button.h
--- a/Engine/Source/Runtime/UMG/Public/Components/Button.h
+++ b/Engine/Source/Runtime/UMG/Public/Components/Button.h
@@ -36,13 +36,17 @@ public:
 	/** The scaling factor for the button content */
 	UPROPERTY(EditDefaultsOnly, Category=Appearance)
 	FVector2D ContentScale;
-
+	
 	/** The color multiplier for the button images */
-	UPROPERTY(EditDefaultsOnly, Category=Appearance)
-	FLinearColor ButtonColorAndOpacity;
+	UPROPERTY(EditDefaultsOnly, Category=Appearance )
+	FLinearColor ColorAndOpacity;
+	
+	/** The color multiplier for the button background */
+	UPROPERTY(EditDefaultsOnly, Category=Appearance )
+	FLinearColor BackgroundColor;
 
 	/** The foreground color of the button */
-	UPROPERTY(EditDefaultsOnly, Category=Appearance)
+	UPROPERTY(EditDefaultsOnly, Category=Appearance )
 	FLinearColor ForegroundColor;
 
 	UPROPERTY(EditDefaultsOnly, Category=Sound)
@@ -55,17 +59,21 @@ public:
 	UPROPERTY(EditDefaultsOnly, Category=Events)
 	FOnReply OnClickedEvent;
 	
-	/*  */
+	/**  */
 	UFUNCTION(BlueprintCallable, Category="Appearance")
 	void SetButtonColorAndOpacity(FLinearColor InButtonColorAndOpacity);
 
-	/*  */
+	/**  */
 	UFUNCTION(BlueprintCallable, Category="Appearance")
 	void SetForegroundColor(FLinearColor InForegroundColor);
 
-	/*  */
+	/**  */
 	UFUNCTION(BlueprintCallable, Category="Appearance")
 	void SetContentPadding(FMargin InContentPadding);
+	
+	/**  */
+	UFUNCTION(BlueprintCallable, Category="Widget")
+	bool IsPressed() const;
 
 	// UContentWidget interface
 	virtual void SetContent(UWidget* Content) override;
diff --git a/Engine/Source/Runtime/UMG/Public/Components/CanvasPanel.h b/Engine/Source/Runtime/UMG/Public/Components/CanvasPanel.h
--- a/Engine/Source/Runtime/UMG/Public/Components/CanvasPanel.h
+++ b/Engine/Source/Runtime/UMG/Public/Components/CanvasPanel.h
@@ -17,6 +17,7 @@ class UMG_API UCanvasPanel : public UPanelWidget
 	UPROPERTY(EditAnywhere, Category=Appearance)
 	FVector2D DesiredCanvasSize;
 
+	/** Adds content to a new slot.  The slot will automatically be allocated and returned. */
 	UCanvasPanelSlot* AddSlot(UWidget* Content);
 
 	// UPanelWidget
diff --git a/Engine/Source/Runtime/UMG/Public/Components/CheckBox.h b/Engine/Source/Runtime/UMG/Public/Components/CheckBox.h
--- a/Engine/Source/Runtime/UMG/Public/Components/CheckBox.h
+++ b/Engine/Source/Runtime/UMG/Public/Components/CheckBox.h
@@ -14,7 +14,7 @@ class UMG_API UCheckBox : public UWidget
 
 protected:
 	/** Style of the check box */
-	UPROPERTY(EditAnywhere, BlueprintReadOnly, Category=Style)
+	UPROPERTY(EditDefaultsOnly, BlueprintReadOnly, Category=Style, meta=( DisplayThumbnail = "true" ))
 	USlateWidgetStyleAsset* Style;
 
 	/** Called when the checked state has changed */
@@ -22,23 +22,23 @@ protected:
 	FOnCheckBoxComponentStateChanged OnCheckStateChanged;
 
 	/** Whether the check box is currently in a checked state */
-	UPROPERTY(EditAnywhere, BlueprintReadWrite, Category=Appearance)
+	UPROPERTY(EditDefaultsOnly, BlueprintReadWrite, Category=Appearance)
 	bool bIsChecked;
 
 	/** How the content of the toggle button should align within the given space */
-	UPROPERTY(EditAnywhere, BlueprintReadOnly, Category=Appearance)
+	UPROPERTY(EditDefaultsOnly, BlueprintReadOnly, Category=Appearance)
 	TEnumAsByte<EHorizontalAlignment> HorizontalAlignment;
 
 	/** Spacing between the check box image and its content */
-	UPROPERTY(EditAnywhere, BlueprintReadOnly, Category=Appearance)
+	UPROPERTY(EditDefaultsOnly, BlueprintReadOnly, Category=Appearance)
 	FMargin Padding;
 
 	/** Foreground color for the checkbox's content and parts */
-	UPROPERTY(EditAnywhere, BlueprintReadOnly, Category=Appearance)
+	UPROPERTY(EditDefaultsOnly, BlueprintReadOnly, Category=Appearance)
 	FSlateColor ForegroundColor;
 
 	/** The color of the background border */
-	UPROPERTY(EditAnywhere, BlueprintReadOnly, Category=Appearance)
+	UPROPERTY(EditDefaultsOnly, BlueprintReadOnly, Category=Appearance)
 	FSlateColor BorderBackgroundColor;
 
 	//SLATE_ATTRIBUTE( bool, ReadOnly )
@@ -46,14 +46,62 @@ protected:
 
 	//SLATE_EVENT( FOnGetContent, OnGetMenuContent )
 
-	UPROPERTY(EditAnywhere, Category=Sound)
+	UPROPERTY(EditDefaultsOnly, Category="Sound")
 	FSlateSound CheckedSound;
 
-	UPROPERTY(EditAnywhere, Category=Sound)
+	UPROPERTY(EditDefaultsOnly, Category="Sound")
 	FSlateSound UncheckedSound;
 
-	UPROPERTY(EditAnywhere, Category=Sound)
+	UPROPERTY(EditDefaultsOnly, Category="Sound")
 	FSlateSound HoveredSound;
+	
+	/** Image to use when the checkbox is unchecked */
+	UPROPERTY(EditDefaultsOnly, Category="Style Overrides", meta=( DisplayThumbnail = "true" ))
+	USlateBrushAsset* UncheckedImage;
+	
+	/** Image to use when the checkbox is unchecked and hovered */
+	UPROPERTY(EditDefaultsOnly, Category="Style Overrides", meta=( DisplayThumbnail = "true" ))
+	USlateBrushAsset* UncheckedHoveredImage;
+	
+	/** Image to use when the checkbox is unchecked and pressed */
+	UPROPERTY(EditDefaultsOnly, Category="Style Overrides", meta=( DisplayThumbnail = "true" ))
+	USlateBrushAsset* UncheckedPressedImage;
+	
+	/** Image to use when the checkbox is checked */
+	UPROPERTY(EditDefaultsOnly, Category="Style Overrides", meta=( DisplayThumbnail = "true" ))
+	USlateBrushAsset* CheckedImage;
+	
+	/** Image to use when the checkbox is checked and hovered */
+	UPROPERTY(EditDefaultsOnly, Category="Style Overrides", meta=( DisplayThumbnail = "true" ))
+	USlateBrushAsset* CheckedHoveredImage;
+	
+	/** Image to use when the checkbox is checked and pressed */
+	UPROPERTY(EditDefaultsOnly, Category="Style Overrides", meta=( DisplayThumbnail = "true" ))
+	USlateBrushAsset* CheckedPressedImage;
+	
+	/** Image to use when the checkbox is in an ambiguous state and hovered */
+	UPROPERTY(EditDefaultsOnly, Category="Style Overrides", meta=( DisplayThumbnail = "true" ))
+	USlateBrushAsset* UndeterminedImage;
+	
+	/** Image to use when the checkbox is checked and hovered */
+	UPROPERTY(EditDefaultsOnly, Category="Style Overrides", meta=( DisplayThumbnail = "true" ))
+	USlateBrushAsset* UndeterminedHoveredImage;
+	
+	/** Image to use when the checkbox is in an ambiguous state and pressed */
+	UPROPERTY(EditDefaultsOnly, Category="Style Overrides", meta=( DisplayThumbnail = "true" ))
+	USlateBrushAsset* UndeterminedPressedImage;
+	
+	/** Returns true if this button is currently pressed */
+	UFUNCTION(BlueprintCallable, Category="Widget")
+	bool IsPressed() const;
+	
+	/** Returns true if the checkbox is currently checked */
+	UFUNCTION(BlueprintCallable, Category="Widget")
+	bool IsChecked() const;
+	
+	// UWidget interface
+	void SyncronizeProperties() override;
+	// End of UWidget interface
 
 protected:
 	// UWidget interface
@@ -62,4 +110,7 @@ protected:
 
 	ESlateCheckBoxState::Type GetCheckState() const;
 	void SlateOnCheckStateChangedCallback(ESlateCheckBoxState::Type NewState);
+	
+protected:
+	TSharedPtr<SCheckBox> MyCheckbox;
 };
diff --git a/Engine/Source/Runtime/UMG/Public/Components/EditableText.h b/Engine/Source/Runtime/UMG/Public/Components/EditableText.h
--- a/Engine/Source/Runtime/UMG/Public/Components/EditableText.h
+++ b/Engine/Source/Runtime/UMG/Public/Components/EditableText.h
@@ -86,18 +86,19 @@ protected:
 	/** Called whenever the text is committed.  This happens when the user presses enter or the text box loses focus. */
 	UPROPERTY(BlueprintAssignable)
 	FOnEditableTextCommittedEvent OnTextCommitted;
+	
+	// UWidget interface
+	void SyncronizeProperties() override;
+	// End of UWidget interface
 
 protected:
 	// UWidget interface
 	virtual TSharedRef<SWidget> RebuildWidget() override;
 	// End of UWidget
 
-#if WITH_EDITOR
-	// UObject interface
-	virtual void PostEditChangeProperty(struct FPropertyChangedEvent& PropertyChangedEvent) override;
-	// End of UObject interface
-#endif
-
 	void HandleOnTextChanged(const FText& Text);
 	void HandleOnTextCommitted(const FText& Text, ETextCommit::Type CommitMethod);
+
+protected:
+	TSharedPtr<SEditableText> MyEditableText;
 };
diff --git a/Engine/Source/Runtime/UMG/Public/Components/HorizontalBox.h b/Engine/Source/Runtime/UMG/Public/Components/HorizontalBox.h
--- a/Engine/Source/Runtime/UMG/Public/Components/HorizontalBox.h
+++ b/Engine/Source/Runtime/UMG/Public/Components/HorizontalBox.h
@@ -37,7 +37,7 @@ class UMG_API UHorizontalBox : public UPanelWidget
 
 protected:
 
-	TWeakPtr<class SHorizontalBox> MyHorizontalBox;
+	TSharedPtr<class SHorizontalBox> MyHorizontalBox;
 
 protected:
 	// UWidget interface
diff --git a/Engine/Source/Runtime/UMG/Public/Components/Overlay.h b/Engine/Source/Runtime/UMG/Public/Components/Overlay.h
--- a/Engine/Source/Runtime/UMG/Public/Components/Overlay.h
+++ b/Engine/Source/Runtime/UMG/Public/Components/Overlay.h
@@ -35,7 +35,7 @@ class UMG_API UOverlay : public UPanelWidget
 
 protected:
 
-	TWeakPtr<class SOverlay> MyVerticalBox;
+	TSharedPtr<class SOverlay> MyOverlay;
 
 protected:
 	// UWidget interface
diff --git a/Engine/Source/Runtime/UMG/Public/Components/PanelSlot.h b/Engine/Source/Runtime/UMG/Public/Components/PanelSlot.h
--- a/Engine/Source/Runtime/UMG/Public/Components/PanelSlot.h
+++ b/Engine/Source/Runtime/UMG/Public/Components/PanelSlot.h
@@ -17,6 +17,8 @@ public:
 
 	UPROPERTY()
 	class UWidget* Content;
+	
+	bool IsDesignTime() const;
 
 	virtual void Resize(const FVector2D& Direction, const FVector2D& Amount);
 
diff --git a/Engine/Source/Runtime/UMG/Public/Components/ProgressBar.h b/Engine/Source/Runtime/UMG/Public/Components/ProgressBar.h
--- a/Engine/Source/Runtime/UMG/Public/Components/ProgressBar.h
+++ b/Engine/Source/Runtime/UMG/Public/Components/ProgressBar.h
@@ -5,23 +5,27 @@
 #include "ProgressBar.generated.h"
 
 /** ProgressBar widget */
-UCLASS(meta=( Category="Misc" ), ClassGroup=UserInterface)
+UCLASS(meta=( Category="Common" ), ClassGroup=UserInterface)
 class UMG_API UProgressBar : public UWidget
 {
 	GENERATED_UCLASS_BODY()
+	
+public:
 
-	///** Style used for the progress bar */
-	//UPROPERTY(EditDefaultsOnly, Category=Appearance)
-	//FProgressBarStyle Style;
+	/** Style used for the progress bar */
+	UPROPERTY(EditDefaultsOnly, Category=Style, meta=( DisplayThumbnail = "true" ))
+	USlateWidgetStyleAsset* Style;
 
 	/** Defines if this progress bar fills Left to right or right to left */
 	UPROPERTY(EditDefaultsOnly, Category=Appearance)
 	TEnumAsByte<EProgressBarFillType::Type> BarFillType;
+	
+	UPROPERTY()
+	bool bOverride_Percent;
 
 	/** Used to determine the fill position of the progress bar ranging 0..1 */
-	UPROPERTY(EditDefaultsOnly, Category=Appearance, meta=( UIMin = "0", UIMax = "1" ))
+	UPROPERTY(EditDefaultsOnly, Category=Appearance, meta=( UIMin = "0", UIMax = "1", editcondition = "bOverride_Percent" ))
 	float Percent;
-	//TODO UMG Slate supports TOptional<float> for Percent.  UMG should as well.
 
 	/** Fill Color and Opacity */
 	UPROPERTY(EditDefaultsOnly, Category=Appearance)
@@ -30,10 +34,26 @@ class UMG_API UProgressBar : public UWidget
 	/** Border Padding around fill bar */
 	UPROPERTY(EditDefaultsOnly, Category=Appearance)
 	FVector2D BorderPadding;
-
+	
+	/** The brush to use as the background of the progress bar */
+	UPROPERTY(EditDefaultsOnly, Category=Appearance, meta=( DisplayThumbnail = "true" ))
+	USlateBrushAsset* BackgroundImage;
+	
+	/** The brush to use as the fill image */
+	UPROPERTY(EditDefaultsOnly, Category=Appearance, meta=( DisplayThumbnail = "true" ))
+	USlateBrushAsset* FillImage;
+	
+	/** The brush to use as the marquee image */
+	UPROPERTY(EditDefaultsOnly, Category=Appearance, meta=( DisplayThumbnail = "true" ))
+	USlateBrushAsset* MarqueeImage;
+	
 	/** Sets the current value of the ProgressBar. */
 	UFUNCTION(BlueprintCallable, Category="Behavior")
 	void SetPercent(float InPercent);
+	
+	// UWidget interface
+	void SyncronizeProperties() override;
+	// End of UWidget interface
 
 protected:
 	/** Native Slate Widget */
diff --git a/Engine/Source/Runtime/UMG/Public/Components/ScrollBox.h b/Engine/Source/Runtime/UMG/Public/Components/ScrollBox.h
--- a/Engine/Source/Runtime/UMG/Public/Components/ScrollBox.h
+++ b/Engine/Source/Runtime/UMG/Public/Components/ScrollBox.h
@@ -40,7 +40,7 @@ class UMG_API UScrollBox : public UPanelWidget
 
 protected:
 
-	TWeakPtr<class SScrollBox> MyScrollBox;
+	TSharedPtr<class SScrollBox> MyScrollBox;
 
 protected:
 	// UWidget interface
diff --git a/Engine/Source/Runtime/UMG/Public/Components/SlateWrapperTypes.h b/Engine/Source/Runtime/UMG/Public/Components/SlateWrapperTypes.h
--- a/Engine/Source/Runtime/UMG/Public/Components/SlateWrapperTypes.h
+++ b/Engine/Source/Runtime/UMG/Public/Components/SlateWrapperTypes.h
@@ -52,7 +52,10 @@ struct FSReply
 	UPROPERTY(EditAnywhere, Category=User)
 	uint32 bIsHandled : 1;
 
-	FReply ToReply() const
+	UPROPERTY(EditAnywhere, Category=User)
+	uint32 bCaptureMouse : 1;
+
+	FReply ToReply( TSharedRef<SWidget> Widget ) const
 	{
 		FReply Reply = FReply::Unhandled();
 
@@ -61,6 +64,11 @@ struct FSReply
 			Reply = FReply::Handled();
 		}
 
+		if ( bCaptureMouse )
+		{
+			Reply = Reply.CaptureMouse(Widget);
+		}
+
 		return Reply;
 	}
 };
diff --git a/Engine/Source/Runtime/UMG/Public/Components/Slider.h b/Engine/Source/Runtime/UMG/Public/Components/Slider.h
--- a/Engine/Source/Runtime/UMG/Public/Components/Slider.h
+++ b/Engine/Source/Runtime/UMG/Public/Components/Slider.h
@@ -14,11 +14,13 @@ class UMG_API USlider : public UWidget
 {
 	GENERATED_UCLASS_BODY()
 
-	///** Whether the slidable area should be indented to fit the handle. */
-	//SLATE_ATTRIBUTE(bool, IndentHandle)
+	/** Whether the slidable area should be indented to fit the handle. */
+	UPROPERTY(EditDefaultsOnly, Category=Appearance)
+	bool IndentHandle;
 
-	///** Whether the handle is interactive or fixed. */
-	//SLATE_ATTRIBUTE(bool, Locked)
+	/** Whether the handle is interactive or fixed. */
+	UPROPERTY(EditDefaultsOnly, Category=Appearance)
+	bool Locked;
 
 	///** The style used to draw the slider. */
 	//SLATE_STYLE_ARGUMENT(FSliderStyle, Style)
@@ -58,6 +60,10 @@ class UMG_API USlider : public UWidget
 	/** Sets the current value of the slider. */
 	UFUNCTION(BlueprintCallable, Category="Behavior")
 	void SetValue(float InValue);
+	
+	// UWidget interface
+	void SyncronizeProperties() override;
+	// End of UWidget interface
 
 protected:
 	/** Native Slate Widget */
diff --git a/Engine/Source/Runtime/UMG/Public/Components/Spacer.h b/Engine/Source/Runtime/UMG/Public/Components/Spacer.h
--- a/Engine/Source/Runtime/UMG/Public/Components/Spacer.h
+++ b/Engine/Source/Runtime/UMG/Public/Components/Spacer.h
@@ -17,6 +17,10 @@ class UMG_API USpacer : public UWidget
 	/** Sets the size of the spacer */
 	UFUNCTION(BlueprintCallable, Category="Widget")
 	void SetSize(FVector2D InSize);
+	
+	// UWidget interface
+	void SyncronizeProperties() override;
+	// End of UWidget interface
 
 protected:
 	// UWidget interface
diff --git a/Engine/Source/Runtime/UMG/Public/Components/UniformGridPanel.h b/Engine/Source/Runtime/UMG/Public/Components/UniformGridPanel.h
--- a/Engine/Source/Runtime/UMG/Public/Components/UniformGridPanel.h
+++ b/Engine/Source/Runtime/UMG/Public/Components/UniformGridPanel.h
@@ -51,7 +51,7 @@ class UMG_API UUniformGridPanel : public UPanelWidget
 
 protected:
 
-	TWeakPtr<class SUniformGridPanel> MyUniformGridPanel;
+	TSharedPtr<SUniformGridPanel> MyUniformGridPanel;
 
 protected:
 	// UWidget interface
diff --git a/Engine/Source/Runtime/UMG/Public/Components/VerticalBox.h b/Engine/Source/Runtime/UMG/Public/Components/VerticalBox.h
--- a/Engine/Source/Runtime/UMG/Public/Components/VerticalBox.h
+++ b/Engine/Source/Runtime/UMG/Public/Components/VerticalBox.h
@@ -43,7 +43,7 @@ class UMG_API UVerticalBox : public UPanelWidget
 
 protected:
 
-	TWeakPtr<class SVerticalBox> MyVerticalBox;
+	TSharedPtr<class SVerticalBox> MyVerticalBox;
 
 protected:
 	// UWidget interface
diff --git a/Engine/Source/Runtime/UMG/Public/Components/Widget.h b/Engine/Source/Runtime/UMG/Public/Components/Widget.h
--- a/Engine/Source/Runtime/UMG/Public/Components/Widget.h
+++ b/Engine/Source/Runtime/UMG/Public/Components/Widget.h
@@ -39,6 +39,7 @@ public:
 
 	// Events
 	DECLARE_DYNAMIC_DELEGATE_RetVal(FSReply, FOnReply);
+	DECLARE_DYNAMIC_DELEGATE_RetVal_TwoParams(FSReply, FOnPointerEvent, FGeometry, MyGeometry, const FPointerEvent&, MouseEvent);
 
 	/**
 	 * Allows controls to be exposed as variables in a blueprint.  Not all controls need to be exposed
@@ -120,7 +121,7 @@ public:
 	bool IsHovered() const;
 
 	/** Gets the underlying Slate SWidget for this UWidget. */
-	TSharedRef<SWidget> GetWidget();
+	TSharedRef<SWidget> GetWidget() const;
 
 	/**
 	 * Applies all properties to the native widget if possible.  This is called after a widget is constructed.
@@ -129,17 +130,21 @@ public:
 	 */
 	virtual void SyncronizeProperties();
 
+	/** Returns if the widget is currently being displayed in the designer, it may want to display different data. */
+	bool IsDesignTime() const;
+	
+	/** Sets that this widget is being designed */
+	void IsDesignTime(bool bInDesignTime);
+	
 #if WITH_EDITOR
+	/** Gets a widget representing the tiny preview of the toolbox */
+	virtual TSharedRef<SWidget> GetToolboxPreviewWidget() const;
+	
 	/** Allows general fixups and connections only used at editor time. */
 	virtual void ConnectEditorData() { }
 
 	// UObject interface
-	virtual void PostEditChangeProperty(struct FPropertyChangedEvent& PropertyChangedEvent) override
-	{
-		Super::PostEditChangeProperty(PropertyChangedEvent);
-
-		SyncronizeProperties();
-	}
+	virtual void PostEditChangeProperty(struct FPropertyChangedEvent& PropertyChangedEvent) override;
 	// End of UObject interface
 #endif
 
@@ -153,6 +158,8 @@ public:
 protected:
 	/** Function implemented by all subclasses of UWidget is called when the underlying SWidget needs to be constructed. */
 	virtual TSharedRef<SWidget> RebuildWidget();
+	
+	TSharedRef<SWidget> BuildDesignTimeWidget(TSharedRef<SWidget> WrapWidget);
 
 	EVisibility ConvertVisibility(TAttribute<ESlateVisibility::Type> SerializedType) const
 	{
@@ -161,5 +168,8 @@ protected:
 
 protected:
 	/** The underlying SWidget. */
-	TSharedPtr<SWidget> MyWidget;
+	mutable TSharedPtr<SWidget> MyWidget;
+	
+	/** Is this widget being displayed on a designer surface */
+	bool bDesignTime;
 };
