UMG - Fixing a lot of Undo/Redo issues.  Adding design time borders for list and tile views.  Exposing more overridable style data on the EditableText widget.

diff --git a/Engine/Source/Editor/UMGEditor/Private/Extensions/CanvasSlotExtension.cpp b/Engine/Source/Editor/UMGEditor/Private/Extensions/CanvasSlotExtension.cpp
--- a/Engine/Source/Editor/UMGEditor/Private/Extensions/CanvasSlotExtension.cpp
+++ b/Engine/Source/Editor/UMGEditor/Private/Extensions/CanvasSlotExtension.cpp
@@ -51,6 +51,9 @@ void FCanvasSlotExtension::BuildWidgetsForSelection(const TArray< FSelectedWidge
 FReply FCanvasSlotExtension::HandleBeginDrag(const FGeometry& Geometry, const FPointerEvent& Event)
 {
 	bDragging = true;
+
+	BeginTransaction(LOCTEXT("MoveWidget", "Move Widget"));
+
 	return FReply::Handled().CaptureMouse(MoveHandle.ToSharedRef());
 }
 
@@ -58,6 +61,8 @@ FReply FCanvasSlotExtension::HandleEndDrag(const FGeometry& Geometry, const FPoi
 {
 	bDragging = false;
 
+	EndTransaction();
+
 	FBlueprintEditorUtils::MarkBlueprintAsModified(Blueprint);
 
 	return FReply::Handled().ReleaseMouseCapture();
diff --git a/Engine/Source/Editor/UMGEditor/Private/Extensions/CanvasSlotExtension.h b/Engine/Source/Editor/UMGEditor/Private/Extensions/CanvasSlotExtension.h
--- a/Engine/Source/Editor/UMGEditor/Private/Extensions/CanvasSlotExtension.h
+++ b/Engine/Source/Editor/UMGEditor/Private/Extensions/CanvasSlotExtension.h
@@ -21,7 +21,6 @@ private:
 
 	void MoveByAmount(UWidget* Widget, FVector2D Delta);
 
-	TArray< FSelectedWidget > SelectionCache;
 	bool bDragging;
 
 	TSharedPtr<SBorder> MoveHandle;
diff --git a/Engine/Source/Editor/UMGEditor/Private/Extensions/DesignerExtension.cpp b/Engine/Source/Editor/UMGEditor/Private/Extensions/DesignerExtension.cpp
--- a/Engine/Source/Editor/UMGEditor/Private/Extensions/DesignerExtension.cpp
+++ b/Engine/Source/Editor/UMGEditor/Private/Extensions/DesignerExtension.cpp
@@ -73,6 +73,12 @@ UWidget* FSelectedWidget::GetPreview() const
 
 // Designer Extension
 
+FDesignerExtension::FDesignerExtension()
+	: ScopedTransaction(NULL)
+{
+
+}
+
 void FDesignerExtension::Initialize(UWidgetBlueprint* InBlueprint)
 {
 	Blueprint = InBlueprint;
@@ -83,4 +89,27 @@ FName FDesignerExtension::GetExtensionId() const
 	return ExtensionId;
 }
 
+void FDesignerExtension::BeginTransaction(const FText& SessionName)
+{
+	if ( ScopedTransaction == NULL )
+	{
+		ScopedTransaction = new FScopedTransaction(SessionName);
+	}
+
+	for ( FSelectedWidget& Selection : SelectionCache )
+	{
+		Selection.GetPreview()->Modify();
+		Selection.GetTemplate()->Modify();
+	}
+}
+
+void FDesignerExtension::EndTransaction()
+{
+	if ( ScopedTransaction != NULL )
+	{
+		delete ScopedTransaction;
+		ScopedTransaction = NULL;
+	}
+}
+
 #undef LOCTEXT_NAMESPACE
diff --git a/Engine/Source/Editor/UMGEditor/Private/Extensions/HorizontalSlotExtension.cpp b/Engine/Source/Editor/UMGEditor/Private/Extensions/HorizontalSlotExtension.cpp
--- a/Engine/Source/Editor/UMGEditor/Private/Extensions/HorizontalSlotExtension.cpp
+++ b/Engine/Source/Editor/UMGEditor/Private/Extensions/HorizontalSlotExtension.cpp
@@ -48,12 +48,16 @@ void FHorizontalSlotExtension::BuildWidgetsForSelection(const TArray< FSelectedW
 
 FReply FHorizontalSlotExtension::HandleShift(int32 ShiftAmount)
 {
+	BeginTransaction(LOCTEXT("MoveWidget", "Move Widget"));
+
 	for ( FSelectedWidget& Selection : SelectionCache )
 	{
 		ShiftHorizontal(Selection.GetPreview(), ShiftAmount);
 		ShiftHorizontal(Selection.GetTemplate(), ShiftAmount);
 	}
 
+	EndTransaction();
+
 	//TODO UMG Reorder the live slot without rebuilding the structure
 	FBlueprintEditorUtils::MarkBlueprintAsStructurallyModified(Blueprint);
 
diff --git a/Engine/Source/Editor/UMGEditor/Private/Extensions/HorizontalSlotExtension.h b/Engine/Source/Editor/UMGEditor/Private/Extensions/HorizontalSlotExtension.h
--- a/Engine/Source/Editor/UMGEditor/Private/Extensions/HorizontalSlotExtension.h
+++ b/Engine/Source/Editor/UMGEditor/Private/Extensions/HorizontalSlotExtension.h
@@ -18,6 +18,4 @@ private:
 	FReply HandleShift(int32 ShiftAmount);
 	
 	void ShiftHorizontal(UWidget* Widget, int32 ShiftAmount);
-
-	TArray< FSelectedWidget > SelectionCache;
 };
\ No newline at end of file
diff --git a/Engine/Source/Editor/UMGEditor/Private/Extensions/UniformGridSlotExtension.cpp b/Engine/Source/Editor/UMGEditor/Private/Extensions/UniformGridSlotExtension.cpp
--- a/Engine/Source/Editor/UMGEditor/Private/Extensions/UniformGridSlotExtension.cpp
+++ b/Engine/Source/Editor/UMGEditor/Private/Extensions/UniformGridSlotExtension.cpp
@@ -59,12 +59,16 @@ void FUniformGridSlotExtension::BuildWidgetsForSelection(const TArray< FSelected
 
 FReply FUniformGridSlotExtension::HandleShiftRow(int32 ShiftAmount)
 {
+	BeginTransaction(LOCTEXT("MoveWidget", "Move Widget"));
+
 	for ( FSelectedWidget& Selection : SelectionCache )
 	{
 		ShiftRow(Selection.GetPreview(), ShiftAmount);
 		ShiftRow(Selection.GetTemplate(), ShiftAmount);
 	}
 
+	EndTransaction();
+
 	FBlueprintEditorUtils::MarkBlueprintAsModified(Blueprint);
 
 	return FReply::Handled();
@@ -72,12 +76,16 @@ FReply FUniformGridSlotExtension::HandleShiftRow(int32 ShiftAmount)
 
 FReply FUniformGridSlotExtension::HandleShiftColumn(int32 ShiftAmount)
 {
+	BeginTransaction(LOCTEXT("MoveWidget", "Move Widget"));
+
 	for ( FSelectedWidget& Selection : SelectionCache )
 	{
 		ShiftColumn(Selection.GetPreview(), ShiftAmount);
 		ShiftColumn(Selection.GetTemplate(), ShiftAmount);
 	}
 
+	EndTransaction();
+
 	FBlueprintEditorUtils::MarkBlueprintAsModified(Blueprint);
 
 	return FReply::Handled();
diff --git a/Engine/Source/Editor/UMGEditor/Private/Extensions/UniformGridSlotExtension.h b/Engine/Source/Editor/UMGEditor/Private/Extensions/UniformGridSlotExtension.h
--- a/Engine/Source/Editor/UMGEditor/Private/Extensions/UniformGridSlotExtension.h
+++ b/Engine/Source/Editor/UMGEditor/Private/Extensions/UniformGridSlotExtension.h
@@ -20,8 +20,4 @@ private:
 
 	void ShiftRow(UWidget* Widget, int32 ShiftAmount);
 	void ShiftColumn(UWidget* Widget, int32 ShiftAmount);
-
-private:
-
-	TArray< FSelectedWidget > SelectionCache;
 };
\ No newline at end of file
diff --git a/Engine/Source/Editor/UMGEditor/Private/Extensions/VerticalSlotExtension.cpp b/Engine/Source/Editor/UMGEditor/Private/Extensions/VerticalSlotExtension.cpp
--- a/Engine/Source/Editor/UMGEditor/Private/Extensions/VerticalSlotExtension.cpp
+++ b/Engine/Source/Editor/UMGEditor/Private/Extensions/VerticalSlotExtension.cpp
@@ -35,37 +35,34 @@ void FVerticalSlotExtension::BuildWidgetsForSelection(const TArray< FSelectedWid
 	TSharedRef<SButton> UpButton =
 		SNew(SButton)
 		.Text(LOCTEXT("UpArrow", "↑"))
-		.OnClicked(this, &FVerticalSlotExtension::HandleUpPressed);
+		.OnClicked(this, &FVerticalSlotExtension::HandleShiftVertical, -1);
 
 	TSharedRef<SButton> DownButton =
 		SNew(SButton)
 		.Text(LOCTEXT("DownArrow", "↓"))
-		.OnClicked(this, &FVerticalSlotExtension::HandleDownPressed);
+		.OnClicked(this, &FVerticalSlotExtension::HandleShiftVertical, 1);
+
+	UpButton->SetEnabled(TAttribute<bool>::Create(TAttribute<bool>::FGetter::CreateSP(this, &FVerticalSlotExtension::CanShift, -1)));
+	DownButton->SetEnabled(TAttribute<bool>::Create(TAttribute<bool>::FGetter::CreateSP(this, &FVerticalSlotExtension::CanShift, 1)));
 
 	Widgets.Add(UpButton);
 	Widgets.Add(DownButton);
 }
 
-FReply FVerticalSlotExtension::HandleUpPressed()
+bool FVerticalSlotExtension::CanShift(int32 ShiftAmount) const
 {
-	for ( FSelectedWidget& Selection : SelectionCache )
-	{
-		MoveUp(Selection.GetPreview());
-		MoveUp(Selection.GetTemplate());
-	}
-
-	//TODO UMG Reorder the live slot without rebuilding the structure
-	FBlueprintEditorUtils::MarkBlueprintAsStructurallyModified(Blueprint);
-
-	return FReply::Handled();
+	//TODO UMG Provide feedback if shifting is possible.  Tricky with multiple items selected, if we ever support that.
+	return true;
 }
 
-FReply FVerticalSlotExtension::HandleDownPressed()
+FReply FVerticalSlotExtension::HandleShiftVertical(int32 ShiftAmount)
 {
+	BeginTransaction(LOCTEXT("MoveWidget", "Move Widget"));
+
 	for ( FSelectedWidget& Selection : SelectionCache )
 	{
-		MoveDown(Selection.GetPreview());
-		MoveDown(Selection.GetTemplate());
+		ShiftVertical(Selection.GetPreview(), ShiftAmount);
+		ShiftVertical(Selection.GetTemplate(), ShiftAmount);
 	}
 
 	//TODO UMG Reorder the live slot without rebuilding the structure
@@ -74,24 +71,14 @@ FReply FVerticalSlotExtension::HandleDownPressed()
 	return FReply::Handled();
 }
 
-void FVerticalSlotExtension::MoveUp(UWidget* Widget)
-{
-	UVerticalBoxSlot* VerticalSlot = Cast<UVerticalBoxSlot>(Widget->Slot);
-	UVerticalBox* Parent = Cast<UVerticalBox>(VerticalSlot->Parent);
-
-	int32 CurrentIndex = Parent->GetChildIndex(Widget);
-	Parent->Slots.RemoveAt(CurrentIndex);
-	Parent->Slots.Insert(VerticalSlot, FMath::Max(CurrentIndex - 1, 0));
-}
-
-void FVerticalSlotExtension::MoveDown(UWidget* Widget)
+void FVerticalSlotExtension::ShiftVertical(UWidget* Widget, int32 ShiftAmount)
 {
-	UVerticalBoxSlot* VerticalSlot = Cast<UVerticalBoxSlot>(Widget->Slot);
-	UVerticalBox* Parent = Cast<UVerticalBox>(VerticalSlot->Parent);
+	UVerticalBoxSlot* Slot = Cast<UVerticalBoxSlot>(Widget->Slot);
+	UVerticalBox* Parent = Cast<UVerticalBox>(Slot->Parent);
 
 	int32 CurrentIndex = Parent->GetChildIndex(Widget);
 	Parent->Slots.RemoveAt(CurrentIndex);
-	Parent->Slots.Insert(VerticalSlot, FMath::Min(CurrentIndex + 1, Parent->GetChildrenCount()));
+	Parent->Slots.Insert(Slot, FMath::Clamp(CurrentIndex + ShiftAmount, 0, Parent->GetChildrenCount()));
 }
 
 #undef LOCTEXT_NAMESPACE
diff --git a/Engine/Source/Editor/UMGEditor/Private/Extensions/VerticalSlotExtension.h b/Engine/Source/Editor/UMGEditor/Private/Extensions/VerticalSlotExtension.h
--- a/Engine/Source/Editor/UMGEditor/Private/Extensions/VerticalSlotExtension.h
+++ b/Engine/Source/Editor/UMGEditor/Private/Extensions/VerticalSlotExtension.h
@@ -15,11 +15,9 @@ public:
 
 private:
 
-	FReply HandleUpPressed();
-	FReply HandleDownPressed();
+	bool CanShift(int32 ShiftAmount) const;
 
-	void MoveUp(UWidget* Widget);
-	void MoveDown(UWidget* Widget);
+	FReply HandleShiftVertical(int32 ShiftAmount);
 
-	TArray< FSelectedWidget > SelectionCache;
+	void ShiftVertical(UWidget* Widget, int32 ShiftAmount);
 };
\ No newline at end of file
diff --git a/Engine/Source/Editor/UMGEditor/Private/SUMGDesigner.cpp b/Engine/Source/Editor/UMGEditor/Private/SUMGDesigner.cpp
--- a/Engine/Source/Editor/UMGEditor/Private/SUMGDesigner.cpp
+++ b/Engine/Source/Editor/UMGEditor/Private/SUMGDesigner.cpp
@@ -67,9 +67,11 @@ static bool LocateWidgetsUnderCursor_Helper(FArrangedWidget& Candidate, FVector2
 
 void SUMGDesigner::Construct(const FArguments& InArgs, TSharedPtr<FWidgetBlueprintEditor> InBlueprintEditor)
 {
+	ScopedTransaction = NULL;
+
 	PreviewWidgetObject = NULL;
-    DropPreviewWidget = NULL;
-    DropPreviewParent = NULL;
+	DropPreviewWidget = NULL;
+	DropPreviewParent = NULL;
 	BlueprintEditor = InBlueprintEditor;
 	CurrentHandle = DH_NONE;
 
@@ -187,14 +189,14 @@ FSelectedWidget SUMGDesigner::GetWidgetAtCursor(const FGeometry& MyGeometry, con
 		{
 			FArrangedWidget& Child = Children.GetInternalArray()[ChildIndex];
 			Preview = WidgetActor->GetWidgetHandle(Child.Widget);
-            
-            // Ignore the drop preview widget when doing widget picking
-            if (Preview == DropPreviewWidget)
-            {
-                Preview = NULL;
-                continue;
-            }
-            
+			
+			// Ignore the drop preview widget when doing widget picking
+			if (Preview == DropPreviewWidget)
+			{
+				Preview = NULL;
+				continue;
+			}
+			
 			if ( Preview )
 			{
 				ArrangedWidget = Child;
@@ -220,6 +222,8 @@ FReply SUMGDesigner::OnMouseButtonDown(const FGeometry& MyGeometry, const FPoint
 		FArrangedWidget ArrangedWidget(SNullWidget::NullWidget, FGeometry());
 		CurrentSelection = GetWidgetAtCursor(MyGeometry, MouseEvent, ArrangedWidget);
 
+		//TODO UMG Undoable Selection
+
 		if ( CurrentSelection.IsValid() )
 		{
 			//@TODO UMG primary FBlueprintEditor needs to be inherited and selection control needs to be centralized.
@@ -265,6 +269,9 @@ FReply SUMGDesigner::OnMouseButtonDown(const FGeometry& MyGeometry, const FPoint
 	}
 	else
 	{
+		BeginTransaction(LOCTEXT("ResizeWidget", "Resize Widget"));
+
+		// Capture mouse for the drag handle
 		return FReply::Handled().PreventThrottling().CaptureMouse(AsShared());
 	}
 
@@ -273,6 +280,8 @@ FReply SUMGDesigner::OnMouseButtonDown(const FGeometry& MyGeometry, const FPoint
 
 FReply SUMGDesigner::OnMouseButtonUp(const FGeometry& MyGeometry, const FPointerEvent& MouseEvent)
 {
+	EndTransaction();
+
 	CurrentHandle = DH_NONE;
 	return FReply::Handled().ReleaseMouseCapture();
 }
@@ -681,157 +690,187 @@ void SUMGDesigner::OnDragLeave(const FDragDropEvent& DragDropEvent)
 	{
 		DragDropOp->ResetToDefaultToolTip();
 	}
-    
-    if (DropPreviewWidget)
-    {
-        UWidgetBlueprint* BP = GetBlueprint();
-        BP->WidgetTree->RemoveWidget(DropPreviewWidget);
-        DropPreviewWidget = NULL;
-    }
+	
+	if (DropPreviewWidget)
+	{
+		UWidgetBlueprint* BP = GetBlueprint();
+		BP->WidgetTree->RemoveWidget(DropPreviewWidget);
+		DropPreviewWidget = NULL;
+	}
 }
 
 FReply SUMGDesigner::OnDragOver(const FGeometry& MyGeometry, const FDragDropEvent& DragDropEvent)
 {
-    UWidgetBlueprint* BP = GetBlueprint();
-    
+	UWidgetBlueprint* BP = GetBlueprint();
+	
 	if (DropPreviewWidget)
-    {
-        if (DropPreviewParent)
-        {
-            DropPreviewParent->RemoveChild(DropPreviewWidget);
-        }
-        
-        BP->WidgetTree->RemoveWidget(DropPreviewWidget);
-        DropPreviewWidget = NULL;
-    }
-    
-    DropPreviewWidget = AddPreview(MyGeometry, DragDropEvent);
-    if ( DropPreviewWidget )
-    {
-        //@TODO UMG Drop Feedback
-        return FReply::Handled();
-    }
+	{
+		if (DropPreviewParent)
+		{
+			DropPreviewParent->RemoveChild(DropPreviewWidget);
+		}
+		
+		BP->WidgetTree->RemoveWidget(DropPreviewWidget);
+		DropPreviewWidget = NULL;
+	}
+	
+	DropPreviewWidget = AddPreview(MyGeometry, DragDropEvent);
+	if ( DropPreviewWidget )
+	{
+		//@TODO UMG Drop Feedback
+		return FReply::Handled();
+	}
 
 	return FReply::Unhandled();
 }
 
 UWidget* SUMGDesigner::AddPreview(const FGeometry& MyGeometry, const FDragDropEvent& DragDropEvent)
 {
-    TSharedPtr<FWidgetTemplateDragDropOp> DragDropOp = DragDropEvent.GetOperationAs<FWidgetTemplateDragDropOp>();
+	TSharedPtr<FWidgetTemplateDragDropOp> DragDropOp = DragDropEvent.GetOperationAs<FWidgetTemplateDragDropOp>();
 	if ( DragDropOp.IsValid() )
 	{
-        FArrangedWidget ArrangedWidget(SNullWidget::NullWidget, FGeometry());
-        FSelectedWidget Selection = GetWidgetAtCursor(MyGeometry, DragDropEvent, ArrangedWidget);
-        
-        UWidgetBlueprint* BP = GetBlueprint();
-        
-        if ( Selection.IsValid() )
-        {
-            UWidget* Target = Selection.GetPreview();
-            if ( Target->IsA(UPanelWidget::StaticClass()) )
-            {
-                UPanelWidget* Parent = Cast<UPanelWidget>(Target);
-                
-                UWidget* Widget = DragDropOp->Template->Create(BP->WidgetTree);
+		FArrangedWidget ArrangedWidget(SNullWidget::NullWidget, FGeometry());
+		FSelectedWidget Selection = GetWidgetAtCursor(MyGeometry, DragDropEvent, ArrangedWidget);
+		
+		UWidgetBlueprint* BP = GetBlueprint();
+		
+		if ( Selection.IsValid() )
+		{
+			UWidget* Target = Selection.GetPreview();
+			if ( Target->IsA(UPanelWidget::StaticClass()) )
+			{
+				UPanelWidget* Parent = Cast<UPanelWidget>(Target);
+				
+				UWidget* Widget = DragDropOp->Template->Create(BP->WidgetTree);
+				Widget->IsDesignTime(true);
+			
+				FVector2D LocalPosition = ArrangedWidget.Geometry.AbsoluteToLocal(DragDropEvent.GetScreenSpacePosition());
+				Parent->AddChild(Widget, LocalPosition);
+				//@TODO UMG When we add a child blindly we need to default the slot size to the preferred size of the widget if the container supports such things.
+				//@TODO UMG We may need a desired size canvas, where the slots have no size, they only give you position, alternatively, maybe slots that don't clip, so center is still easy.
+			
+				DropPreviewParent = Parent;
+			
+				return Widget;
+			}
+			else if ( BP->WidgetTree->WidgetTemplates.Num() == 1 )
+			{
+				UWidget* Widget = DragDropOp->Template->Create(BP->WidgetTree);
 				Widget->IsDesignTime(true);
-            
-                FVector2D LocalPosition = ArrangedWidget.Geometry.AbsoluteToLocal(DragDropEvent.GetScreenSpacePosition());
-                Parent->AddChild(Widget, LocalPosition);
-                //@TODO UMG When we add a child blindly we need to default the slot size to the preferred size of the widget if the container supports such things.
-                //@TODO UMG We may need a desired size canvas, where the slots have no size, they only give you position, alternatively, maybe slots that don't clip, so center is still easy.
-            
-                DropPreviewParent = Parent;
-            
-                return Widget;
-            }
-            else if ( BP->WidgetTree->WidgetTemplates.Num() == 1 )
-            {
-                UWidget* Widget = DragDropOp->Template->Create(BP->WidgetTree);
-                Widget->IsDesignTime(true);
 				
-                DropPreviewParent = NULL;
-                
-                return Widget;
-            }
-        }
-    }
-    
-    return NULL;
+				DropPreviewParent = NULL;
+				
+				return Widget;
+			}
+		}
+	}
+	
+	return NULL;
 }
 
 bool SUMGDesigner::AddToTemplate(const FGeometry& MyGeometry, const FDragDropEvent& DragDropEvent)
 {
-    TSharedPtr<FWidgetTemplateDragDropOp> DragDropOp = DragDropEvent.GetOperationAs<FWidgetTemplateDragDropOp>();
+	TSharedPtr<FWidgetTemplateDragDropOp> DragDropOp = DragDropEvent.GetOperationAs<FWidgetTemplateDragDropOp>();
 	if ( DragDropOp.IsValid() )
 	{
-        FArrangedWidget ArrangedWidget(SNullWidget::NullWidget, FGeometry());
-        FSelectedWidget Selection = GetWidgetAtCursor(MyGeometry, DragDropEvent, ArrangedWidget);
-        
-        UWidgetBlueprint* BP = GetBlueprint();
-        
-        if ( Selection.IsValid() )
-        {
-            UWidget* Target = Selection.GetTemplate();
-            if ( Target->IsA(UPanelWidget::StaticClass()) )
-            {
-                UPanelWidget* Parent = Cast<UPanelWidget>(Target);
-                
-                UWidget* Widget = DragDropOp->Template->Create(BP->WidgetTree);
+		FArrangedWidget ArrangedWidget(SNullWidget::NullWidget, FGeometry());
+		FSelectedWidget Selection = GetWidgetAtCursor(MyGeometry, DragDropEvent, ArrangedWidget);
+		
+		UWidgetBlueprint* BP = GetBlueprint();
+		
+		if ( Selection.IsValid() )
+		{
+			UWidget* Target = Selection.GetTemplate();
+			if ( Target->IsA(UPanelWidget::StaticClass()) )
+			{
+				UPanelWidget* Parent = Cast<UPanelWidget>(Target);
+
+				const FScopedTransaction Transaction(LOCTEXT("Designer_AddWidget", "Add Widget"));
+				Parent->SetFlags(RF_Transactional);
+				Parent->Modify();
+				
+				UWidget* Widget = DragDropOp->Template->Create(BP->WidgetTree);
 				Widget->IsDesignTime(true);
-                
-                FVector2D LocalPosition = ArrangedWidget.Geometry.AbsoluteToLocal(DragDropEvent.GetScreenSpacePosition());
-                Parent->AddChild(Widget, LocalPosition);
-                //@TODO UMG When we add a child blindly we need to default the slot size to the preferred size of the widget if the container supports such things.
-                //@TODO UMG We may need a desired size canvas, where the slots have no size, they only give you position, alternatively, maybe slots that don't clip, so center is still easy.
-                
-                // Update the selected template to be the newly created one.
-                CurrentSelection = FSelectedWidget::FromTemplate(BlueprintEditor.Pin(), Widget);
-                
-                return true;
-            }
-        }
-        
-        if ( BP->WidgetTree->WidgetTemplates.Num() == 0 )
-        {
-            // TODO UMG This method isn't great, maybe the user widget should just be a canvas.
-            
-            // Add it to the root if there are no other widgets to add it to.
-            UWidget* Widget = DragDropOp->Template->Create(BP->WidgetTree);
+				
+				FVector2D LocalPosition = ArrangedWidget.Geometry.AbsoluteToLocal(DragDropEvent.GetScreenSpacePosition());
+				Parent->AddChild(Widget, LocalPosition);
+				//@TODO UMG When we add a child blindly we need to default the slot size to the preferred size of the widget if the container supports such things.
+				//@TODO UMG We may need a desired size canvas, where the slots have no size, they only give you position, alternatively, maybe slots that don't clip, so center is still easy.
+				
+				// Update the selected template to be the newly created one.
+				CurrentSelection = FSelectedWidget::FromTemplate(BlueprintEditor.Pin(), Widget);
+				
+				return true;
+			}
+		}
+		
+		if ( BP->WidgetTree->WidgetTemplates.Num() == 0 )
+		{
+			const FScopedTransaction Transaction(LOCTEXT("Designer_AddWidget", "Add Widget"));
+			BP->WidgetTree->SetFlags(RF_Transactional);
+			BP->WidgetTree->Modify();
+
+			// TODO UMG This method isn't great, maybe the user widget should just be a canvas.
+			
+			// Add it to the root if there are no other widgets to add it to.
+			UWidget* Widget = DragDropOp->Template->Create(BP->WidgetTree);
 			Widget->IsDesignTime(true);
-            
-            CurrentSelection = FSelectedWidget::FromTemplate(BlueprintEditor.Pin(), Widget);
-            
-            return true;
-        }
-    }
-    
-    return false;
+			
+			CurrentSelection = FSelectedWidget::FromTemplate(BlueprintEditor.Pin(), Widget);
+			
+			return true;
+		}
+	}
+	
+	return false;
 }
 
 FReply SUMGDesigner::OnDrop(const FGeometry& MyGeometry, const FDragDropEvent& DragDropEvent)
 {
-    UWidgetBlueprint* BP = GetBlueprint();
-    
-    if (DropPreviewWidget)
-    {
-        if (DropPreviewParent)
-        {
-            DropPreviewParent->RemoveChild(DropPreviewWidget);
-        }
-        
-        BP->WidgetTree->RemoveWidget(DropPreviewWidget);
-        DropPreviewWidget = NULL;
-    }
-    
-    if ( AddToTemplate(MyGeometry, DragDropEvent) )
-    {
-        FBlueprintEditorUtils::MarkBlueprintAsStructurallyModified(BP);
+	UWidgetBlueprint* BP = GetBlueprint();
+	
+	if (DropPreviewWidget)
+	{
+		if (DropPreviewParent)
+		{
+			DropPreviewParent->RemoveChild(DropPreviewWidget);
+		}
+		
+		BP->WidgetTree->RemoveWidget(DropPreviewWidget);
+		DropPreviewWidget = NULL;
+	}
+	
+	if ( AddToTemplate(MyGeometry, DragDropEvent) )
+	{
+		FBlueprintEditorUtils::MarkBlueprintAsStructurallyModified(BP);
 
 		return FReply::Handled();
-    }
-    
+	}
+	
 	return FReply::Unhandled();
 }
 
+void SUMGDesigner::BeginTransaction(const FText& SessionName)
+{
+	if ( ScopedTransaction == NULL )
+	{
+		ScopedTransaction = new FScopedTransaction(SessionName);
+	}
+
+	if ( CurrentSelection.IsValid() )
+	{
+		CurrentSelection.GetPreview()->Modify();
+		CurrentSelection.GetTemplate()->Modify();
+	}
+}
+
+void SUMGDesigner::EndTransaction()
+{
+	if ( ScopedTransaction != NULL )
+	{
+		delete ScopedTransaction;
+		ScopedTransaction = NULL;
+	}
+}
 
 #undef LOCTEXT_NAMESPACE
diff --git a/Engine/Source/Editor/UMGEditor/Private/SUMGDesigner.h b/Engine/Source/Editor/UMGEditor/Private/SUMGDesigner.h
--- a/Engine/Source/Editor/UMGEditor/Private/SUMGDesigner.h
+++ b/Engine/Source/Editor/UMGEditor/Private/SUMGDesigner.h
@@ -52,6 +52,9 @@ private:
 	FVector2D GetSelectionDesignerWidgetsLocation() const;
 	FVector2D GetCachedSelectionDesignerWidgetsLocation() const;
 
+	void BeginTransaction(const FText& SessionName);
+	void EndTransaction();
+
 private:
 	enum DragHandle
 	{
@@ -82,11 +85,11 @@ private:
 private:
 	void DrawDragHandles(const FPaintGeometry& SelectionGeometry, const FSlateRect& MyClippingRect, FSlateWindowElementList& OutDrawElements, int32 LayerId, const FWidgetStyle& InWidgetStyle) const;
 	DragHandle HitTestDragHandles(const FGeometry& AllottedGeometry, const FPointerEvent& PointerEvent) const;
-    
-    UWidget* AddPreview(const FGeometry& MyGeometry, const FDragDropEvent& DragDropEvent);
-    
-    bool AddToTemplate(const FGeometry& MyGeometry, const FDragDropEvent& DragDropEvent);
-    
+	
+	UWidget* AddPreview(const FGeometry& MyGeometry, const FDragDropEvent& DragDropEvent);
+	
+	bool AddToTemplate(const FGeometry& MyGeometry, const FDragDropEvent& DragDropEvent);
+	
 private:
 	TWeakPtr<FWidgetBlueprintEditor> BlueprintEditor;
 
@@ -95,9 +98,9 @@ private:
 
 	FSelectedWidget CurrentSelection;
 	TWeakPtr<SWidget> SelectedWidget;
-    
-    UWidget* DropPreviewWidget;
-    UPanelWidget* DropPreviewParent;
+	
+	UWidget* DropPreviewWidget;
+	UPanelWidget* DropPreviewParent;
 
 	TSharedPtr<SBorder> PreviewSurface;
 	TSharedPtr<SCanvas> ExtensionWidgetCanvas;
@@ -105,4 +108,6 @@ private:
 	DragHandle CurrentHandle;
 
 	FVector2D CachedDesignerWidgetLocation;
+
+	FScopedTransaction* ScopedTransaction;
 };
diff --git a/Engine/Source/Editor/UMGEditor/Private/SUMGEditorTree.cpp b/Engine/Source/Editor/UMGEditor/Private/SUMGEditorTree.cpp
--- a/Engine/Source/Editor/UMGEditor/Private/SUMGEditorTree.cpp
+++ b/Engine/Source/Editor/UMGEditor/Private/SUMGEditorTree.cpp
@@ -223,6 +223,10 @@ void SUMGEditorTree::DeleteSelected()
 	{
 		UWidgetBlueprint* BP = GetBlueprint();
 
+		const FScopedTransaction Transaction(LOCTEXT("RemoveWidget", "Remove Widget"));
+		BP->WidgetTree->SetFlags(RF_Transactional);
+		BP->WidgetTree->Modify();
+
 		bool bRemoved = false;
 		for ( UWidget* Item : SelectedItems )
 		{
diff --git a/Engine/Source/Editor/UMGEditor/Private/UMGEditorPrivatePCH.h b/Engine/Source/Editor/UMGEditor/Private/UMGEditorPrivatePCH.h
--- a/Engine/Source/Editor/UMGEditor/Private/UMGEditorPrivatePCH.h
+++ b/Engine/Source/Editor/UMGEditor/Private/UMGEditorPrivatePCH.h
@@ -9,6 +9,8 @@
 
 #include "BlueprintEditorUtils.h"
 
+#include "ScopedTransaction.h"
+
 #include "UMG.h"
 
 #include "UMGEditorClasses.h"
diff --git a/Engine/Source/Editor/UMGEditor/Public/DesignerExtension.h b/Engine/Source/Editor/UMGEditor/Public/DesignerExtension.h
--- a/Engine/Source/Editor/UMGEditor/Public/DesignerExtension.h
+++ b/Engine/Source/Editor/UMGEditor/Public/DesignerExtension.h
@@ -35,13 +35,25 @@ private:
 class FDesignerExtension : public TSharedFromThis<FDesignerExtension>
 {
 public:
+	FDesignerExtension();
+
 	virtual void Initialize(UWidgetBlueprint* InBlueprint);
 
 	virtual void BuildWidgetsForSelection(const TArray< FSelectedWidget >& Selection, TArray< TSharedRef<SWidget> >& Widgets) = 0;
 
 	FName GetExtensionId() const;
 
+protected:
+	void BeginTransaction(const FText& SessionName);
+
+	void EndTransaction();
+
 protected:
 	FName ExtensionId;
 	UWidgetBlueprint* Blueprint;
+
+	TArray< FSelectedWidget > SelectionCache;
+
+private:
+	FScopedTransaction* ScopedTransaction;
 };
diff --git a/Engine/Source/Runtime/UMG/Private/Components/EditableText.cpp b/Engine/Source/Runtime/UMG/Private/Components/EditableText.cpp
--- a/Engine/Source/Runtime/UMG/Private/Components/EditableText.cpp
+++ b/Engine/Source/Runtime/UMG/Private/Components/EditableText.cpp
@@ -48,6 +48,10 @@ TSharedRef<SWidget> UEditableText::RebuildWidget()
 	.RevertTextOnEscape(RevertTextOnEscape)
 	.ClearKeyboardFocusOnCommit(ClearKeyboardFocusOnCommit)
 	.SelectAllTextOnCommit(SelectAllTextOnCommit)
+	.BackgroundImageSelected(BackgroundImageSelected ? TAttribute<const FSlateBrush*>(&BackgroundImageSelected->Brush) : TAttribute<const FSlateBrush*>())
+	.BackgroundImageSelectionTarget(BackgroundImageSelectionTarget ? TAttribute<const FSlateBrush*>(&BackgroundImageSelectionTarget->Brush) : TAttribute<const FSlateBrush*>())
+	.BackgroundImageComposing(BackgroundImageComposing ? TAttribute<const FSlateBrush*>(&BackgroundImageComposing->Brush) : TAttribute<const FSlateBrush*>())
+	.CaretImage(CaretImage ? TAttribute<const FSlateBrush*>(&CaretImage->Brush) : TAttribute<const FSlateBrush*>())
 	.OnTextChanged(BIND_UOBJECT_DELEGATE(FOnTextChanged, HandleOnTextChanged))
 	.OnTextCommitted(BIND_UOBJECT_DELEGATE(FOnTextCommitted, HandleOnTextCommitted))
 	;
@@ -62,6 +66,8 @@ void UEditableText::SyncronizeProperties()
 	MyEditableText->SetIsReadOnly(IsReadOnly);
 	MyEditableText->SetIsPassword(IsPassword);
 	MyEditableText->SetColorAndOpacity(ColorAndOpacity);
+
+	// TODO UMG Complete making all properties settable on SEditableText
 }
 
 void UEditableText::HandleOnTextChanged(const FText& Text)
diff --git a/Engine/Source/Runtime/UMG/Private/Components/ListView.cpp b/Engine/Source/Runtime/UMG/Private/Components/ListView.cpp
--- a/Engine/Source/Runtime/UMG/Private/Components/ListView.cpp
+++ b/Engine/Source/Runtime/UMG/Private/Components/ListView.cpp
@@ -34,7 +34,7 @@ TSharedRef<SWidget> UListView::RebuildWidget()
 		//	);
 		;
 
-	return MyListView.ToSharedRef();
+	return BuildDesignTimeWidget( MyListView.ToSharedRef() );
 }
 
 TSharedRef<ITableRow> UListView::HandleOnGenerateRow(UObject* Item, const TSharedRef< STableViewBase >& OwnerTable) const
diff --git a/Engine/Source/Runtime/UMG/Private/Components/TileView.cpp b/Engine/Source/Runtime/UMG/Private/Components/TileView.cpp
--- a/Engine/Source/Runtime/UMG/Private/Components/TileView.cpp
+++ b/Engine/Source/Runtime/UMG/Private/Components/TileView.cpp
@@ -19,7 +19,7 @@ UTileView::UTileView(const FPostConstructInitializeProperties& PCIP)
 
 TSharedRef<SWidget> UTileView::RebuildWidget()
 {
-	return SNew(STileView< UObject* >)
+	MyTileView = SNew(STileView< UObject* >)
 		.SelectionMode(SelectionMode)
 		.ListItemsSource(&Items)
 		.ItemWidth(ItemWidth)
@@ -35,6 +35,8 @@ TSharedRef<SWidget> UTileView::RebuildWidget()
 		//		+ SHeaderRow::Column(TEXT("Socket"))
 		//	);
 		;
+
+	return BuildDesignTimeWidget( MyTileView.ToSharedRef() );
 }
 
 TSharedRef<ITableRow> UTileView::HandleOnGenerateTile(UObject* Item, const TSharedRef< STableViewBase >& OwnerTable) const
@@ -59,6 +61,21 @@ TSharedRef<ITableRow> UTileView::HandleOnGenerateTile(UObject* Item, const TShar
 		];
 }
 
+void UTileView::SetItemWidth(float Width)
+{
+	MyTileView->SetItemWidth(Width);
+}
+
+void UTileView::SetItemHeight(float Height)
+{
+	MyTileView->SetItemHeight(Height);
+}
+
+void UTileView::RequestListRefresh()
+{
+	MyTileView->RequestListRefresh();
+}
+
 /////////////////////////////////////////////////////
 
 #undef LOCTEXT_NAMESPACE
diff --git a/Engine/Source/Runtime/UMG/Private/Components/Widget.cpp b/Engine/Source/Runtime/UMG/Private/Components/Widget.cpp
--- a/Engine/Source/Runtime/UMG/Private/Components/Widget.cpp
+++ b/Engine/Source/Runtime/UMG/Private/Components/Widget.cpp
@@ -128,6 +128,16 @@ void UWidget::PostEditChangeProperty(FPropertyChangedEvent& PropertyChangedEvent
 
 #endif
 
+bool UWidget::Modify(bool bAlwaysMarkDirty)
+{
+	bool Modified = Super::Modify(bAlwaysMarkDirty);
+
+	if ( Slot )
+		Modified &= Slot->Modify(bAlwaysMarkDirty);
+
+	return Modified;
+}
+
 TSharedRef<SWidget> UWidget::RebuildWidget()
 {
 	ensureMsg(false, TEXT("You must implement RebuildWidget() in your child class"));
diff --git a/Engine/Source/Runtime/UMG/Private/WidgetBlueprint.cpp b/Engine/Source/Runtime/UMG/Private/WidgetBlueprint.cpp
--- a/Engine/Source/Runtime/UMG/Private/WidgetBlueprint.cpp
+++ b/Engine/Source/Runtime/UMG/Private/WidgetBlueprint.cpp
@@ -22,6 +22,7 @@ UWidgetBlueprint::UWidgetBlueprint(const FPostConstructInitializeProperties& PCI
 	: Super(PCIP)
 {
 	WidgetTree = ConstructObject<UWidgetTree>(UWidgetTree::StaticClass(), this);
+	WidgetTree->SetFlags(RF_Transactional);
 }
 
 void UWidgetBlueprint::PostLoad()
diff --git a/Engine/Source/Runtime/UMG/Private/WidgetBlueprintGeneratedClass.cpp b/Engine/Source/Runtime/UMG/Private/WidgetBlueprintGeneratedClass.cpp
--- a/Engine/Source/Runtime/UMG/Private/WidgetBlueprintGeneratedClass.cpp
+++ b/Engine/Source/Runtime/UMG/Private/WidgetBlueprintGeneratedClass.cpp
@@ -11,6 +11,7 @@ UWidgetBlueprintGeneratedClass::UWidgetBlueprintGeneratedClass(const FPostConstr
 	: Super(PCIP)
 {
 	WidgetTree = ConstructObject<UWidgetTree>(UWidgetTree::StaticClass(), this);
+	WidgetTree->SetFlags(RF_Transactional);
 }
 
 void UWidgetBlueprintGeneratedClass::Link(FArchive& Ar, bool bRelinkExistingProperties)
diff --git a/Engine/Source/Runtime/UMG/Private/WidgetTree.cpp b/Engine/Source/Runtime/UMG/Private/WidgetTree.cpp
--- a/Engine/Source/Runtime/UMG/Private/WidgetTree.cpp
+++ b/Engine/Source/Runtime/UMG/Private/WidgetTree.cpp
@@ -12,9 +12,10 @@ UWidgetTree::UWidgetTree(const FPostConstructInitializeProperties& PCIP)
 
 void UWidgetTree::RenameWidget(UWidget* Widget, FString& NewName)
 {
+	Widget->Modify();
 	Widget->Rename(*NewName);
 
-	// TODO Update nodes in the blueprint!
+	// TODO UMG Update nodes in the blueprint!
 }
 
 UWidget* UWidgetTree::FindWidget(FString& Name) const
@@ -73,6 +74,7 @@ bool UWidgetTree::RemoveWidgetRecursive(UWidget* InRemovedWidget)
 		UPanelWidget* NonLeafTemplate = Cast<UPanelWidget>(Template);
 		if ( NonLeafTemplate )
 		{
+			NonLeafTemplate->Modify();
 			if ( NonLeafTemplate->RemoveChild(InRemovedWidget) )
 			{
 				return true;
diff --git a/Engine/Source/Runtime/UMG/Public/Components/CheckBox.h b/Engine/Source/Runtime/UMG/Public/Components/CheckBox.h
--- a/Engine/Source/Runtime/UMG/Public/Components/CheckBox.h
+++ b/Engine/Source/Runtime/UMG/Public/Components/CheckBox.h
@@ -17,6 +17,42 @@ protected:
 	UPROPERTY(EditDefaultsOnly, BlueprintReadOnly, Category=Style, meta=( DisplayThumbnail = "true" ))
 	USlateWidgetStyleAsset* Style;
 
+	/** Image to use when the checkbox is unchecked */
+	UPROPERTY(EditDefaultsOnly, Category="Style", meta=( DisplayThumbnail = "true" ), AdvancedDisplay)
+	USlateBrushAsset* UncheckedImage;
+	
+	/** Image to use when the checkbox is unchecked and hovered */
+	UPROPERTY(EditDefaultsOnly, Category="Style", meta=( DisplayThumbnail = "true" ), AdvancedDisplay)
+	USlateBrushAsset* UncheckedHoveredImage;
+	
+	/** Image to use when the checkbox is unchecked and pressed */
+	UPROPERTY(EditDefaultsOnly, Category="Style", meta=( DisplayThumbnail = "true" ), AdvancedDisplay)
+	USlateBrushAsset* UncheckedPressedImage;
+	
+	/** Image to use when the checkbox is checked */
+	UPROPERTY(EditDefaultsOnly, Category="Style", meta=( DisplayThumbnail = "true" ), AdvancedDisplay)
+	USlateBrushAsset* CheckedImage;
+	
+	/** Image to use when the checkbox is checked and hovered */
+	UPROPERTY(EditDefaultsOnly, Category="Style", meta=( DisplayThumbnail = "true" ), AdvancedDisplay)
+	USlateBrushAsset* CheckedHoveredImage;
+	
+	/** Image to use when the checkbox is checked and pressed */
+	UPROPERTY(EditDefaultsOnly, Category="Style", meta=( DisplayThumbnail = "true" ), AdvancedDisplay)
+	USlateBrushAsset* CheckedPressedImage;
+	
+	/** Image to use when the checkbox is in an ambiguous state and hovered */
+	UPROPERTY(EditDefaultsOnly, Category="Style", meta=( DisplayThumbnail = "true" ), AdvancedDisplay)
+	USlateBrushAsset* UndeterminedImage;
+	
+	/** Image to use when the checkbox is checked and hovered */
+	UPROPERTY(EditDefaultsOnly, Category="Style", meta=( DisplayThumbnail = "true" ), AdvancedDisplay)
+	USlateBrushAsset* UndeterminedHoveredImage;
+	
+	/** Image to use when the checkbox is in an ambiguous state and pressed */
+	UPROPERTY(EditDefaultsOnly, Category="Style", meta=( DisplayThumbnail = "true" ), AdvancedDisplay)
+	USlateBrushAsset* UndeterminedPressedImage;
+
 	/** Called when the checked state has changed */
 	UPROPERTY(BlueprintAssignable)
 	FOnCheckBoxComponentStateChanged OnCheckStateChanged;
@@ -55,42 +91,6 @@ protected:
 	UPROPERTY(EditDefaultsOnly, Category="Sound")
 	FSlateSound HoveredSound;
 	
-	/** Image to use when the checkbox is unchecked */
-	UPROPERTY(EditDefaultsOnly, Category="Style Overrides", meta=( DisplayThumbnail = "true" ))
-	USlateBrushAsset* UncheckedImage;
-	
-	/** Image to use when the checkbox is unchecked and hovered */
-	UPROPERTY(EditDefaultsOnly, Category="Style Overrides", meta=( DisplayThumbnail = "true" ))
-	USlateBrushAsset* UncheckedHoveredImage;
-	
-	/** Image to use when the checkbox is unchecked and pressed */
-	UPROPERTY(EditDefaultsOnly, Category="Style Overrides", meta=( DisplayThumbnail = "true" ))
-	USlateBrushAsset* UncheckedPressedImage;
-	
-	/** Image to use when the checkbox is checked */
-	UPROPERTY(EditDefaultsOnly, Category="Style Overrides", meta=( DisplayThumbnail = "true" ))
-	USlateBrushAsset* CheckedImage;
-	
-	/** Image to use when the checkbox is checked and hovered */
-	UPROPERTY(EditDefaultsOnly, Category="Style Overrides", meta=( DisplayThumbnail = "true" ))
-	USlateBrushAsset* CheckedHoveredImage;
-	
-	/** Image to use when the checkbox is checked and pressed */
-	UPROPERTY(EditDefaultsOnly, Category="Style Overrides", meta=( DisplayThumbnail = "true" ))
-	USlateBrushAsset* CheckedPressedImage;
-	
-	/** Image to use when the checkbox is in an ambiguous state and hovered */
-	UPROPERTY(EditDefaultsOnly, Category="Style Overrides", meta=( DisplayThumbnail = "true" ))
-	USlateBrushAsset* UndeterminedImage;
-	
-	/** Image to use when the checkbox is checked and hovered */
-	UPROPERTY(EditDefaultsOnly, Category="Style Overrides", meta=( DisplayThumbnail = "true" ))
-	USlateBrushAsset* UndeterminedHoveredImage;
-	
-	/** Image to use when the checkbox is in an ambiguous state and pressed */
-	UPROPERTY(EditDefaultsOnly, Category="Style Overrides", meta=( DisplayThumbnail = "true" ))
-	USlateBrushAsset* UndeterminedPressedImage;
-	
 	/** Returns true if this button is currently pressed */
 	UFUNCTION(BlueprintCallable, Category="Widget")
 	bool IsPressed() const;
diff --git a/Engine/Source/Runtime/UMG/Public/Components/EditableText.h b/Engine/Source/Runtime/UMG/Public/Components/EditableText.h
--- a/Engine/Source/Runtime/UMG/Public/Components/EditableText.h
+++ b/Engine/Source/Runtime/UMG/Public/Components/EditableText.h
@@ -27,6 +27,22 @@ protected:
 	UPROPERTY(EditDefaultsOnly, Category=Style, meta=( DisplayThumbnail = "true" ))
 	USlateWidgetStyleAsset* Style;
 
+	/** Background image for the selected text (overrides Style) */
+	UPROPERTY(EditDefaultsOnly, Category="Style", meta=( DisplayThumbnail = "true" ), AdvancedDisplay)
+	USlateBrushAsset* BackgroundImageSelected;
+
+	/** Background image for the selection targeting effect (overrides Style) */
+	UPROPERTY(EditDefaultsOnly, Category="Style", meta=( DisplayThumbnail = "true" ), AdvancedDisplay)
+	USlateBrushAsset* BackgroundImageSelectionTarget;
+
+	/** Background image for the composing text (overrides Style) */
+	UPROPERTY(EditDefaultsOnly, Category="Style", meta=( DisplayThumbnail = "true" ), AdvancedDisplay)
+	USlateBrushAsset* BackgroundImageComposing;
+
+	/** Image brush used for the caret (overrides Style) */
+	UPROPERTY(EditDefaultsOnly, Category="Style", meta=( DisplayThumbnail = "true" ), AdvancedDisplay)
+	USlateBrushAsset* CaretImage;
+
 	/** Font color and opacity (overrides Style) */
 	UPROPERTY(EditDefaultsOnly, Category=Appearance)
 	FSlateFontInfo Font;
@@ -67,18 +83,6 @@ protected:
 	UPROPERTY(EditDefaultsOnly, Category=Behavior)
 	bool SelectAllTextOnCommit;
 
-	///** Background image for the selected text (overrides Style) */
-	//SLATE_ATTRIBUTE(const FSlateBrush*, BackgroundImageSelected)
-
-	///** Background image for the selection targeting effect (overrides Style) */
-	//SLATE_ATTRIBUTE(const FSlateBrush*, BackgroundImageSelectionTarget)
-
-	///** Background image for the composing text (overrides Style) */
-	//SLATE_ATTRIBUTE(const FSlateBrush*, BackgroundImageComposing)
-
-	///** Image brush used for the caret (overrides Style) */
-	//SLATE_ATTRIBUTE(const FSlateBrush*, CaretImage)
-
 	/** Called whenever the text is changed interactively by the user */
 	UPROPERTY(BlueprintAssignable)
 	FOnEditableTextChangedEvent OnTextChanged;
diff --git a/Engine/Source/Runtime/UMG/Public/Components/TileView.h b/Engine/Source/Runtime/UMG/Public/Components/TileView.h
--- a/Engine/Source/Runtime/UMG/Public/Components/TileView.h
+++ b/Engine/Source/Runtime/UMG/Public/Components/TileView.h
@@ -10,37 +10,45 @@ class UMG_API UTileView : public UTableViewBase
 {
 	GENERATED_UCLASS_BODY()
 
-public:
-
-	UPROPERTY(EditDefaultsOnly, Category=Content)
+	/**  */
+	UPROPERTY(EditDefaultsOnly, Category=Appearance)
 	float ItemWidth;
 
-	UPROPERTY(EditDefaultsOnly, Category=Content)
+	/**  */
+	UPROPERTY(EditDefaultsOnly, Category=Appearance)
 	float ItemHeight;
 
+	/**  */
 	UPROPERTY(EditAnywhere, BlueprintReadWrite, Category=Content)
 	TArray<UObject*> Items;
 
+	/**  */
 	UPROPERTY(EditDefaultsOnly, Category=Content)
 	TEnumAsByte<ESelectionMode::Type> SelectionMode;
 
+	/**  */
 	UPROPERTY(EditDefaultsOnly, Category=Events)
 	FOnGenerateRowUObject OnGenerateTileEvent;
 
-protected:
-	TSharedPtr< STileView<UObject*> > ListWidget() const
-	{
-		if ( MyWidget.IsValid() )
-		{
-			return StaticCastSharedRef< STileView<UObject*> >(MyWidget.ToSharedRef());
-		}
+	/** Set item width */
+	UFUNCTION(BlueprintCallable, Category="Appearance")
+	void SetItemWidth(float Width);
 
-		return TSharedPtr< STileView<UObject*> >();
-	}
+	/** Set item height */
+	UFUNCTION(BlueprintCallable, Category="Appearance")
+	void SetItemHeight(float Height);
 
+	/** Refreshes the list */
+	UFUNCTION(BlueprintCallable, Category="Behavior")
+	void RequestListRefresh();
+
+protected:
 	TSharedRef<ITableRow> HandleOnGenerateTile(UObject* Item, const TSharedRef< STableViewBase >& OwnerTable) const;
 
 	// UWidget interface
 	virtual TSharedRef<SWidget> RebuildWidget() override;
 	// End of UWidget interface
+
+protected:
+	TSharedPtr< STileView< UObject* > > MyTileView;
 };
diff --git a/Engine/Source/Runtime/UMG/Public/Components/Widget.h b/Engine/Source/Runtime/UMG/Public/Components/Widget.h
--- a/Engine/Source/Runtime/UMG/Public/Components/Widget.h
+++ b/Engine/Source/Runtime/UMG/Public/Components/Widget.h
@@ -135,6 +135,9 @@ public:
 	
 	/** Sets that this widget is being designed */
 	void IsDesignTime(bool bInDesignTime);
+
+	/** Mark this object as modified, also mark the slot as modified. */
+	virtual bool Modify(bool bAlwaysMarkDirty = true);
 	
 #if WITH_EDITOR
 	/** Gets a widget representing the tiny preview of the toolbox */
