UMG - Working on syncing selection.  Slots are no longer a transient parenting system.  They're now a serialized part of the data, because they'll be nessesary at runtime.  Adding a do-nothing slot to content widgets to act as a parent pointer proxy.

diff --git a/Engine/Source/Editor/UMGEditor/Private/Extensions/CanvasSlotExtension.cpp b/Engine/Source/Editor/UMGEditor/Private/Extensions/CanvasSlotExtension.cpp
--- a/Engine/Source/Editor/UMGEditor/Private/Extensions/CanvasSlotExtension.cpp
+++ b/Engine/Source/Editor/UMGEditor/Private/Extensions/CanvasSlotExtension.cpp
@@ -10,9 +10,9 @@ FCanvasSlotExtension::FCanvasSlotExtension()
 	ExtensionId = FName(TEXT("CanvasSlot"));
 }
 
-bool FCanvasSlotExtension::IsActive(const TArray< FSelectedWidget >& Selection)
+bool FCanvasSlotExtension::IsActive(const TArray< FWidgetReference >& Selection)
 {
-	for ( const FSelectedWidget& Widget : Selection )
+	for ( const FWidgetReference& Widget : Selection )
 	{
 		if ( !Widget.GetTemplate()->Slot || !Widget.GetTemplate()->Slot->IsA(UCanvasPanelSlot::StaticClass()) )
 		{
@@ -23,7 +23,7 @@ bool FCanvasSlotExtension::IsActive(const TArray< FSelectedWidget >& Selection)
 	return Selection.Num() == 1;
 }
 
-void FCanvasSlotExtension::BuildWidgetsForSelection(const TArray< FSelectedWidget >& Selection, TArray< TSharedRef<SWidget> >& Widgets)
+void FCanvasSlotExtension::BuildWidgetsForSelection(const TArray< FWidgetReference >& Selection, TArray< TSharedRef<SWidget> >& Widgets)
 {
 	SelectionCache = Selection;
 
@@ -72,7 +72,7 @@ FReply FCanvasSlotExtension::HandleDragging(const FGeometry& Geometry, const FPo
 {
 	if ( bDragging )
 	{
-		for ( FSelectedWidget& Selection : SelectionCache )
+		for ( FWidgetReference& Selection : SelectionCache )
 		{
 			MoveByAmount(Selection.GetPreview(), Event.GetCursorDelta());
 			MoveByAmount(Selection.GetTemplate(), Event.GetCursorDelta());
diff --git a/Engine/Source/Editor/UMGEditor/Private/Extensions/CanvasSlotExtension.h b/Engine/Source/Editor/UMGEditor/Private/Extensions/CanvasSlotExtension.h
--- a/Engine/Source/Editor/UMGEditor/Private/Extensions/CanvasSlotExtension.h
+++ b/Engine/Source/Editor/UMGEditor/Private/Extensions/CanvasSlotExtension.h
@@ -9,9 +9,9 @@ class FCanvasSlotExtension : public FDesignerExtension
 public:
 	FCanvasSlotExtension();
 
-	bool IsActive(const TArray< FSelectedWidget >& Selection);
+	bool IsActive(const TArray< FWidgetReference >& Selection);
 	
-	virtual void BuildWidgetsForSelection(const TArray< FSelectedWidget >& Selection, TArray< TSharedRef<SWidget> >& Widgets) override;
+	virtual void BuildWidgetsForSelection(const TArray< FWidgetReference >& Selection, TArray< TSharedRef<SWidget> >& Widgets) override;
 
 private:
 
diff --git a/Engine/Source/Editor/UMGEditor/Private/Extensions/DesignerExtension.cpp b/Engine/Source/Editor/UMGEditor/Private/Extensions/DesignerExtension.cpp
--- a/Engine/Source/Editor/UMGEditor/Private/Extensions/DesignerExtension.cpp
+++ b/Engine/Source/Editor/UMGEditor/Private/Extensions/DesignerExtension.cpp
@@ -6,26 +6,26 @@
 
 #define LOCTEXT_NAMESPACE "UMG"
 
-// FSelectedWidget
+// FWidgetReference
 
-FSelectedWidget::FSelectedWidget()
+FWidgetReference::FWidgetReference()
 	: WidgetEditor()
 	, TemplateWidget(NULL)
 {
 }
 
-FSelectedWidget::FSelectedWidget(TSharedPtr<FWidgetBlueprintEditor> WidgetEditor, UWidget* TemplateWidget)
+FWidgetReference::FWidgetReference(TSharedPtr<FWidgetBlueprintEditor> WidgetEditor, UWidget* TemplateWidget)
 	: WidgetEditor(WidgetEditor)
 	, TemplateWidget(TemplateWidget)
 {
 }
 
-FSelectedWidget FSelectedWidget::FromTemplate(TSharedPtr<FWidgetBlueprintEditor> WidgetEditor, UWidget* TemplateWidget)
+FWidgetReference FWidgetReference::FromTemplate(TSharedPtr<FWidgetBlueprintEditor> WidgetEditor, UWidget* TemplateWidget)
 {
-	return FSelectedWidget(WidgetEditor, TemplateWidget);
+	return FWidgetReference(WidgetEditor, TemplateWidget);
 }
 
-FSelectedWidget FSelectedWidget::FromPreview(TSharedPtr<FWidgetBlueprintEditor> WidgetEditor, UWidget* PreviewWidget)
+FWidgetReference FWidgetReference::FromPreview(TSharedPtr<FWidgetBlueprintEditor> WidgetEditor, UWidget* PreviewWidget)
 {
 	if ( WidgetEditor.IsValid() )
 	{
@@ -42,20 +42,20 @@ FSelectedWidget FSelectedWidget::FromPreview(TSharedPtr<FWidgetBlueprintEditor>
 		}
 	}
 
-	return FSelectedWidget();
+	return FWidgetReference();
 }
 
-bool FSelectedWidget::IsValid() const
+bool FWidgetReference::IsValid() const
 {
 	return TemplateWidget != NULL && GetPreview();
 }
 
-UWidget* FSelectedWidget::GetTemplate() const
+UWidget* FWidgetReference::GetTemplate() const
 {
 	return TemplateWidget;
 }
 
-UWidget* FSelectedWidget::GetPreview() const
+UWidget* FWidgetReference::GetPreview() const
 {
 	if ( WidgetEditor.IsValid() )
 	{
@@ -96,7 +96,7 @@ void FDesignerExtension::BeginTransaction(const FText& SessionName)
 		ScopedTransaction = new FScopedTransaction(SessionName);
 	}
 
-	for ( FSelectedWidget& Selection : SelectionCache )
+	for ( FWidgetReference& Selection : SelectionCache )
 	{
 		Selection.GetPreview()->Modify();
 		Selection.GetTemplate()->Modify();
diff --git a/Engine/Source/Editor/UMGEditor/Private/Extensions/HorizontalSlotExtension.cpp b/Engine/Source/Editor/UMGEditor/Private/Extensions/HorizontalSlotExtension.cpp
--- a/Engine/Source/Editor/UMGEditor/Private/Extensions/HorizontalSlotExtension.cpp
+++ b/Engine/Source/Editor/UMGEditor/Private/Extensions/HorizontalSlotExtension.cpp
@@ -10,9 +10,9 @@ FHorizontalSlotExtension::FHorizontalSlotExtension()
 	ExtensionId = FName(TEXT("HorizontalSlot"));
 }
 
-bool FHorizontalSlotExtension::IsActive(const TArray< FSelectedWidget >& Selection)
+bool FHorizontalSlotExtension::IsActive(const TArray< FWidgetReference >& Selection)
 {
-	for ( const FSelectedWidget& Widget : Selection )
+	for ( const FWidgetReference& Widget : Selection )
 	{
 		if ( !Widget.GetTemplate()->Slot || !Widget.GetTemplate()->Slot->IsA(UHorizontalBoxSlot::StaticClass()) )
 		{
@@ -23,7 +23,7 @@ bool FHorizontalSlotExtension::IsActive(const TArray< FSelectedWidget >& Selecti
 	return Selection.Num() == 1;
 }
 
-void FHorizontalSlotExtension::BuildWidgetsForSelection(const TArray< FSelectedWidget >& Selection, TArray< TSharedRef<SWidget> >& Widgets)
+void FHorizontalSlotExtension::BuildWidgetsForSelection(const TArray< FWidgetReference >& Selection, TArray< TSharedRef<SWidget> >& Widgets)
 {
 	SelectionCache = Selection;
 
@@ -50,7 +50,7 @@ FReply FHorizontalSlotExtension::HandleShift(int32 ShiftAmount)
 {
 	BeginTransaction(LOCTEXT("MoveWidget", "Move Widget"));
 
-	for ( FSelectedWidget& Selection : SelectionCache )
+	for ( FWidgetReference& Selection : SelectionCache )
 	{
 		ShiftHorizontal(Selection.GetPreview(), ShiftAmount);
 		ShiftHorizontal(Selection.GetTemplate(), ShiftAmount);
diff --git a/Engine/Source/Editor/UMGEditor/Private/Extensions/HorizontalSlotExtension.h b/Engine/Source/Editor/UMGEditor/Private/Extensions/HorizontalSlotExtension.h
--- a/Engine/Source/Editor/UMGEditor/Private/Extensions/HorizontalSlotExtension.h
+++ b/Engine/Source/Editor/UMGEditor/Private/Extensions/HorizontalSlotExtension.h
@@ -9,9 +9,9 @@ class FHorizontalSlotExtension : public FDesignerExtension
 public:
 	FHorizontalSlotExtension();
 
-	bool IsActive(const TArray< FSelectedWidget >& Selection);
+	bool IsActive(const TArray< FWidgetReference >& Selection);
 	
-	virtual void BuildWidgetsForSelection(const TArray< FSelectedWidget >& Selection, TArray< TSharedRef<SWidget> >& Widgets) override;
+	virtual void BuildWidgetsForSelection(const TArray< FWidgetReference >& Selection, TArray< TSharedRef<SWidget> >& Widgets) override;
 
 private:
 
diff --git a/Engine/Source/Editor/UMGEditor/Private/Extensions/UniformGridSlotExtension.cpp b/Engine/Source/Editor/UMGEditor/Private/Extensions/UniformGridSlotExtension.cpp
--- a/Engine/Source/Editor/UMGEditor/Private/Extensions/UniformGridSlotExtension.cpp
+++ b/Engine/Source/Editor/UMGEditor/Private/Extensions/UniformGridSlotExtension.cpp
@@ -9,9 +9,9 @@ FUniformGridSlotExtension::FUniformGridSlotExtension()
 	ExtensionId = FName(TEXT("UniformGridSlot"));
 }
 
-bool FUniformGridSlotExtension::IsActive(const TArray< FSelectedWidget >& Selection)
+bool FUniformGridSlotExtension::IsActive(const TArray< FWidgetReference >& Selection)
 {
-	for ( const FSelectedWidget& Widget : Selection )
+	for ( const FWidgetReference& Widget : Selection )
 	{
 		if ( !Widget.GetTemplate()->Slot || !Widget.GetTemplate()->Slot->IsA(UUniformGridSlot::StaticClass()) )
 		{
@@ -22,7 +22,7 @@ bool FUniformGridSlotExtension::IsActive(const TArray< FSelectedWidget >& Select
 	return Selection.Num() == 1;
 }
 
-void FUniformGridSlotExtension::BuildWidgetsForSelection(const TArray< FSelectedWidget >& Selection, TArray< TSharedRef<SWidget> >& Widgets)
+void FUniformGridSlotExtension::BuildWidgetsForSelection(const TArray< FWidgetReference >& Selection, TArray< TSharedRef<SWidget> >& Widgets)
 {
 	SelectionCache = Selection;
 
@@ -61,7 +61,7 @@ FReply FUniformGridSlotExtension::HandleShiftRow(int32 ShiftAmount)
 {
 	BeginTransaction(LOCTEXT("MoveWidget", "Move Widget"));
 
-	for ( FSelectedWidget& Selection : SelectionCache )
+	for ( FWidgetReference& Selection : SelectionCache )
 	{
 		ShiftRow(Selection.GetPreview(), ShiftAmount);
 		ShiftRow(Selection.GetTemplate(), ShiftAmount);
@@ -78,7 +78,7 @@ FReply FUniformGridSlotExtension::HandleShiftColumn(int32 ShiftAmount)
 {
 	BeginTransaction(LOCTEXT("MoveWidget", "Move Widget"));
 
-	for ( FSelectedWidget& Selection : SelectionCache )
+	for ( FWidgetReference& Selection : SelectionCache )
 	{
 		ShiftColumn(Selection.GetPreview(), ShiftAmount);
 		ShiftColumn(Selection.GetTemplate(), ShiftAmount);
diff --git a/Engine/Source/Editor/UMGEditor/Private/Extensions/UniformGridSlotExtension.h b/Engine/Source/Editor/UMGEditor/Private/Extensions/UniformGridSlotExtension.h
--- a/Engine/Source/Editor/UMGEditor/Private/Extensions/UniformGridSlotExtension.h
+++ b/Engine/Source/Editor/UMGEditor/Private/Extensions/UniformGridSlotExtension.h
@@ -9,9 +9,9 @@ class FUniformGridSlotExtension : public FDesignerExtension
 public:
 	FUniformGridSlotExtension();
 
-	bool IsActive(const TArray< FSelectedWidget >& Selection);
+	bool IsActive(const TArray< FWidgetReference >& Selection);
 	
-	virtual void BuildWidgetsForSelection(const TArray< FSelectedWidget >& Selection, TArray< TSharedRef<SWidget> >& Widgets) override;
+	virtual void BuildWidgetsForSelection(const TArray< FWidgetReference >& Selection, TArray< TSharedRef<SWidget> >& Widgets) override;
 
 private:
 
diff --git a/Engine/Source/Editor/UMGEditor/Private/Extensions/VerticalSlotExtension.cpp b/Engine/Source/Editor/UMGEditor/Private/Extensions/VerticalSlotExtension.cpp
--- a/Engine/Source/Editor/UMGEditor/Private/Extensions/VerticalSlotExtension.cpp
+++ b/Engine/Source/Editor/UMGEditor/Private/Extensions/VerticalSlotExtension.cpp
@@ -10,9 +10,9 @@ FVerticalSlotExtension::FVerticalSlotExtension()
 	ExtensionId = FName(TEXT("VerticalSlot"));
 }
 
-bool FVerticalSlotExtension::IsActive(const TArray< FSelectedWidget >& Selection)
+bool FVerticalSlotExtension::IsActive(const TArray< FWidgetReference >& Selection)
 {
-	for ( const FSelectedWidget& Widget : Selection )
+	for ( const FWidgetReference& Widget : Selection )
 	{
 		if ( !Widget.GetTemplate()->Slot || !Widget.GetTemplate()->Slot->IsA(UVerticalBoxSlot::StaticClass()) )
 		{
@@ -23,7 +23,7 @@ bool FVerticalSlotExtension::IsActive(const TArray< FSelectedWidget >& Selection
 	return Selection.Num() == 1;
 }
 
-void FVerticalSlotExtension::BuildWidgetsForSelection(const TArray< FSelectedWidget >& Selection, TArray< TSharedRef<SWidget> >& Widgets)
+void FVerticalSlotExtension::BuildWidgetsForSelection(const TArray< FWidgetReference >& Selection, TArray< TSharedRef<SWidget> >& Widgets)
 {
 	SelectionCache = Selection;
 
@@ -59,7 +59,7 @@ FReply FVerticalSlotExtension::HandleShiftVertical(int32 ShiftAmount)
 {
 	BeginTransaction(LOCTEXT("MoveWidget", "Move Widget"));
 
-	for ( FSelectedWidget& Selection : SelectionCache )
+	for ( FWidgetReference& Selection : SelectionCache )
 	{
 		ShiftVertical(Selection.GetPreview(), ShiftAmount);
 		ShiftVertical(Selection.GetTemplate(), ShiftAmount);
diff --git a/Engine/Source/Editor/UMGEditor/Private/Extensions/VerticalSlotExtension.h b/Engine/Source/Editor/UMGEditor/Private/Extensions/VerticalSlotExtension.h
--- a/Engine/Source/Editor/UMGEditor/Private/Extensions/VerticalSlotExtension.h
+++ b/Engine/Source/Editor/UMGEditor/Private/Extensions/VerticalSlotExtension.h
@@ -9,9 +9,9 @@ class FVerticalSlotExtension : public FDesignerExtension
 public:
 	FVerticalSlotExtension();
 
-	bool IsActive(const TArray< FSelectedWidget >& Selection);
+	bool IsActive(const TArray< FWidgetReference >& Selection);
 	
-	virtual void BuildWidgetsForSelection(const TArray< FSelectedWidget >& Selection, TArray< TSharedRef<SWidget> >& Widgets) override;
+	virtual void BuildWidgetsForSelection(const TArray< FWidgetReference >& Selection, TArray< TSharedRef<SWidget> >& Widgets) override;
 
 private:
 
diff --git a/Engine/Source/Editor/UMGEditor/Private/SUMGDesigner.cpp b/Engine/Source/Editor/UMGEditor/Private/SUMGDesigner.cpp
--- a/Engine/Source/Editor/UMGEditor/Private/SUMGDesigner.cpp
+++ b/Engine/Source/Editor/UMGEditor/Private/SUMGDesigner.cpp
@@ -163,12 +163,7 @@ void SUMGDesigner::OnObjectPropertyChanged(UObject* ObjectBeingModified, FProper
 	//UpdatePreview(InBlueprint);
 }
 
-void SUMGDesigner::ShowDetailsForObjects(TArray<UWidget*> Widgets)
-{
-	BlueprintEditor.Pin()->SelectWidgets(Widgets);
-}
-
-FSelectedWidget SUMGDesigner::GetWidgetAtCursor(const FGeometry& MyGeometry, const FPointerEvent& MouseEvent, FArrangedWidget& ArrangedWidget)
+FWidgetReference SUMGDesigner::GetWidgetAtCursor(const FGeometry& MyGeometry, const FPointerEvent& MouseEvent, FArrangedWidget& ArrangedWidget)
 {
 	//@TODO UMG Make it so you can request dropable widgets only, to find the first parentable.
 
@@ -206,11 +201,11 @@ FSelectedWidget SUMGDesigner::GetWidgetAtCursor(const FGeometry& MyGeometry, con
 
 		if ( Preview )
 		{
-			return FSelectedWidget::FromPreview(BlueprintEditor.Pin(), Preview);
+			return FWidgetReference::FromPreview(BlueprintEditor.Pin(), Preview);
 		}
 	}
 
-	return FSelectedWidget();
+	return FWidgetReference();
 }
 
 FReply SUMGDesigner::OnMouseButtonDown(const FGeometry& MyGeometry, const FPointerEvent& MouseEvent)
@@ -228,17 +223,16 @@ FReply SUMGDesigner::OnMouseButtonDown(const FGeometry& MyGeometry, const FPoint
 		{
 			//@TODO UMG primary FBlueprintEditor needs to be inherited and selection control needs to be centralized.
 			// Set the template as selected in the details panel
-			TArray<UWidget*> SelectedTemplates;
-			//SelectedTemplates.Add(CurrentSelection.GetTemplate());
-			SelectedTemplates.Add(CurrentSelection.GetPreview());
-			ShowDetailsForObjects(SelectedTemplates);
+			TSet<FWidgetReference> SelectedTemplates;
+			SelectedTemplates.Add(CurrentSelection);
+			BlueprintEditor.Pin()->SelectWidgets(SelectedTemplates);
 
 			// Remove all the current extension widgets
 			ExtensionWidgetCanvas->ClearChildren();
 
 			ExtensionWidgets.Reset();
 
-			TArray<FSelectedWidget> Selected;
+			TArray<FWidgetReference> Selected;
 			Selected.Add(CurrentSelection);
 
 			// Build extension widgets for new selection
@@ -735,7 +729,7 @@ UWidget* SUMGDesigner::AddPreview(const FGeometry& MyGeometry, const FDragDropEv
 	if ( DragDropOp.IsValid() )
 	{
 		FArrangedWidget ArrangedWidget(SNullWidget::NullWidget, FGeometry());
-		FSelectedWidget Selection = GetWidgetAtCursor(MyGeometry, DragDropEvent, ArrangedWidget);
+		FWidgetReference Selection = GetWidgetAtCursor(MyGeometry, DragDropEvent, ArrangedWidget);
 		
 		UWidgetBlueprint* BP = GetBlueprint();
 		
@@ -779,7 +773,7 @@ bool SUMGDesigner::AddToTemplate(const FGeometry& MyGeometry, const FDragDropEve
 	if ( DragDropOp.IsValid() )
 	{
 		FArrangedWidget ArrangedWidget(SNullWidget::NullWidget, FGeometry());
-		FSelectedWidget Selection = GetWidgetAtCursor(MyGeometry, DragDropEvent, ArrangedWidget);
+		FWidgetReference Selection = GetWidgetAtCursor(MyGeometry, DragDropEvent, ArrangedWidget);
 		
 		UWidgetBlueprint* BP = GetBlueprint();
 		
@@ -806,7 +800,7 @@ bool SUMGDesigner::AddToTemplate(const FGeometry& MyGeometry, const FDragDropEve
 				//@TODO UMG We may need a desired size canvas, where the slots have no size, they only give you position, alternatively, maybe slots that don't clip, so center is still easy.
 				
 				// Update the selected template to be the newly created one.
-				CurrentSelection = FSelectedWidget::FromTemplate(BlueprintEditor.Pin(), Widget);
+				CurrentSelection = FWidgetReference::FromTemplate(BlueprintEditor.Pin(), Widget);
 				
 				return true;
 			}
@@ -824,7 +818,7 @@ bool SUMGDesigner::AddToTemplate(const FGeometry& MyGeometry, const FDragDropEve
 			UWidget* Widget = DragDropOp->Template->Create(BP->WidgetTree);
 			Widget->IsDesignTime(true);
 			
-			CurrentSelection = FSelectedWidget::FromTemplate(BlueprintEditor.Pin(), Widget);
+			CurrentSelection = FWidgetReference::FromTemplate(BlueprintEditor.Pin(), Widget);
 			
 			return true;
 		}
diff --git a/Engine/Source/Editor/UMGEditor/Private/SUMGDesigner.h b/Engine/Source/Editor/UMGEditor/Private/SUMGDesigner.h
--- a/Engine/Source/Editor/UMGEditor/Private/SUMGDesigner.h
+++ b/Engine/Source/Editor/UMGEditor/Private/SUMGDesigner.h
@@ -9,6 +9,9 @@
 
 class FDesignerExtension;
 
+/**
+ * The designer for widgets.  Allows for laying out widgets in a drag and drop environment.
+ */
 class SUMGDesigner : public SDesignSurface
 {
 public:
@@ -38,12 +41,15 @@ public:
 	void Register(TSharedRef<FDesignerExtension> Extension);
 
 private:
-	FSelectedWidget GetWidgetAtCursor(const FGeometry& MyGeometry, const FPointerEvent& MouseEvent, FArrangedWidget& ArrangedWidget);
+	/** @returns Gets the widget under the cursor based on a mouse pointer event. */
+	FWidgetReference GetWidgetAtCursor(const FGeometry& MyGeometry, const FPointerEvent& MouseEvent, FArrangedWidget& ArrangedWidget);
 
+	/** Gets the blueprint being edited by the designer */
 	UWidgetBlueprint* GetBlueprint() const;
+
+	/** Called whenever the blueprint is structurally changed. */
 	void OnBlueprintChanged(UBlueprint* InBlueprint);
 	void OnObjectPropertyChanged(UObject* ObjectBeingModified, FPropertyChangedEvent& PropertyChangedEvent);
-	void ShowDetailsForObjects(TArray<UWidget*> Widgets);
 
 	bool GetArrangedWidget(TSharedRef<SWidget> Widget, FArrangedWidget& ArrangedWidget) const;
 	bool GetArrangedWidgetRelativeToWindow(TSharedRef<SWidget> Widget, FArrangedWidget& ArrangedWidget) const;
@@ -96,7 +102,7 @@ private:
 	TWeakPtr<SWidget> PreviewWidget;
 	UUserWidget* PreviewWidgetObject;
 
-	FSelectedWidget CurrentSelection;
+	FWidgetReference CurrentSelection;
 	TWeakPtr<SWidget> SelectedWidget;
 	
 	UWidget* DropPreviewWidget;
diff --git a/Engine/Source/Editor/UMGEditor/Private/SUMGEditorTree.cpp b/Engine/Source/Editor/UMGEditor/Private/SUMGEditorTree.cpp
--- a/Engine/Source/Editor/UMGEditor/Private/SUMGEditorTree.cpp
+++ b/Engine/Source/Editor/UMGEditor/Private/SUMGEditorTree.cpp
@@ -44,6 +44,8 @@ void SUMGEditorTree::Construct(const FArguments& InArgs, TSharedPtr<FWidgetBluep
 		]
 	];
 
+	BlueprintEditor.Pin()->OnSelectedWidgetsChanged.AddRaw(this, &SUMGEditorTree::OnEditorSelectionChanged);
+
 	RefreshTree();
 }
 
@@ -56,6 +58,31 @@ SUMGEditorTree::~SUMGEditorTree()
 	}
 
 	FCoreDelegates::OnObjectPropertyChanged.RemoveAll(this);
+
+	if ( BlueprintEditor.IsValid() )
+	{
+		BlueprintEditor.Pin()->OnSelectedWidgetsChanged.RemoveAll(this);
+	}
+}
+
+void SUMGEditorTree::OnEditorSelectionChanged()
+{
+	WidgetTreeView->ClearSelection();
+
+	// Update the selection and expansion in the tree to match the new selection
+	const TSet<FWidgetReference>& SelectedWidgets = BlueprintEditor.Pin()->GetSelectedWidgets();
+	for ( const FWidgetReference& WidgetRef : SelectedWidgets )
+	{
+		WidgetTreeView->SetItemSelection(WidgetRef.GetTemplate(), true);
+
+		// Expand the path leading to this widget in the tree.
+		UWidget* Parent = WidgetRef.GetTemplate()->GetParent();
+		while ( Parent != NULL )
+		{
+			WidgetTreeView->SetItemExpansion(Parent, true);
+			Parent = Parent->GetParent();
+		}
+	}
 }
 
 UWidgetBlueprint* SUMGEditorTree::GetBlueprint() const
@@ -87,15 +114,14 @@ void SUMGEditorTree::OnObjectPropertyChanged(UObject* ObjectBeingModified, FProp
 
 void SUMGEditorTree::ShowDetailsForObjects(TArray<UWidget*> TemplateWidgets)
 {
-	TArray<UWidget*> PreviewWidgets;
+	TSet<FWidgetReference> SelectedWidgets;
 	for ( UWidget* TemplateWidget : TemplateWidgets )
 	{
-		FSelectedWidget Selection = FSelectedWidget::FromTemplate(BlueprintEditor.Pin(), TemplateWidget);
-		PreviewWidgets.Add(Selection.GetPreview());
-		//PreviewWidgets.Add(Selection.GetTemplate());
+		FWidgetReference Selection = FWidgetReference::FromTemplate(BlueprintEditor.Pin(), TemplateWidget);
+		SelectedWidgets.Add(Selection);
 	}
 
-	BlueprintEditor.Pin()->SelectWidgets(PreviewWidgets);
+	BlueprintEditor.Pin()->SelectWidgets(SelectedWidgets);
 }
 
 void SUMGEditorTree::BuildWrapWithMenu(FMenuBuilder& Menu)
diff --git a/Engine/Source/Editor/UMGEditor/Private/SUMGEditorTree.h b/Engine/Source/Editor/UMGEditor/Private/SUMGEditorTree.h
--- a/Engine/Source/Editor/UMGEditor/Private/SUMGEditorTree.h
+++ b/Engine/Source/Editor/UMGEditor/Private/SUMGEditorTree.h
@@ -6,6 +6,11 @@
 #include "BlueprintEditor.h"
 
 //TODO rename SUMGEditorHierarchy
+
+/**
+ * The tree view presenting the widget hierarchy.  This allows users to edit the hierarchy of widgets easily by dragging and 
+ * dropping them logically, which in some cases may be significantly easier than doing it visually in the widget designer.
+ */
 class SUMGEditorTree : public SCompoundWidget
 {
 public:
@@ -13,14 +18,21 @@ public:
 	SLATE_END_ARGS()
 
 	void Construct(const FArguments& InArgs, TSharedPtr<FWidgetBlueprintEditor> InBlueprintEditor, USimpleConstructionScript* InSCS);
+
 	virtual ~SUMGEditorTree();
 
 private:
+	/** @returns the current blueprint being edited */
 	UWidgetBlueprint* GetBlueprint() const;
 
+	/** Called when the blueprint is structurally changed. */
 	void OnBlueprintChanged(UBlueprint* InBlueprint);
+
 	void OnObjectPropertyChanged(UObject* ObjectBeingModified, FPropertyChangedEvent& PropertyChangedEvent);
 
+	/** Called when the selected widget has changed.  The treeview then needs to match the new selection. */
+	void OnEditorSelectionChanged();
+
 	void ShowDetailsForObjects(TArray<UWidget*> Widgets);
 
 	void BuildWrapWithMenu(FMenuBuilder& Menu);
diff --git a/Engine/Source/Editor/UMGEditor/Private/WidgetBlueprintEditor.cpp b/Engine/Source/Editor/UMGEditor/Private/WidgetBlueprintEditor.cpp
--- a/Engine/Source/Editor/UMGEditor/Private/WidgetBlueprintEditor.cpp
+++ b/Engine/Source/Editor/UMGEditor/Private/WidgetBlueprintEditor.cpp
@@ -26,18 +26,32 @@ FWidgetBlueprintEditor::~FWidgetBlueprintEditor()
 	}
 }
 
-void FWidgetBlueprintEditor::SelectWidgets(TArray<UWidget*> Widgets)
+void FWidgetBlueprintEditor::SelectWidgets(const TSet<FWidgetReference>& Widgets)
+{
+	SelectedWidgets = Widgets;
+
+	RefreshDetails();
+
+	OnSelectedWidgetsChanged.Broadcast();
+}
+
+const TSet<FWidgetReference>& FWidgetBlueprintEditor::GetSelectedWidgets() const
+{
+	return SelectedWidgets;
+}
+
+void FWidgetBlueprintEditor::RefreshDetails()
 {
 	// Convert the selection set to an array of UObject* pointers
-	FString InspectorTitle;
+	FString InspectorTitle = "Widget";// Widget->GetDisplayString();
+
 	TArray<UObject*> InspectorObjects;
-	InspectorObjects.Empty(Widgets.Num());
-	for ( UWidget* Widget : Widgets )
+	for ( FWidgetReference& WidgetRef : SelectedWidgets )
 	{
-		//if ( NodePtr->CanEditDefaults() )
+		UWidget* PreviewWidget = WidgetRef.GetPreview();
+		if ( PreviewWidget )
 		{
-			InspectorTitle = "Widget";// Widget->GetDisplayString();
-			InspectorObjects.Add(Widget);
+			InspectorObjects.Add(PreviewWidget);
 		}
 	}
 
@@ -46,8 +60,6 @@ void FWidgetBlueprintEditor::SelectWidgets(TArray<UWidget*> Widgets)
 	// Update the details panel
 	SKismetInspector::FShowDetailsOptions Options(InspectorTitle, true);
 	GetInspector()->ShowDetailsForObjects(InspectorObjects, Options);
-
-	SelectedPreviewWidgets = Widgets;
 }
 
 void FWidgetBlueprintEditor::OnBlueprintChanged(UBlueprint* InBlueprint)
@@ -55,6 +67,8 @@ void FWidgetBlueprintEditor::OnBlueprintChanged(UBlueprint* InBlueprint)
 	if ( InBlueprint )
 	{
 		UpdatePreview(InBlueprint, true);
+
+		RefreshDetails();
 	}
 }
 
@@ -176,15 +190,20 @@ void FWidgetBlueprintEditor::MigrateFromChain(FEditPropertyChain* PropertyThatCh
 	UUserWidget* PreviewActor = GetPreview();
 	if ( PreviewActor != NULL )
 	{
-		for ( UWidget* PreviewWidget : SelectedPreviewWidgets )
+		for ( FWidgetReference& WidgetRef : SelectedWidgets )
 		{
-			FString PreviewWidgetName = PreviewWidget->GetName();
-			UWidget* TemplateWidget = Blueprint->WidgetTree->FindWidget(PreviewWidgetName);
+			UWidget* PreviewWidget = WidgetRef.GetPreview();
 
-			if ( TemplateWidget )
+			if ( PreviewWidget )
 			{
-				FEditPropertyChain::TDoubleLinkedListNode* PropertyChainNode = PropertyThatChanged->GetHead();
-				MigratePropertyValue(PreviewWidget, TemplateWidget, PropertyChainNode, PropertyChainNode->GetValue(), bIsModify);
+				FString PreviewWidgetName = PreviewWidget->GetName();
+				UWidget* TemplateWidget = Blueprint->WidgetTree->FindWidget(PreviewWidgetName);
+
+				if ( TemplateWidget )
+				{
+					FEditPropertyChain::TDoubleLinkedListNode* PropertyChainNode = PropertyThatChanged->GetHead();
+					MigratePropertyValue(PreviewWidget, TemplateWidget, PropertyChainNode, PropertyChainNode->GetValue(), bIsModify);
+				}
 			}
 		}
 	}
@@ -298,6 +317,7 @@ void FWidgetBlueprintEditor::UpdatePreview(UBlueprint* InBlueprint, bool bInForc
 			SubPreviewWidget->IsDesignTime(true);
 		}
 
+		// Store a reference to the preview actor.
 		PreviewWidgetActorPtr = PreviewActor;
 	}
 }
diff --git a/Engine/Source/Editor/UMGEditor/Private/WidgetBlueprintEditor.h b/Engine/Source/Editor/UMGEditor/Private/WidgetBlueprintEditor.h
--- a/Engine/Source/Editor/UMGEditor/Private/WidgetBlueprintEditor.h
+++ b/Engine/Source/Editor/UMGEditor/Private/WidgetBlueprintEditor.h
@@ -13,6 +13,9 @@ class ISequencer;
  */
 class FWidgetBlueprintEditor : public FBlueprintEditor
 {
+public:
+	DECLARE_MULTICAST_DELEGATE(FOnSelectedWidgetsChanged)
+
 public:
 	FWidgetBlueprintEditor();
 
@@ -27,21 +30,38 @@ public:
 	/** FGCObjectInterface */
 	virtual void AddReferencedObjects( FReferenceCollector& Collector ) override;
 
+	/** @return The widget blueprint currently being edited in this editor */
 	class UWidgetBlueprint* GetWidgetBlueprintObj() const;
 
+	/** @return The preview widget. */
 	UUserWidget* GetPreview() const;
 
 	/** @return The sequencer used to create widget animations */
 	TSharedPtr<ISequencer>& GetSequencer();
 
-	void SelectWidgets(TArray<UWidget*> Widgets);
+	/** Sets the currently selected set of widgets */
+	void SelectWidgets(const TSet<FWidgetReference>& Widgets);
+
+	/** @return The selected set of widgets */
+	const TSet<FWidgetReference>& GetSelectedWidgets() const;
+
+	/** Fires whenever the selected set of widgets changes */
+	FOnSelectedWidgetsChanged OnSelectedWidgetsChanged;
 
 private:
+	/** Updates the inspector to be viewing the currently selected set of widgets */
+	void RefreshDetails();
+
+	/** Called whenever the blueprint is structurally changed. */
 	virtual void OnBlueprintChanged(UBlueprint* InBlueprint) override;
 
+	/** Destroy the current preview GUI object */
 	void DestroyPreview();
+
+	/** Tick the current preview GUI object */
 	void UpdatePreview(UBlueprint* InBlueprint, bool bInForceFullUpdate);
 
+	/** Migrate a property change from the preview GUI to the template GUI. */
 	void MigrateFromChain(FEditPropertyChain* PropertyThatChanged, bool bIsModify);
 
 	/**
@@ -51,6 +71,7 @@ private:
 	UMovieScene* GetDefaultMovieScene();
 
 private:
+	/** The preview scene that owns the preview GUI */
 	FPreviewScene PreviewScene;
 
 	/** Sequencer for creating and previewing widget animations */
@@ -62,7 +83,9 @@ private:
 	/** Default movie scene for new animations */
 	UMovieScene* DefaultMovieScene;
 
-	TArray<UWidget*> SelectedPreviewWidgets;
+	/** The currently selected preview widgets in the preview GUI */
+	TSet<FWidgetReference> SelectedWidgets;
 
+	/** The preview GUI object */
 	mutable TWeakObjectPtr<UUserWidget> PreviewWidgetActorPtr;
 };
diff --git a/Engine/Source/Editor/UMGEditor/Public/DesignerExtension.h b/Engine/Source/Editor/UMGEditor/Public/DesignerExtension.h
--- a/Engine/Source/Editor/UMGEditor/Public/DesignerExtension.h
+++ b/Engine/Source/Editor/UMGEditor/Public/DesignerExtension.h
@@ -2,45 +2,29 @@
 
 #pragma once
 
+#include "WidgetReference.h"
+
 class UWidget;
 class SWidget;
 class UWidgetBlueprint;
 
 /**
- * The Selected widget is a useful way to hold onto the selection in a way that allows for up to date access to the current preview object.
- * Because the designer could end up rebuilding the preview, it's best to hold onto an FSelectedWidget.
+ * The Designer extension allows developers to provide additional widgets and custom painting to the designer surface for
+ * specific widgets.  Allowing for a more customized and specific editors for the different widgets.
  */
-struct FSelectedWidget
-{
-public:
-	static FSelectedWidget FromTemplate(TSharedPtr<class FWidgetBlueprintEditor> WidgetEditor, UWidget* TemplateWidget);
-	static FSelectedWidget FromPreview(TSharedPtr<class FWidgetBlueprintEditor> WidgetEditor, UWidget* PreviewWidget);
-
-	FSelectedWidget();
-
-	bool IsValid() const;
-
-	UWidget* GetTemplate() const;
-	UWidget* GetPreview() const;
-
-private:
-	FSelectedWidget(TSharedPtr<class FWidgetBlueprintEditor> WidgetEditor, UWidget* TemplateWidget);
-
-private:
-
-	TSharedPtr<class FWidgetBlueprintEditor> WidgetEditor;
-	UWidget* TemplateWidget;
-};
-
 class FDesignerExtension : public TSharedFromThis<FDesignerExtension>
 {
 public:
+	/** Constructor */
 	FDesignerExtension();
 
+	/** Initializes the designer extension, this is called the first time a designer extension is registered */
 	virtual void Initialize(UWidgetBlueprint* InBlueprint);
 
-	virtual void BuildWidgetsForSelection(const TArray< FSelectedWidget >& Selection, TArray< TSharedRef<SWidget> >& Widgets) = 0;
+	/** Called every time the selection in the designer changes. */
+	virtual void BuildWidgetsForSelection(const TArray< FWidgetReference >& Selection, TArray< TSharedRef<SWidget> >& Widgets) = 0;
 
+	/** Gets the ID identifying this extension. */
 	FName GetExtensionId() const;
 
 protected:
@@ -52,7 +36,7 @@ protected:
 	FName ExtensionId;
 	UWidgetBlueprint* Blueprint;
 
-	TArray< FSelectedWidget > SelectionCache;
+	TArray< FWidgetReference > SelectionCache;
 
 private:
 	FScopedTransaction* ScopedTransaction;
diff --git a/Engine/Source/Editor/UMGEditor/Public/WidgetReference.h b/Engine/Source/Editor/UMGEditor/Public/WidgetReference.h
new file mode 100644
--- /dev/null
+++ b/Engine/Source/Editor/UMGEditor/Public/WidgetReference.h
@@ -0,0 +1,57 @@
+// Copyright 1998-2014 Epic Games, Inc. All Rights Reserved.
+
+#pragma once
+
+class UWidget;
+class SWidget;
+class UWidgetBlueprint;
+
+/**
+ * The Widget reference is a useful way to hold onto the selection in a way that allows for up to date access to the current preview object.
+ * Because the designer could end up rebuilding the preview, it's best to hold onto an FWidgetReference.
+ */
+struct FWidgetReference
+{
+public:
+	/**
+	 * Creates a widget reference using the template.
+	 */
+	static FWidgetReference FromTemplate(TSharedPtr<class FWidgetBlueprintEditor> WidgetEditor, UWidget* TemplateWidget);
+
+	/**
+	 * Creates a widget reference using the preview.  Which is used to lookup the stable template pointer.
+	 */
+	static FWidgetReference FromPreview(TSharedPtr<class FWidgetBlueprintEditor> WidgetEditor, UWidget* PreviewWidget);
+
+	/** Constructor for the null widget reference.  Not intended for normal use. */
+	FWidgetReference();
+
+	/** @returns true if both the template and the preview pointers are valid. */
+	bool IsValid() const;
+
+	/** @returns The template widget.  This is the widget that is serialized to disk */
+	UWidget* GetTemplate() const;
+
+	/** @returns the preview widget.  This is the transient representation of the template.  Constantly being destroyed and recreated.  Do not cache this pointer. */
+	UWidget* GetPreview() const;
+
+	/** Checks if widget reference is the same as nother widget reference, based on the template pointers. */
+	bool operator==( const FWidgetReference& Other ) const
+	{
+		return TemplateWidget == Other.TemplateWidget;
+	}
+
+private:
+	FWidgetReference(TSharedPtr<class FWidgetBlueprintEditor> WidgetEditor, UWidget* TemplateWidget);
+
+private:
+
+	TSharedPtr<class FWidgetBlueprintEditor> WidgetEditor;
+
+	UWidget* TemplateWidget;
+};
+
+inline uint32 GetTypeHash(const struct FWidgetReference& WidgetRef)
+{
+	return ::GetTypeHash(reinterpret_cast<void*>( WidgetRef.GetTemplate() ));
+}
diff --git a/Engine/Source/Runtime/UMG/Private/Components/Border.cpp b/Engine/Source/Runtime/UMG/Private/Components/Border.cpp
--- a/Engine/Source/Runtime/UMG/Private/Components/Border.cpp
+++ b/Engine/Source/Runtime/UMG/Private/Components/Border.cpp
@@ -30,7 +30,7 @@ UBorder::UBorder(const FPostConstructInitializeProperties& PCIP)
 TSharedRef<SWidget> UBorder::RebuildWidget()
 {
 	MyBorder = SNew(SBorder);
-	MyBorder->SetContent(Content ? Content->GetWidget() : SNullWidget::NullWidget);
+	MyBorder->SetContent(ContentSlot->Content ? ContentSlot->Content->GetWidget() : SNullWidget::NullWidget);
 
 	return MyBorder.ToSharedRef();
 }
@@ -60,13 +60,13 @@ void UBorder::SyncronizeProperties()
 	MyBorder->SetOnMouseDoubleClick(BIND_UOBJECT_DELEGATE(FPointerEventHandler, HandleMouseDoubleClick));
 }
 
-void UBorder::SetContent(UWidget* Content)
+void UBorder::SetContent(UWidget* InContent)
 {
-	Super::SetContent(Content);
+	Super::SetContent(InContent);
 
 	if ( MyBorder.IsValid() )
 	{
-		MyBorder->SetContent(Content ? Content->GetWidget() : SNullWidget::NullWidget);
+		MyBorder->SetContent(InContent ? InContent->GetWidget() : SNullWidget::NullWidget);
 	}
 }
 
diff --git a/Engine/Source/Runtime/UMG/Private/Components/Button.cpp b/Engine/Source/Runtime/UMG/Private/Components/Button.cpp
--- a/Engine/Source/Runtime/UMG/Private/Components/Button.cpp
+++ b/Engine/Source/Runtime/UMG/Private/Components/Button.cpp
@@ -29,7 +29,7 @@ UButton::UButton(const FPostConstructInitializeProperties& PCIP)
 TSharedRef<SWidget> UButton::RebuildWidget()
 {
 	MyButton = SNew(SButton);
-	MyButton->SetContent(Content ? Content->GetWidget() : SNullWidget::NullWidget);
+	MyButton->SetContent(ContentSlot->Content ? ContentSlot->Content->GetWidget() : SNullWidget::NullWidget);
 	
 	return MyButton.ToSharedRef();
 }
diff --git a/Engine/Source/Runtime/UMG/Private/Components/CanvasPanel.cpp b/Engine/Source/Runtime/UMG/Private/Components/CanvasPanel.cpp
--- a/Engine/Source/Runtime/UMG/Private/Components/CanvasPanel.cpp
+++ b/Engine/Source/Runtime/UMG/Private/Components/CanvasPanel.cpp
@@ -114,9 +114,7 @@ void UCanvasPanel::ReplaceChildAt(int32 Index, UWidget* Content)
 	UCanvasPanelSlot* Slot = Slots[Index];
 	Slot->Content = Content;
 
-#if WITH_EDITOR
 	Content->Slot = Slot;
-#endif
 	
 	//TODO UMG How do we handle this swap for live controls?  It doesn't matter for canvas, but matters for flow controls it will matter
 }
@@ -128,9 +126,7 @@ void UCanvasPanel::InsertChildAt(int32 Index, UWidget* Content)
 	Slot->Content = Content;
 	Slot->Parent = this;
 
-#if WITH_EDITOR
 	Content->Slot = Slot;
-#endif
 
 	Slots.Insert(Slot, Index);
 }
@@ -162,9 +158,7 @@ UCanvasPanelSlot* UCanvasPanel::AddSlot(UWidget* Content)
 	Slot->Content = Content;
 	Slot->Parent = this;
 
-#if WITH_EDITOR
 	Content->Slot = Slot;
-#endif
 	
 	Slots.Add(Slot);
 
diff --git a/Engine/Source/Runtime/UMG/Private/Components/ContentWidget.cpp b/Engine/Source/Runtime/UMG/Private/Components/ContentWidget.cpp
--- a/Engine/Source/Runtime/UMG/Private/Components/ContentWidget.cpp
+++ b/Engine/Source/Runtime/UMG/Private/Components/ContentWidget.cpp
@@ -15,31 +15,39 @@
 UContentWidget::UContentWidget(const FPostConstructInitializeProperties& PCIP)
 	: Super(PCIP)
 {
+	ContentSlot = ConstructObject<UPanelSlot>(UPanelSlot::StaticClass(), this);
+	ContentSlot->SetFlags(RF_Transactional);
+	ContentSlot->Parent = this;
 }
 
 UWidget* UContentWidget::GetContent()
 {
-	return Content;
+	return ContentSlot->Content;
 }
 
 void UContentWidget::SetContent(UWidget* InContent)
 {
-	Content = InContent;
+	ContentSlot->Content = InContent;
+
+	if ( InContent )
+	{
+		InContent->Slot = ContentSlot;
+	}
 }
 
 int32 UContentWidget::GetChildrenCount() const
 {
-	return Content != NULL ? 1 : 0;
+	return ContentSlot->Content != NULL ? 1 : 0;
 }
 
 UWidget* UContentWidget::GetChildAt(int32 Index) const
 {
-	return Content;
+	return ContentSlot->Content;
 }
 
 bool UContentWidget::AddChild(UWidget* InContent, FVector2D Position)
 {
-	if ( Content == NULL )
+	if ( ContentSlot->Content == NULL )
 	{
 		SetContent(InContent);
 		return true;
@@ -50,7 +58,7 @@ bool UContentWidget::AddChild(UWidget* InContent, FVector2D Position)
 
 bool UContentWidget::RemoveChild(UWidget* Child)
 {
-	if ( Content == Child )
+	if ( ContentSlot->Content == Child )
 	{
 		SetContent(NULL);
 		return true;
diff --git a/Engine/Source/Runtime/UMG/Private/Components/HorizontalBox.cpp b/Engine/Source/Runtime/UMG/Private/Components/HorizontalBox.cpp
--- a/Engine/Source/Runtime/UMG/Private/Components/HorizontalBox.cpp
+++ b/Engine/Source/Runtime/UMG/Private/Components/HorizontalBox.cpp
@@ -66,9 +66,7 @@ void UHorizontalBox::ReplaceChildAt(int32 Index, UWidget* Content)
 	UHorizontalBoxSlot* Slot = Slots[Index];
 	Slot->Content = Content;
 
-#if WITH_EDITOR
 	Content->Slot = Slot;
-#endif
 }
 
 void UHorizontalBox::InsertChildAt(int32 Index, UWidget* Content)
@@ -78,9 +76,7 @@ void UHorizontalBox::InsertChildAt(int32 Index, UWidget* Content)
 	Slot->Content = Content;
 	Slot->Parent = this;
 
-#if WITH_EDITOR
 	Content->Slot = Slot;
-#endif
 
 	Slots.Insert(Slot, Index);
 }
@@ -105,9 +101,7 @@ UHorizontalBoxSlot* UHorizontalBox::AddSlot(UWidget* Content)
 	Slot->Content = Content;
 	Slot->Parent = this;
 
-#if WITH_EDITOR
 	Content->Slot = Slot;
-#endif
 	
 	Slots.Add(Slot);
 
diff --git a/Engine/Source/Runtime/UMG/Private/Components/Overlay.cpp b/Engine/Source/Runtime/UMG/Private/Components/Overlay.cpp
--- a/Engine/Source/Runtime/UMG/Private/Components/Overlay.cpp
+++ b/Engine/Source/Runtime/UMG/Private/Components/Overlay.cpp
@@ -62,9 +62,7 @@ void UOverlay::ReplaceChildAt(int32 Index, UWidget* Content)
 	UOverlaySlot* Slot = Slots[Index];
 	Slot->Content = Content;
 
-#if WITH_EDITOR
 	Content->Slot = Slot;
-#endif
 }
 
 TSharedRef<SWidget> UOverlay::RebuildWidget()
@@ -87,9 +85,7 @@ UOverlaySlot* UOverlay::AddSlot(UWidget* Content)
 	Slot->Content = Content;
 	Slot->Parent = this;
 
-#if WITH_EDITOR
 	Content->Slot = Slot;
-#endif
 	
 	Slots.Add(Slot);
 
diff --git a/Engine/Source/Runtime/UMG/Private/Components/PanelWidget.cpp b/Engine/Source/Runtime/UMG/Private/Components/PanelWidget.cpp
--- a/Engine/Source/Runtime/UMG/Private/Components/PanelWidget.cpp
+++ b/Engine/Source/Runtime/UMG/Private/Components/PanelWidget.cpp
@@ -15,10 +15,6 @@
 UPanelWidget::UPanelWidget(const FPostConstructInitializeProperties& PCIP)
 	: Super(PCIP)
 {
-//	PrimaryComponentTick.bCanEverTick = true;
-
-//	bAutoActivate = true;
-//	bTickInEditor = true;
 }
 
 /////////////////////////////////////////////////////
diff --git a/Engine/Source/Runtime/UMG/Private/Components/ScrollBox.cpp b/Engine/Source/Runtime/UMG/Private/Components/ScrollBox.cpp
--- a/Engine/Source/Runtime/UMG/Private/Components/ScrollBox.cpp
+++ b/Engine/Source/Runtime/UMG/Private/Components/ScrollBox.cpp
@@ -75,9 +75,7 @@ void UScrollBox::ReplaceChildAt(int32 Index, UWidget* Content)
 	UScrollBoxSlot* Slot = Slots[Index];
 	Slot->Content = Content;
 
-#if WITH_EDITOR
 	Content->Slot = Slot;
-#endif
 }
 
 void UScrollBox::InsertChildAt(int32 Index, UWidget* Content)
@@ -87,9 +85,7 @@ void UScrollBox::InsertChildAt(int32 Index, UWidget* Content)
 	Slot->Content = Content;
 	Slot->Parent = this;
 
-#if WITH_EDITOR
 	Content->Slot = Slot;
-#endif
 
 	Slots.Insert(Slot, Index);
 }
@@ -114,9 +110,7 @@ UScrollBoxSlot* UScrollBox::AddSlot(UWidget* Content)
 	Slot->Content = Content;
 	Slot->Parent = this;
 
-#if WITH_EDITOR
 	Content->Slot = Slot;
-#endif
 	
 	Slots.Add(Slot);
 
diff --git a/Engine/Source/Runtime/UMG/Private/Components/UniformGridPanel.cpp b/Engine/Source/Runtime/UMG/Private/Components/UniformGridPanel.cpp
--- a/Engine/Source/Runtime/UMG/Private/Components/UniformGridPanel.cpp
+++ b/Engine/Source/Runtime/UMG/Private/Components/UniformGridPanel.cpp
@@ -62,9 +62,7 @@ void UUniformGridPanel::ReplaceChildAt(int32 Index, UWidget* Content)
 	UUniformGridSlot* Slot = Slots[Index];
 	Slot->Content = Content;
 
-#if WITH_EDITOR
 	Content->Slot = Slot;
-#endif
 }
 
 TSharedRef<SWidget> UUniformGridPanel::RebuildWidget()
@@ -91,9 +89,7 @@ UUniformGridSlot* UUniformGridPanel::AddSlot(UWidget* Content)
 	Slot->Content = Content;
 	Slot->Parent = this;
 
-#if WITH_EDITOR
 	Content->Slot = Slot;
-#endif
 	
 	Slots.Add(Slot);
 
diff --git a/Engine/Source/Runtime/UMG/Private/Components/VerticalBox.cpp b/Engine/Source/Runtime/UMG/Private/Components/VerticalBox.cpp
--- a/Engine/Source/Runtime/UMG/Private/Components/VerticalBox.cpp
+++ b/Engine/Source/Runtime/UMG/Private/Components/VerticalBox.cpp
@@ -66,9 +66,7 @@ void UVerticalBox::ReplaceChildAt(int32 Index, UWidget* Content)
 	UVerticalBoxSlot* Slot = Slots[Index];
 	Slot->Content = Content;
 
-#if WITH_EDITOR
 	Content->Slot = Slot;
-#endif
 }
 
 void UVerticalBox::InsertChildAt(int32 Index, UWidget* Content)
@@ -78,9 +76,7 @@ void UVerticalBox::InsertChildAt(int32 Index, UWidget* Content)
 	Slot->Content = Content;
 	Slot->Parent = this;
 
-#if WITH_EDITOR
 	Content->Slot = Slot;
-#endif
 
 	Slots.Insert(Slot, Index);
 }
@@ -105,9 +101,7 @@ UVerticalBoxSlot* UVerticalBox::AddSlot(UWidget* Content)
 	Slot->Content = Content;
 	Slot->Parent = this;
 
-#if WITH_EDITOR
 	Content->Slot = Slot;
-#endif
 	
 	Slots.Add(Slot);
 
diff --git a/Engine/Source/Runtime/UMG/Private/Components/Widget.cpp b/Engine/Source/Runtime/UMG/Private/Components/Widget.cpp
--- a/Engine/Source/Runtime/UMG/Private/Components/Widget.cpp
+++ b/Engine/Source/Runtime/UMG/Private/Components/Widget.cpp
@@ -63,6 +63,16 @@ bool UWidget::IsHovered() const
 	return MyWidget->IsHovered();
 }
 
+UWidget* UWidget::GetParent() const
+{
+	if ( Slot )
+	{
+		return Slot->Parent;
+	}
+
+	return NULL;
+}
+
 TSharedRef<SWidget> UWidget::GetWidget() const
 {
 	if ( !MyWidget.IsValid() )
diff --git a/Engine/Source/Runtime/UMG/Private/Components/WidgetSwitcher.cpp b/Engine/Source/Runtime/UMG/Private/Components/WidgetSwitcher.cpp
--- a/Engine/Source/Runtime/UMG/Private/Components/WidgetSwitcher.cpp
+++ b/Engine/Source/Runtime/UMG/Private/Components/WidgetSwitcher.cpp
@@ -97,9 +97,7 @@ void UWidgetSwitcher::ReplaceChildAt(int32 Index, UWidget* Content)
 	UWidgetSwitcherSlot* Slot = Slots[Index];
 	Slot->Content = Content;
 
-#if WITH_EDITOR
 	Content->Slot = Slot;
-#endif
 }
 
 void UWidgetSwitcher::InsertChildAt(int32 Index, UWidget* Content)
@@ -109,9 +107,7 @@ void UWidgetSwitcher::InsertChildAt(int32 Index, UWidget* Content)
 	Slot->Content = Content;
 	Slot->Parent = this;
 
-#if WITH_EDITOR
 	Content->Slot = Slot;
-#endif
 
 	Slots.Insert(Slot, Index);
 }
@@ -123,9 +119,7 @@ UWidgetSwitcherSlot* UWidgetSwitcher::AddSlot(UWidget* Content)
 	Slot->Content = Content;
 	Slot->Parent = this;
 
-#if WITH_EDITOR
 	Content->Slot = Slot;
-#endif
 
 	Slots.Add(Slot);
 
diff --git a/Engine/Source/Runtime/UMG/Public/Components/ContentWidget.h b/Engine/Source/Runtime/UMG/Public/Components/ContentWidget.h
--- a/Engine/Source/Runtime/UMG/Public/Components/ContentWidget.h
+++ b/Engine/Source/Runtime/UMG/Public/Components/ContentWidget.h
@@ -26,5 +26,5 @@ class UMG_API UContentWidget : public UPanelWidget
 protected:
 
 	UPROPERTY()
-	UWidget* Content;
+	UPanelSlot* ContentSlot;
 };
diff --git a/Engine/Source/Runtime/UMG/Public/Components/PanelSlot.h b/Engine/Source/Runtime/UMG/Public/Components/PanelSlot.h
--- a/Engine/Source/Runtime/UMG/Public/Components/PanelSlot.h
+++ b/Engine/Source/Runtime/UMG/Public/Components/PanelSlot.h
@@ -5,7 +5,7 @@
 #include "PanelSlot.generated.h"
 
 /** The base class for all Slots in UMG. */
-UCLASS(Abstract)
+UCLASS()
 class UMG_API UPanelSlot : public UObject
 {
 	GENERATED_UCLASS_BODY()
diff --git a/Engine/Source/Runtime/UMG/Public/Components/Widget.h b/Engine/Source/Runtime/UMG/Public/Components/Widget.h
--- a/Engine/Source/Runtime/UMG/Public/Components/Widget.h
+++ b/Engine/Source/Runtime/UMG/Public/Components/Widget.h
@@ -52,13 +52,11 @@ public:
 	UPROPERTY(Transient)
 	bool bCreatedByConstructionScript;
 
-#if WITH_EDITOR
 	/**
 	 * The parent slot of the UWidget.  Allows us to easily inline edit the layout controlling this widget.
 	 */
-	UPROPERTY(Transient, EditInline, EditDefaultsOnly, BlueprintReadWrite, Category=Layout, meta=( ShowOnlyInnerProperties ))
+	UPROPERTY(EditInline, EditDefaultsOnly, BlueprintReadWrite, Category=Layout, meta=( ShowOnlyInnerProperties ))
 	UPanelSlot* Slot;
-#endif
 
 	/** Sets whether this widget can be modified interactively by the user */
 	UPROPERTY(EditDefaultsOnly, Category=Behavior)
@@ -120,6 +118,10 @@ public:
 	UFUNCTION(BlueprintCallable, Category="Widget")
 	bool IsHovered() const;
 
+	/** Gets the parent widget */
+	UFUNCTION(BlueprintCallable, Category="Widget")
+	UWidget* GetParent() const;
+
 	/** Gets the underlying Slate SWidget for this UWidget. */
 	TSharedRef<SWidget> GetWidget() const;
 
