UMG - A tiny UI polish, adding filter boxes.  No effect yet.

diff --git a/Engine/Source/Editor/UMGEditor/Private/SUMGEditorTree.cpp b/Engine/Source/Editor/UMGEditor/Private/SUMGEditorTree.cpp
--- a/Engine/Source/Editor/UMGEditor/Private/SUMGEditorTree.cpp
+++ b/Engine/Source/Editor/UMGEditor/Private/SUMGEditorTree.cpp
@@ -26,21 +26,39 @@ void SUMGEditorTree::Construct(const FArguments& InArgs, TSharedPtr<FWidgetBluep
 
 	FCoreDelegates::OnObjectPropertyChanged.AddRaw(this, &SUMGEditorTree::OnObjectPropertyChanged);
 
+	SAssignNew(WidgetTreeView, STreeView< UWidget* >)
+	.ItemHeight(20.0f)
+	.SelectionMode(ESelectionMode::Single)
+	.OnGetChildren(this, &SUMGEditorTree::WidgetHierarchy_OnGetChildren)
+	.OnGenerateRow(this, &SUMGEditorTree::WidgetHierarchy_OnGenerateRow)
+	.OnSelectionChanged(this, &SUMGEditorTree::WidgetHierarchy_OnSelectionChanged)
+	.OnContextMenuOpening(this, &SUMGEditorTree::WidgetHierarchy_OnContextMenuOpening)
+	.TreeItemsSource(&RootWidgets);
+
 	ChildSlot
 	[
-		SNew(SVerticalBox)
-
-		+ SVerticalBox::Slot()
-		.FillHeight(1.0f)
+		SNew(SBorder)
+		.BorderImage(FEditorStyle::GetBrush("ToolPanel.GroupBorder"))
 		[
-			SAssignNew(WidgetTreeView, STreeView< UWidget* >)
-			.ItemHeight(20.0f)
-			.SelectionMode(ESelectionMode::Single)
-			.OnGetChildren(this, &SUMGEditorTree::WidgetHierarchy_OnGetChildren)
-			.OnGenerateRow(this, &SUMGEditorTree::WidgetHierarchy_OnGenerateRow)
-			.OnSelectionChanged(this, &SUMGEditorTree::WidgetHierarchy_OnSelectionChanged)
-			.OnContextMenuOpening(this, &SUMGEditorTree::WidgetHierarchy_OnContextMenuOpening)
-			.TreeItemsSource(&RootWidgets)
+			SNew(SVerticalBox)
+
+			+ SVerticalBox::Slot()
+			.Padding(4)
+			.AutoHeight()
+			[
+				SNew(SSearchBox)
+				.HintText(LOCTEXT("SearchWidgets", "Search Widgets"))
+				.OnTextChanged(this, &SUMGEditorTree::OnSearchChanged)
+			]
+
+			+ SVerticalBox::Slot()
+			.FillHeight(1.0f)
+			[
+				SNew(SScrollBorder, WidgetTreeView.ToSharedRef())
+				[
+					WidgetTreeView.ToSharedRef()
+				]
+			]
 		]
 	];
 
@@ -65,16 +83,31 @@ SUMGEditorTree::~SUMGEditorTree()
 	}
 }
 
+void SUMGEditorTree::OnSearchChanged(const FText& InFilterText)
+{
+	bRefreshRequested = true;
+	SearchText = InFilterText;
+}
+
 void SUMGEditorTree::OnEditorSelectionChanged()
 {
 	WidgetTreeView->ClearSelection();
 
+	bool bFirst = true;
+
 	// Update the selection and expansion in the tree to match the new selection
 	const TSet<FWidgetReference>& SelectedWidgets = BlueprintEditor.Pin()->GetSelectedWidgets();
 	for ( const FWidgetReference& WidgetRef : SelectedWidgets )
 	{
 		WidgetTreeView->SetItemSelection(WidgetRef.GetTemplate(), true);
 
+		// Attempt to scroll the first widget we find into view.
+		if ( bFirst )
+		{
+			bFirst = false;
+			WidgetTreeView->RequestScrollIntoView(WidgetRef.GetTemplate());
+		}
+
 		// Expand the path leading to this widget in the tree.
 		UWidget* Parent = WidgetRef.GetTemplate()->GetParent();
 		while ( Parent != NULL )
diff --git a/Engine/Source/Editor/UMGEditor/Private/SUMGEditorTree.h b/Engine/Source/Editor/UMGEditor/Private/SUMGEditorTree.h
--- a/Engine/Source/Editor/UMGEditor/Private/SUMGEditorTree.h
+++ b/Engine/Source/Editor/UMGEditor/Private/SUMGEditorTree.h
@@ -48,6 +48,9 @@ private:
 
 	void WrapSelectedWidgets(UClass* WidgetClass);
 	
+	/** Called when the filter text is changed. */
+	void OnSearchChanged(const FText& InFilterText);
+
 private:
 
 	TWeakPtr<class FWidgetBlueprintEditor> BlueprintEditor;
@@ -55,4 +58,7 @@ private:
 	TArray< UWidget* > RootWidgets;
 
 	TSharedPtr< STreeView< UWidget* > > WidgetTreeView;
+
+	bool bRefreshRequested;
+	FText SearchText;
 };
diff --git a/Engine/Source/Editor/UMGEditor/Private/SUMGEditorWidgetTemplates.cpp b/Engine/Source/Editor/UMGEditor/Private/SUMGEditorWidgetTemplates.cpp
--- a/Engine/Source/Editor/UMGEditor/Private/SUMGEditorWidgetTemplates.cpp
+++ b/Engine/Source/Editor/UMGEditor/Private/SUMGEditorWidgetTemplates.cpp
@@ -130,18 +130,32 @@ void SUMGEditorWidgetTemplates::Construct(const FArguments& InArgs, TSharedPtr<F
 
 	//FCoreDelegates::OnObjectPropertyChanged.AddRaw(this, &SUMGEditorWidgetTemplates::OnObjectPropertyChanged);
 
+	SAssignNew(WidgetTemplatesView, STreeView< TSharedPtr<FWidgetViewModel> >)
+	.SelectionMode(ESelectionMode::Single)
+	.OnGenerateRow(this, &SUMGEditorWidgetTemplates::OnGenerateWidgetTemplateItem)
+	.OnGetChildren(this, &SUMGEditorWidgetTemplates::OnGetChildren)
+	.TreeItemsSource(&WidgetViewModels);
+
 	ChildSlot
 	[
 		SNew(SVerticalBox)
 
+		+ SVerticalBox::Slot()
+		.Padding(4)
+		.AutoHeight()
+		[
+			SNew(SSearchBox)
+			.HintText(LOCTEXT("SearchTemplates", "Search Templates"))
+			.OnTextChanged(this, &SUMGEditorWidgetTemplates::OnSearchChanged)
+		]
+
 		+ SVerticalBox::Slot()
 		.FillHeight(1.0f)
 		[
-			SAssignNew(WidgetTemplatesView, STreeView< TSharedPtr<FWidgetViewModel> >)
-			.OnGenerateRow(this, &SUMGEditorWidgetTemplates::OnGenerateWidgetTemplateItem)
-			.OnGetChildren(this, &SUMGEditorWidgetTemplates::OnGetChildren)
-			.TreeItemsSource(&WidgetViewModels)
-			.SelectionMode(ESelectionMode::Single)
+			SNew(SScrollBorder, WidgetTemplatesView.ToSharedRef())
+			[
+				WidgetTemplatesView.ToSharedRef()
+			]
 		]
 	];
 
@@ -154,6 +168,12 @@ SUMGEditorWidgetTemplates::~SUMGEditorWidgetTemplates()
 	SaveItemExpansion();
 }
 
+void SUMGEditorWidgetTemplates::OnSearchChanged(const FText& InFilterText)
+{
+	bRefreshRequested = true;
+	SearchText = InFilterText;
+}
+
 void SUMGEditorWidgetTemplates::LoadItemExpanssion()
 {
 	// Restore the expansion state of the widget groups.
diff --git a/Engine/Source/Editor/UMGEditor/Private/SUMGEditorWidgetTemplates.h b/Engine/Source/Editor/UMGEditor/Private/SUMGEditorWidgetTemplates.h
--- a/Engine/Source/Editor/UMGEditor/Private/SUMGEditorWidgetTemplates.h
+++ b/Engine/Source/Editor/UMGEditor/Private/SUMGEditorWidgetTemplates.h
@@ -40,6 +40,9 @@ private:
 	void OnGetChildren(TSharedPtr<FWidgetViewModel> Item, TArray< TSharedPtr<FWidgetViewModel> >& Children);
 	TSharedRef<ITableRow> OnGenerateWidgetTemplateItem(TSharedPtr<FWidgetViewModel> Item, const TSharedRef<STableViewBase>& OwnerTable);
 
+	/** Called when the filter text is changed. */
+	void OnSearchChanged(const FText& InFilterText);
+
 private:
 	void LoadItemExpanssion();
 	void SaveItemExpansion();
@@ -55,4 +58,7 @@ private:
 	ViewModelsArray WidgetViewModels;
 
 	TSharedPtr< STreeView< TSharedPtr<FWidgetViewModel> > > WidgetTemplatesView;
+
+	bool bRefreshRequested;
+	FText SearchText;
 };
