UMG - The search box in the widget hierarchy view now works.

diff --git a/Engine/Source/Editor/UMGEditor/Private/SUMGEditorTree.cpp b/Engine/Source/Editor/UMGEditor/Private/SUMGEditorTree.cpp
--- a/Engine/Source/Editor/UMGEditor/Private/SUMGEditorTree.cpp
+++ b/Engine/Source/Editor/UMGEditor/Private/SUMGEditorTree.cpp
@@ -20,6 +20,10 @@
 void SUMGEditorTree::Construct(const FArguments& InArgs, TSharedPtr<FWidgetBlueprintEditor> InBlueprintEditor, USimpleConstructionScript* InSCS)
 {
 	BlueprintEditor = InBlueprintEditor;
+	bRefreshRequested = false;
+	bIsFilterActive = false;
+
+	SearchBoxWidgetFilter = MakeShareable(new WidgetTextFilter(WidgetTextFilter::FItemToStringArray::CreateSP(this, &SUMGEditorTree::TransformWidgetToString)));
 
 	UWidgetBlueprint* Blueprint = GetBlueprint();
 	Blueprint->OnChanged().AddSP(this, &SUMGEditorTree::OnBlueprintChanged);
@@ -83,10 +87,30 @@ SUMGEditorTree::~SUMGEditorTree()
 	}
 }
 
+void SUMGEditorTree::Tick(const FGeometry& AllottedGeometry, const double InCurrentTime, const float InDeltaTime)
+{
+	if ( bRefreshRequested )
+	{
+		bRefreshRequested = false;
+
+		RefreshTree();
+	}
+}
+
+void SUMGEditorTree::TransformWidgetToString(const UWidget* Widget, OUT TArray< FString >& Array)
+{
+	Array.Add(Widget->GetName());
+}
+
 void SUMGEditorTree::OnSearchChanged(const FText& InFilterText)
 {
 	bRefreshRequested = true;
-	SearchText = InFilterText;
+	SearchBoxWidgetFilter->SetRawFilterText(InFilterText);
+}
+
+FText SUMGEditorTree::GetSearchText() const
+{
+	return SearchBoxWidgetFilter->GetRawFilterText();
 }
 
 void SUMGEditorTree::OnEditorSelectionChanged()
@@ -108,13 +132,18 @@ void SUMGEditorTree::OnEditorSelectionChanged()
 			WidgetTreeView->RequestScrollIntoView(WidgetRef.GetTemplate());
 		}
 
-		// Expand the path leading to this widget in the tree.
-		UWidget* Parent = WidgetRef.GetTemplate()->GetParent();
-		while ( Parent != NULL )
-		{
-			WidgetTreeView->SetItemExpansion(Parent, true);
-			Parent = Parent->GetParent();
-		}
+		ExpandPathToWidget(WidgetRef.GetTemplate());
+	}
+}
+
+void SUMGEditorTree::ExpandPathToWidget(UWidget* TemplateWidget)
+{
+	// Expand the path leading to this widget in the tree.
+	UWidget* Parent = TemplateWidget->GetParent();
+	while ( Parent != NULL )
+	{
+		WidgetTreeView->SetItemExpansion(Parent, true);
+		Parent = Parent->GetParent();
 	}
 }
 
@@ -220,6 +249,11 @@ void SUMGEditorTree::WidgetHierarchy_OnGetChildren(UWidget* InParent, TArray< UW
 			UWidget* Child = Widget->GetChildAt(i);
 			if ( Child )
 			{
+				if ( bIsFilterActive && !WidgetsPassingFilter.Contains(Child) )
+				{
+					continue;
+				}
+
 				OutChildren.Add(Child);
 			}
 		}
@@ -234,6 +268,7 @@ TSharedRef< ITableRow > SUMGEditorTree::WidgetHierarchy_OnGenerateRow(UWidget* I
 //		.OnDragDetected(this, &SUMGEditorWidgetTemplates::OnDraggingWidgetTemplateItem)
 		[
 			SNew(SUMGEditorTreeItem, BlueprintEditor.Pin(), InItem)
+			.HighlightText(this, &SUMGEditorTree::GetSearchText)
 		];
 }
 
@@ -290,8 +325,12 @@ void SUMGEditorTree::DeleteSelected()
 		for ( UWidget* Item : SelectedItems )
 		{
 			bRemoved = BP->WidgetTree->RemoveWidget(Item);
+
+			CachedExpandedWidgets.Remove(Item);
 		}
 
+		//TODO UMG There needs to be an event for widget removal so that caches can be updated, and selection
+
 		if ( bRemoved )
 		{
 			FBlueprintEditorUtils::MarkBlueprintAsStructurallyModified(BP);
@@ -299,6 +338,43 @@ void SUMGEditorTree::DeleteSelected()
 	}
 }
 
+bool SUMGEditorTree::FilterWidgetHierarchy(UWidget* CurrentWidget)
+{
+	bool bAnyChildrenPass = false;
+
+	// Iterate over children and check to see if any of them pass the current filter
+	UPanelWidget* Widget = Cast<UPanelWidget>(CurrentWidget);
+	if ( Widget )
+	{
+		for ( int32 i = 0; i < Widget->GetChildrenCount(); i++ )
+		{
+			UWidget* Child = Widget->GetChildAt(i);
+			if ( Child )
+			{
+				bAnyChildrenPass |= FilterWidgetHierarchy(Child);
+			}
+		}
+	}
+
+	// Check to see if I pass the filter
+	const bool bWidgetPass = SearchBoxWidgetFilter->PassesFilter(CurrentWidget);
+
+	// If this particular widget passes the filter, expand every widget leading up to it.
+	if ( bWidgetPass )
+	{
+		ExpandPathToWidget(CurrentWidget);
+	}
+
+	// If either the widget or the children pass, add the current widget
+	if ( bWidgetPass || bAnyChildrenPass )
+	{
+		WidgetsPassingFilter.Add(CurrentWidget);
+		return true;
+	}
+
+	return false;
+}
+
 void SUMGEditorTree::RefreshTree()
 {
 	RootWidgets.Reset();
@@ -308,7 +384,39 @@ void SUMGEditorTree::RefreshTree()
 
 	if ( WidgetTemplates.Num() > 0 )
 	{
-		RootWidgets.Add(WidgetTemplates[0]);
+		bool bRootPassed = true;
+
+		const bool bWillFilterBeActive = !SearchBoxWidgetFilter->GetRawFilterText().IsEmpty();
+
+		// Save the expansion state when the filter becomes active
+		if ( !bIsFilterActive && bWillFilterBeActive )
+		{
+			CachedExpandedWidgets.Empty();
+			WidgetTreeView->GetExpandedItems(CachedExpandedWidgets);
+		}
+		// Restore the expansion state when the filter is removed
+		else if ( bIsFilterActive && !bWillFilterBeActive )
+		{
+			WidgetTreeView->ClearExpandedItems();
+			for ( UWidget* ExpandedWidget : CachedExpandedWidgets )
+			{
+				WidgetTreeView->SetItemExpansion(ExpandedWidget, true);
+			}
+			CachedExpandedWidgets.Empty();
+		}
+		
+		bIsFilterActive = bWillFilterBeActive;
+
+		WidgetsPassingFilter.Empty( WidgetTemplates.Num() );
+		if ( bIsFilterActive )
+		{
+			bRootPassed = FilterWidgetHierarchy(WidgetTemplates[0]);
+		}
+
+		if ( bRootPassed )
+		{
+			RootWidgets.Add(WidgetTemplates[0]);
+		}
 	}
 
 	WidgetTreeView->RequestTreeRefresh();
diff --git a/Engine/Source/Editor/UMGEditor/Private/SUMGEditorTree.h b/Engine/Source/Editor/UMGEditor/Private/SUMGEditorTree.h
--- a/Engine/Source/Editor/UMGEditor/Private/SUMGEditorTree.h
+++ b/Engine/Source/Editor/UMGEditor/Private/SUMGEditorTree.h
@@ -13,6 +13,9 @@
  */
 class SUMGEditorTree : public SCompoundWidget
 {
+public:
+	typedef TTextFilter< const UWidget* > WidgetTextFilter;
+
 public:
 	SLATE_BEGIN_ARGS( SUMGEditorTree ){}
 	SLATE_END_ARGS()
@@ -21,10 +24,22 @@ public:
 
 	virtual ~SUMGEditorTree();
 
+	virtual void Tick(const FGeometry& AllottedGeometry, const double InCurrentTime, const float InDeltaTime);
+
+private:
+	void BuildWrapWithMenu(FMenuBuilder& Menu);
+	TSharedPtr<SWidget> WidgetHierarchy_OnContextMenuOpening();
+	void WidgetHierarchy_OnGetChildren(UWidget* InParent, TArray< UWidget* >& OutChildren);
+	TSharedRef< ITableRow > WidgetHierarchy_OnGenerateRow(UWidget* InItem, const TSharedRef<STableViewBase>& OwnerTable);
+	void WidgetHierarchy_OnSelectionChanged(UWidget* SelectedItem, ESelectInfo::Type SelectInfo);
+
 private:
 	/** @returns the current blueprint being edited */
 	UWidgetBlueprint* GetBlueprint() const;
 
+	/** Expands every item in the tree leading to this widget */
+	void ExpandPathToWidget(UWidget* TemplateWidget);
+
 	/** Called when the blueprint is structurally changed. */
 	void OnBlueprintChanged(UBlueprint* InBlueprint);
 
@@ -33,24 +48,33 @@ private:
 	/** Called when the selected widget has changed.  The treeview then needs to match the new selection. */
 	void OnEditorSelectionChanged();
 
+	/** Notifies the details panel that new widgets have been selected. */
 	void ShowDetailsForObjects(TArray<UWidget*> Widgets);
 
-	void BuildWrapWithMenu(FMenuBuilder& Menu);
-	TSharedPtr<SWidget> WidgetHierarchy_OnContextMenuOpening();
-	void WidgetHierarchy_OnGetChildren(UWidget* InParent, TArray< UWidget* >& OutChildren);
-	TSharedRef< ITableRow > WidgetHierarchy_OnGenerateRow(UWidget* InItem, const TSharedRef<STableViewBase>& OwnerTable);
-	void WidgetHierarchy_OnSelectionChanged(UWidget* SelectedItem, ESelectInfo::Type SelectInfo);
-
+	/** Rebuilds the tree structure based on the current filter options */
 	void RefreshTree();
 
+	/** Handles the menu clicking to delete the selected widgets. */
 	FReply HandleDeleteSelected();
+
+	/** Deletes the selected widgets from the hierarchy */
 	void DeleteSelected();
 
+	/** Wraps the selected widget using a user selected panel */
 	void WrapSelectedWidgets(UClass* WidgetClass);
 	
 	/** Called when the filter text is changed. */
 	void OnSearchChanged(const FText& InFilterText);
 
+	/** Gets the search text currently being used to filter the list, also used to highlight text */
+	FText GetSearchText() const;
+
+	/** Transforms the widget into a searchable string */
+	void TransformWidgetToString(const UWidget* Widget, OUT TArray< FString >& Array);
+
+	/** Builds a filterable set by recursing through the widget hierarchy */
+	bool FilterWidgetHierarchy(UWidget* CurrentWidget);
+
 private:
 
 	TWeakPtr<class FWidgetBlueprintEditor> BlueprintEditor;
@@ -59,6 +83,13 @@ private:
 
 	TSharedPtr< STreeView< UWidget* > > WidgetTreeView;
 
+	TSharedPtr<WidgetTextFilter> SearchBoxWidgetFilter;
+
+	TSet< UWidget* > WidgetsPassingFilter;
+
+	TSet< UWidget* > CachedExpandedWidgets;
+
+	bool bIsFilterActive;
+
 	bool bRefreshRequested;
-	FText SearchText;
 };
diff --git a/Engine/Source/Editor/UMGEditor/Private/SUMGEditorTreeItem.cpp b/Engine/Source/Editor/UMGEditor/Private/SUMGEditorTreeItem.cpp
--- a/Engine/Source/Editor/UMGEditor/Private/SUMGEditorTreeItem.cpp
+++ b/Engine/Source/Editor/UMGEditor/Private/SUMGEditorTreeItem.cpp
@@ -23,6 +23,7 @@ void SUMGEditorTreeItem::Construct(const FArguments& InArgs, TSharedPtr<FBluepri
 		SNew(STextBlock)
 		.ToolTipText(this, &SUMGEditorTreeItem::GetItemTooltipText)
 		.Text(this, &SUMGEditorTreeItem::GetItemText)
+		.HighlightText(InArgs._HighlightText)
 	];
 }
 
diff --git a/Engine/Source/Editor/UMGEditor/Private/SUMGEditorTreeItem.h b/Engine/Source/Editor/UMGEditor/Private/SUMGEditorTreeItem.h
--- a/Engine/Source/Editor/UMGEditor/Private/SUMGEditorTreeItem.h
+++ b/Engine/Source/Editor/UMGEditor/Private/SUMGEditorTreeItem.h
@@ -5,11 +5,17 @@
 #include "SCompoundWidget.h"
 #include "BlueprintEditor.h"
 
-//TODO rename SUMGEditorHierarchy
+/**
+ * An widget item in the hierarchy tree view.
+ */
 class SUMGEditorTreeItem : public SCompoundWidget
 {
 public:
 	SLATE_BEGIN_ARGS( SUMGEditorTreeItem ){}
+		
+		/** The current text to highlight */
+		SLATE_ATTRIBUTE( FText, HighlightText )
+
 	SLATE_END_ARGS()
 
 	void Construct(const FArguments& InArgs, TSharedPtr<FBlueprintEditor> InBlueprintEditor, UWidget* InItem);
diff --git a/Engine/Source/Editor/UMGEditor/Private/UMGEditorPrivatePCH.h b/Engine/Source/Editor/UMGEditor/Private/UMGEditorPrivatePCH.h
--- a/Engine/Source/Editor/UMGEditor/Private/UMGEditorPrivatePCH.h
+++ b/Engine/Source/Editor/UMGEditor/Private/UMGEditorPrivatePCH.h
@@ -10,6 +10,7 @@
 #include "BlueprintEditorUtils.h"
 
 #include "ScopedTransaction.h"
+#include "TextFilter.h"
 
 #include "UMG.h"
 
