UMG - Fixing fixing null parent crash.

diff --git a/Engine/Source/Editor/UMGEditor/Private/SUMGEditorTree.cpp b/Engine/Source/Editor/UMGEditor/Private/SUMGEditorTree.cpp
--- a/Engine/Source/Editor/UMGEditor/Private/SUMGEditorTree.cpp
+++ b/Engine/Source/Editor/UMGEditor/Private/SUMGEditorTree.cpp
@@ -123,16 +123,21 @@ void SUMGEditorTree::OnEditorSelectionChanged()
 	const TSet<FWidgetReference>& SelectedWidgets = BlueprintEditor.Pin()->GetSelectedWidgets();
 	for ( const FWidgetReference& WidgetRef : SelectedWidgets )
 	{
-		WidgetTreeView->SetItemSelection(WidgetRef.GetTemplate(), true);
+		UWidget* TemplateWidget = WidgetRef.GetTemplate();
 
-		// Attempt to scroll the first widget we find into view.
-		if ( bFirst )
+		if ( TemplateWidget )
 		{
-			bFirst = false;
-			WidgetTreeView->RequestScrollIntoView(WidgetRef.GetTemplate());
-		}
+			WidgetTreeView->SetItemSelection(TemplateWidget, true);
 
-		ExpandPathToWidget(WidgetRef.GetTemplate());
+			// Attempt to scroll the first widget we find into view.
+			if ( bFirst )
+			{
+				bFirst = false;
+				WidgetTreeView->RequestScrollIntoView(TemplateWidget);
+			}
+
+			ExpandPathToWidget(TemplateWidget);
+		}
 	}
 }
 
