Fix UMG editor not shutting down properly after selection

diff --git a/Engine/Source/Editor/UMGEditor/Private/Extensions/DesignerExtension.cpp b/Engine/Source/Editor/UMGEditor/Private/Extensions/DesignerExtension.cpp
--- a/Engine/Source/Editor/UMGEditor/Private/Extensions/DesignerExtension.cpp
+++ b/Engine/Source/Editor/UMGEditor/Private/Extensions/DesignerExtension.cpp
@@ -20,24 +20,24 @@ FWidgetReference::FWidgetReference(TSharedPtr<FWidgetBlueprintEditor> WidgetEdit
 {
 }
 
-FWidgetReference FWidgetReference::FromTemplate(TSharedPtr<FWidgetBlueprintEditor> WidgetEditor, UWidget* TemplateWidget)
+FWidgetReference FWidgetReference::FromTemplate(TSharedPtr<FWidgetBlueprintEditor> InWidgetEditor, UWidget* TemplateWidget)
 {
-	return FWidgetReference(WidgetEditor, TemplateWidget);
+	return FWidgetReference(InWidgetEditor, TemplateWidget);
 }
 
-FWidgetReference FWidgetReference::FromPreview(TSharedPtr<FWidgetBlueprintEditor> WidgetEditor, UWidget* PreviewWidget)
+FWidgetReference FWidgetReference::FromPreview(TSharedPtr<FWidgetBlueprintEditor> InWidgetEditor, UWidget* PreviewWidget)
 {
-	if ( WidgetEditor.IsValid() )
+	if ( InWidgetEditor.IsValid() )
 	{
-		UUserWidget* PreviewRoot = WidgetEditor->GetPreview();
+		UUserWidget* PreviewRoot = InWidgetEditor->GetPreview();
 		if ( PreviewRoot )
 		{
-			UWidgetBlueprint* Blueprint = WidgetEditor->GetWidgetBlueprintObj();
+			UWidgetBlueprint* Blueprint = InWidgetEditor->GetWidgetBlueprintObj();
 
 			if ( PreviewWidget )
 			{
 				FString Name = PreviewWidget->GetName();
-				return FromTemplate(WidgetEditor, Blueprint->WidgetTree->FindWidget(Name));
+				return FromTemplate(InWidgetEditor, Blueprint->WidgetTree->FindWidget(Name));
 			}
 		}
 	}
@@ -59,7 +59,7 @@ UWidget* FWidgetReference::GetPreview() const
 {
 	if ( WidgetEditor.IsValid() )
 	{
-		UUserWidget* PreviewRoot = WidgetEditor->GetPreview();
+		UUserWidget* PreviewRoot = WidgetEditor.Pin()->GetPreview();
 
 		if ( PreviewRoot && TemplateWidget )
 		{
diff --git a/Engine/Source/Editor/UMGEditor/Public/WidgetReference.h b/Engine/Source/Editor/UMGEditor/Public/WidgetReference.h
--- a/Engine/Source/Editor/UMGEditor/Public/WidgetReference.h
+++ b/Engine/Source/Editor/UMGEditor/Public/WidgetReference.h
@@ -16,12 +16,12 @@ public:
 	/**
 	 * Creates a widget reference using the template.
 	 */
-	static FWidgetReference FromTemplate(TSharedPtr<class FWidgetBlueprintEditor> WidgetEditor, UWidget* TemplateWidget);
+	static FWidgetReference FromTemplate(TSharedPtr<class FWidgetBlueprintEditor> InWidgetEditor, UWidget* TemplateWidget);
 
 	/**
 	 * Creates a widget reference using the preview.  Which is used to lookup the stable template pointer.
 	 */
-	static FWidgetReference FromPreview(TSharedPtr<class FWidgetBlueprintEditor> WidgetEditor, UWidget* PreviewWidget);
+	static FWidgetReference FromPreview(TSharedPtr<class FWidgetBlueprintEditor> InWidgetEditor, UWidget* PreviewWidget);
 
 	/** Constructor for the null widget reference.  Not intended for normal use. */
 	FWidgetReference();
@@ -46,7 +46,7 @@ private:
 
 private:
 
-	TSharedPtr<class FWidgetBlueprintEditor> WidgetEditor;
+	TWeakPtr<class FWidgetBlueprintEditor> WidgetEditor;
 
 	UWidget* TemplateWidget;
 };
