UMG - Open to the designer by default, add canvas if no other widgets have been added.

diff --git a/Engine/Source/Editor/UMGEditor/Private/AssetTypeActions_WidgetBlueprint.cpp b/Engine/Source/Editor/UMGEditor/Private/AssetTypeActions_WidgetBlueprint.cpp
--- a/Engine/Source/Editor/UMGEditor/Private/AssetTypeActions_WidgetBlueprint.cpp
+++ b/Engine/Source/Editor/UMGEditor/Private/AssetTypeActions_WidgetBlueprint.cpp
@@ -29,7 +29,7 @@ void FAssetTypeActions_WidgetBlueprint::OpenAssetEditor( const TArray<UObject*>&
 
 			TArray<UBlueprint*> Blueprints;
 			Blueprints.Add(Blueprint);
-			NewBlueprintEditor->InitWidgetBlueprintEditor(Mode, EditWithinLevelEditor, Blueprints, ShouldUseDataOnlyEditor(Blueprint));
+			NewBlueprintEditor->InitWidgetBlueprintEditor(Mode, EditWithinLevelEditor, Blueprints, true);
 
 			//FBlueprintEditorModule& BlueprintEditorModule = FModuleManager::LoadModuleChecked<FBlueprintEditorModule>( "Kismet" );
 			//TSharedRef< IBlueprintEditor > NewKismetEditor = BlueprintEditorModule.CreateBlueprintEditor( Mode, EditWithinLevelEditor, Blueprint, ShouldUseDataOnlyEditor(Blueprint) );
diff --git a/Engine/Source/Editor/UMGEditor/Private/SUMGEditorTree.cpp b/Engine/Source/Editor/UMGEditor/Private/SUMGEditorTree.cpp
--- a/Engine/Source/Editor/UMGEditor/Private/SUMGEditorTree.cpp
+++ b/Engine/Source/Editor/UMGEditor/Private/SUMGEditorTree.cpp
@@ -68,7 +68,7 @@ void SUMGEditorTree::Construct(const FArguments& InArgs, TSharedPtr<FWidgetBluep
 
 	BlueprintEditor.Pin()->OnSelectedWidgetsChanged.AddRaw(this, &SUMGEditorTree::OnEditorSelectionChanged);
 
-	RefreshTree();
+	bRefreshRequested = true;
 }
 
 SUMGEditorTree::~SUMGEditorTree()
diff --git a/Engine/Source/Editor/UMGEditor/Private/WidgetBlueprintEditor.cpp b/Engine/Source/Editor/UMGEditor/Private/WidgetBlueprintEditor.cpp
--- a/Engine/Source/Editor/UMGEditor/Private/WidgetBlueprintEditor.cpp
+++ b/Engine/Source/Editor/UMGEditor/Private/WidgetBlueprintEditor.cpp
@@ -76,6 +76,14 @@ void FWidgetBlueprintEditor::InitWidgetBlueprintEditor(const EToolkitMode::Type
 {
 	InitBlueprintEditor(Mode, InitToolkitHost, InBlueprints, bShouldOpenInDefaultsMode);
 
+	UWidgetBlueprint* Blueprint = GetWidgetBlueprintObj();
+
+	// If this blueprint is empty, add a canvas panel as the root widget.
+	if ( Blueprint->WidgetTree->WidgetTemplates.Num() == 0 )
+	{
+		Blueprint->WidgetTree->ConstructWidget<UCanvasPanel>(UCanvasPanel::StaticClass());
+	}
+
 	UpdatePreview(GetWidgetBlueprintObj(), true);
 }
 
