UMG - Adding a hover effect.

diff --git a/Engine/Source/Editor/UMGEditor/Private/SUMGDesigner.cpp b/Engine/Source/Editor/UMGEditor/Private/SUMGDesigner.cpp
--- a/Engine/Source/Editor/UMGEditor/Private/SUMGDesigner.cpp
+++ b/Engine/Source/Editor/UMGEditor/Private/SUMGDesigner.cpp
@@ -69,12 +69,14 @@ void SUMGDesigner::Construct(const FArguments& InArgs, TSharedPtr<FWidgetBluepri
 {
 	ScopedTransaction = NULL;
 
-	PreviewWidgetObject = NULL;
+	PreviewWidget = NULL;
 	DropPreviewWidget = NULL;
 	DropPreviewParent = NULL;
 	BlueprintEditor = InBlueprintEditor;
 	CurrentHandle = DH_NONE;
 
+	HoverTime = 0;
+
 	DragDirections.Init((int32)DH_MAX);
 	DragDirections[DH_TOP_LEFT] = FVector2D(-1, -1);
 	DragDirections[DH_TOP_CENTER] = FVector2D(0, -1);
@@ -215,16 +217,16 @@ FReply SUMGDesigner::OnMouseButtonDown(const FGeometry& MyGeometry, const FPoint
 	if ( CurrentHandle == DH_NONE )
 	{
 		FArrangedWidget ArrangedWidget(SNullWidget::NullWidget, FGeometry());
-		CurrentSelection = GetWidgetAtCursor(MyGeometry, MouseEvent, ArrangedWidget);
+		SelectedWidget = GetWidgetAtCursor(MyGeometry, MouseEvent, ArrangedWidget);
 
 		//TODO UMG Undoable Selection
 
-		if ( CurrentSelection.IsValid() )
+		if ( SelectedWidget.IsValid() )
 		{
 			//@TODO UMG primary FBlueprintEditor needs to be inherited and selection control needs to be centralized.
 			// Set the template as selected in the details panel
 			TSet<FWidgetReference> SelectedTemplates;
-			SelectedTemplates.Add(CurrentSelection);
+			SelectedTemplates.Add(SelectedWidget);
 			BlueprintEditor.Pin()->SelectWidgets(SelectedTemplates);
 
 			// Remove all the current extension widgets
@@ -233,7 +235,7 @@ FReply SUMGDesigner::OnMouseButtonDown(const FGeometry& MyGeometry, const FPoint
 			ExtensionWidgets.Reset();
 
 			TArray<FWidgetReference> Selected;
-			Selected.Add(CurrentSelection);
+			Selected.Add(SelectedWidget);
 
 			// Build extension widgets for new selection
 			for ( TSharedRef<FDesignerExtension>& Ext : DesignerExtensions )
@@ -282,15 +284,21 @@ FReply SUMGDesigner::OnMouseButtonUp(const FGeometry& MyGeometry, const FPointer
 
 FReply SUMGDesigner::OnMouseMove(const FGeometry& MyGeometry, const FPointerEvent& MouseEvent)
 {
+	if ( MouseEvent.GetCursorDelta().IsZero() )
+	{
+		return FReply::Unhandled();
+	}
+
+	// Update the resizing based on the mouse movement
 	if ( CurrentHandle != DH_NONE )
 	{
-		if ( CurrentSelection.IsValid() && !MouseEvent.GetCursorDelta().IsZero() )
+		if ( SelectedWidget.IsValid() )
 		{
-			//@TODO UMG - Implement some system to query slots to know if they support dragging so we can provide visual feedback by hiding handles that wouldnt work and such.
+			//@TODO UMG - Implement some system to query slots to know if they support dragging so we can provide visual feedback by hiding handles that wouldn't work and such.
 			//SelectedTemplate->Slot->CanResize(DH_TOP_CENTER)
 
-			UWidget* Template = CurrentSelection.GetTemplate();
-			UWidget* Preview = CurrentSelection.GetPreview();
+			UWidget* Template = SelectedWidget.GetTemplate();
+			UWidget* Preview = SelectedWidget.GetPreview();
 
 			if ( Preview->Slot )
 			{
@@ -308,9 +316,33 @@ FReply SUMGDesigner::OnMouseMove(const FGeometry& MyGeometry, const FPointerEven
 		}
 	}
 
+	// Update the hovered widget under the mouse
+	if ( CurrentHandle == DH_NONE )
+	{
+		FArrangedWidget ArrangedWidget(SNullWidget::NullWidget, FGeometry());
+		FWidgetReference NewHoveredWidget = GetWidgetAtCursor(MyGeometry, MouseEvent, ArrangedWidget);
+		if ( !( NewHoveredWidget == HoveredWidget ) )
+		{
+			HoveredWidget = NewHoveredWidget;
+			HoverTime = 0;
+		}
+	}
+
 	return FReply::Unhandled();
 }
 
+void SUMGDesigner::OnMouseEnter(const FGeometry& MyGeometry, const FPointerEvent& MouseEvent)
+{
+	HoveredWidget = FWidgetReference();
+	HoverTime = 0;
+}
+
+void SUMGDesigner::OnMouseLeave(const FPointerEvent& MouseEvent)
+{
+	HoveredWidget = FWidgetReference();
+	HoverTime = 0;
+}
+
 bool SUMGDesigner::GetArrangedWidget(TSharedRef<SWidget> Widget, FArrangedWidget& ArrangedWidget) const
 {
 	// We can't screenshot the widget unless there's a valid window handle to draw it in.
@@ -376,9 +408,9 @@ FVector2D SUMGDesigner::GetCachedSelectionDesignerWidgetsLocation() const
 
 FVector2D SUMGDesigner::GetSelectionDesignerWidgetsLocation() const
 {
-	if ( SelectedWidget.IsValid() )
+	if ( SelectedSlateWidget.IsValid() )
 	{
-		TSharedRef<SWidget> Widget = SelectedWidget.Pin().ToSharedRef();
+		TSharedRef<SWidget> Widget = SelectedSlateWidget.Pin().ToSharedRef();
 
 		FArrangedWidget ArrangedWidget(SNullWidget::NullWidget, FGeometry());
 		GetArrangedWidgetRelativeToDesigner(Widget, ArrangedWidget);
@@ -393,17 +425,48 @@ int32 SUMGDesigner::OnPaint(const FGeometry& AllottedGeometry, const FSlateRect&
 {
 	SDesignSurface::OnPaint(AllottedGeometry, MyClippingRect, OutDrawElements, LayerId, InWidgetStyle, bParentEnabled);
 
-	if ( SelectedWidget.IsValid() )
+	LayerId += 1000;
+
+	// Don't draw the hovered effect if it's also the selected widget
+	if ( HoveredSlateWidget.IsValid() && HoveredSlateWidget != SelectedSlateWidget )
 	{
-		TSharedRef<SWidget> Widget = SelectedWidget.Pin().ToSharedRef();
+		TSharedRef<SWidget> Widget = HoveredSlateWidget.Pin().ToSharedRef();
 
 		FArrangedWidget ArrangedWidget(SNullWidget::NullWidget, FGeometry());
 		GetArrangedWidgetRelativeToWindow(Widget, ArrangedWidget);
 
-		const FLinearColor Tint(0, 1, 0);
+		const float HoveredAnimationTime = 0.25f;
+
+		// Draw hovered effect
+		// Azure = 0x007FFF
+		const FLinearColor HoveredTint(0, 0.5, 1, FMath::Clamp(HoverTime / HoveredAnimationTime, 0.0f, 1.0f));
+
+		FPaintGeometry HoveredGeometry(
+			ArrangedWidget.Geometry.AbsolutePosition,
+			ArrangedWidget.Geometry.Size * ArrangedWidget.Geometry.Scale,
+			ArrangedWidget.Geometry.Scale);
 
-		LayerId += 100;
+		FSlateDrawElement::MakeBox(
+			OutDrawElements,
+			LayerId,
+			HoveredGeometry,
+			FCoreStyle::Get().GetBrush(TEXT("Debug.Border")),
+			MyClippingRect,
+			ESlateDrawEffect::None,
+			HoveredTint
+			);
+	}
+
+	if ( SelectedSlateWidget.IsValid() )
+	{
+		TSharedRef<SWidget> Widget = SelectedSlateWidget.Pin().ToSharedRef();
+
+		FArrangedWidget ArrangedWidget(SNullWidget::NullWidget, FGeometry());
+		GetArrangedWidgetRelativeToWindow(Widget, ArrangedWidget);
 
+		const FLinearColor Tint(0, 1, 0);
+
+		// Draw selection effect
 		FPaintGeometry SelectionGeometry(
 			ArrangedWidget.Geometry.AbsolutePosition,
 			ArrangedWidget.Geometry.Size * ArrangedWidget.Geometry.Scale,
@@ -421,7 +484,9 @@ int32 SUMGDesigner::OnPaint(const FGeometry& AllottedGeometry, const FSlateRect&
 
 		DrawDragHandles(SelectionGeometry, MyClippingRect, OutDrawElements, LayerId, InWidgetStyle);
 
-		if ( CurrentSelection.IsValid() && CurrentSelection.GetTemplate()->Slot )
+		//TODO UMG Move this to a function.
+		// Draw anchors.
+		if ( SelectedWidget.IsValid() && SelectedWidget.GetTemplate()->Slot )
 		{
 			const float X = AllottedGeometry.AbsolutePosition.X;
 			const float Y = AllottedGeometry.AbsolutePosition.Y;
@@ -440,7 +505,7 @@ int32 SUMGDesigner::OnPaint(const FGeometry& AllottedGeometry, const FSlateRect&
 			const FSlateBrush* KeyBrush = FEditorStyle::GetBrush("CurveEd.CurveKey");
 			FLinearColor KeyColor = FLinearColor(1.0f, 0.0f, 0.0f, 1.f);
 
-			if ( UCanvasPanelSlot* CanvasSlot = Cast<UCanvasPanelSlot>(CurrentSelection.GetTemplate()->Slot) )
+			if ( UCanvasPanelSlot* CanvasSlot = Cast<UCanvasPanelSlot>(SelectedWidget.GetTemplate()->Slot) )
 			{
 				//Set("UMGEditor.AnchorCenter", new IMAGE_BRUSH("Icons/umg_anchor_center", Icon16x16));
 				//Set("UMGEditor.AnchorTopBottom", new IMAGE_BRUSH("Icons/umg_anchor_top_bottom", FVector2D(16, 32)));
@@ -499,7 +564,7 @@ int32 SUMGDesigner::OnPaint(const FGeometry& AllottedGeometry, const FSlateRect&
 
 void SUMGDesigner::DrawDragHandles(const FPaintGeometry& SelectionGeometry, const FSlateRect& MyClippingRect, FSlateWindowElementList& OutDrawElements, int32 LayerId, const FWidgetStyle& InWidgetStyle) const
 {
-	if ( CurrentSelection.IsValid() && CurrentSelection.GetTemplate()->Slot )
+	if ( SelectedWidget.IsValid() && SelectedWidget.GetTemplate()->Slot )
 	{
 		float X = SelectionGeometry.DrawPosition.X;
 		float Y = SelectionGeometry.DrawPosition.Y;
@@ -528,7 +593,7 @@ void SUMGDesigner::DrawDragHandles(const FPaintGeometry& SelectionGeometry, cons
 		for ( int32 HandleIndex = 0; HandleIndex < Handles.Num(); HandleIndex++ )
 		{
 			const FVector2D& Handle = Handles[HandleIndex];
-			if ( !CurrentSelection.GetTemplate()->Slot->CanResize(DragDirections[HandleIndex]) )
+			if ( !SelectedWidget.GetTemplate()->Slot->CanResize(DragDirections[HandleIndex]) )
 			{
 				// This isn't a valid handle
 				continue;
@@ -551,9 +616,9 @@ SUMGDesigner::DragHandle SUMGDesigner::HitTestDragHandles(const FGeometry& Allot
 {
 	FVector2D LocalPointer = AllottedGeometry.AbsoluteToLocal(PointerEvent.GetScreenSpacePosition());
 
-	if ( CurrentSelection.IsValid() && CurrentSelection.GetTemplate()->Slot && SelectedWidget.IsValid() )
+	if ( SelectedWidget.IsValid() && SelectedWidget.GetTemplate()->Slot && SelectedSlateWidget.IsValid() )
 	{
-		TSharedRef<SWidget> Widget = SelectedWidget.Pin().ToSharedRef();
+		TSharedRef<SWidget> Widget = SelectedSlateWidget.Pin().ToSharedRef();
 
 		FArrangedWidget ArrangedWidget(SNullWidget::NullWidget, FGeometry());
 		GetArrangedWidget(Widget, ArrangedWidget);
@@ -584,7 +649,7 @@ SUMGDesigner::DragHandle SUMGDesigner::HitTestDragHandles(const FGeometry& Allot
 			FSlateRect Rect(FVector2D(Handle.X - HandleSize.X * 0.5f, Handle.Y - HandleSize.Y * 0.5f), Handle + HandleSize);
 			if ( Rect.ContainsPoint(LocalPointer) )
 			{
-				if ( !CurrentSelection.GetTemplate()->Slot->CanResize(DragDirections[i]) )
+				if ( !SelectedWidget.GetTemplate()->Slot->CanResize(DragDirections[i]) )
 				{
 					// This isn't a valid handle
 					break;
@@ -623,18 +688,20 @@ FCursorReply SUMGDesigner::OnCursorQuery(const FGeometry& MyGeometry, const FPoi
 
 void SUMGDesigner::Tick(const FGeometry& AllottedGeometry, const double InCurrentTime, const float InDeltaTime)
 {
+	HoverTime += InDeltaTime;
+
 	UUserWidget* LatestPreviewWidget = BlueprintEditor.Pin()->GetPreview();
 
-	if ( LatestPreviewWidget != PreviewWidgetObject )
+	if ( LatestPreviewWidget != PreviewWidget )
 	{
-		PreviewWidgetObject = LatestPreviewWidget;
-		if ( PreviewWidgetObject )
+		PreviewWidget = LatestPreviewWidget;
+		if ( PreviewWidget )
 		{
-			TSharedRef<SWidget> CurrentWidget = PreviewWidgetObject->MakeFullScreenWidget();
+			TSharedRef<SWidget> CurrentWidget = PreviewWidget->MakeFullScreenWidget();
 
-			if ( CurrentWidget != PreviewWidget.Pin() )
+			if ( CurrentWidget != PreviewSlateWidget.Pin() )
 			{
-				PreviewWidget = CurrentWidget;
+				PreviewSlateWidget = CurrentWidget;
 				PreviewSurface->SetContent(CurrentWidget);
 			}
 		}
@@ -655,15 +722,25 @@ void SUMGDesigner::Tick(const FGeometry& AllottedGeometry, const double InCurren
 	}
 
 	// Update the selected widget to match the selected template.
-	if ( CurrentSelection.IsValid() )
+	if ( LatestPreviewWidget )
 	{
-		if ( LatestPreviewWidget )
+		if ( SelectedWidget.IsValid() )
 		{
 			// Set the selected widget so that we can draw the highlight
-			//@TODO UMG Don't store the transient Widget as the selected object, store the template component and lookup the widget from that.
-			//SelectedWidget = ArrangedWidget.Widget;
+			SelectedSlateWidget = LatestPreviewWidget->GetWidgetFromName(SelectedWidget.GetTemplate()->GetName());
+		}
+		else
+		{
+			SelectedSlateWidget.Reset();
+		}
 
-			SelectedWidget = LatestPreviewWidget->GetWidgetFromName(CurrentSelection.GetTemplate()->GetName());
+		if ( HoveredWidget.IsValid() )
+		{
+			HoveredSlateWidget = LatestPreviewWidget->GetWidgetFromName(HoveredWidget.GetTemplate()->GetName());
+		}
+		else
+		{
+			HoveredSlateWidget.Reset();
 		}
 	}
 
@@ -800,7 +877,7 @@ bool SUMGDesigner::AddToTemplate(const FGeometry& MyGeometry, const FDragDropEve
 				//@TODO UMG We may need a desired size canvas, where the slots have no size, they only give you position, alternatively, maybe slots that don't clip, so center is still easy.
 				
 				// Update the selected template to be the newly created one.
-				CurrentSelection = FWidgetReference::FromTemplate(BlueprintEditor.Pin(), Widget);
+				SelectedWidget = FWidgetReference::FromTemplate(BlueprintEditor.Pin(), Widget);
 				
 				return true;
 			}
@@ -818,7 +895,7 @@ bool SUMGDesigner::AddToTemplate(const FGeometry& MyGeometry, const FDragDropEve
 			UWidget* Widget = DragDropOp->Template->Create(BP->WidgetTree);
 			Widget->IsDesignTime(true);
 			
-			CurrentSelection = FWidgetReference::FromTemplate(BlueprintEditor.Pin(), Widget);
+			SelectedWidget = FWidgetReference::FromTemplate(BlueprintEditor.Pin(), Widget);
 			
 			return true;
 		}
@@ -859,10 +936,10 @@ void SUMGDesigner::BeginTransaction(const FText& SessionName)
 		ScopedTransaction = new FScopedTransaction(SessionName);
 	}
 
-	if ( CurrentSelection.IsValid() )
+	if ( SelectedWidget.IsValid() )
 	{
-		CurrentSelection.GetPreview()->Modify();
-		CurrentSelection.GetTemplate()->Modify();
+		SelectedWidget.GetPreview()->Modify();
+		SelectedWidget.GetTemplate()->Modify();
 	}
 }
 
diff --git a/Engine/Source/Editor/UMGEditor/Private/SUMGDesigner.h b/Engine/Source/Editor/UMGEditor/Private/SUMGDesigner.h
--- a/Engine/Source/Editor/UMGEditor/Private/SUMGDesigner.h
+++ b/Engine/Source/Editor/UMGEditor/Private/SUMGDesigner.h
@@ -26,6 +26,8 @@ public:
 	virtual FReply OnMouseButtonDown(const FGeometry& MyGeometry, const FPointerEvent& MouseEvent) override;
 	virtual FReply OnMouseButtonUp(const FGeometry& MyGeometry, const FPointerEvent& MouseEvent) override;
 	virtual FReply OnMouseMove(const FGeometry& MyGeometry, const FPointerEvent& MouseEvent) override;
+	virtual void OnMouseEnter(const FGeometry& MyGeometry, const FPointerEvent& MouseEvent) override;
+	virtual void OnMouseLeave(const FPointerEvent& MouseEvent) override;
 
 	virtual FCursorReply OnCursorQuery(const FGeometry& MyGeometry, const FPointerEvent& CursorEvent) const override;
 
@@ -97,13 +99,17 @@ private:
 	bool AddToTemplate(const FGeometry& MyGeometry, const FDragDropEvent& DragDropEvent);
 	
 private:
+	/** A reference to the BP Editor that owns this designer */
 	TWeakPtr<FWidgetBlueprintEditor> BlueprintEditor;
 
-	TWeakPtr<SWidget> PreviewWidget;
-	UUserWidget* PreviewWidgetObject;
+	/** The transaction used to commit undoable actions from resize, move...etc */
+	FScopedTransaction* ScopedTransaction;
+
+	/** The current preview widget */
+	UUserWidget* PreviewWidget;
 
-	FWidgetReference CurrentSelection;
-	TWeakPtr<SWidget> SelectedWidget;
+	/** The current preview widget's slate widget */
+	TWeakPtr<SWidget> PreviewSlateWidget;
 	
 	UWidget* DropPreviewWidget;
 	UPanelWidget* DropPreviewParent;
@@ -115,5 +121,18 @@ private:
 
 	FVector2D CachedDesignerWidgetLocation;
 
-	FScopedTransaction* ScopedTransaction;
+	/** The currently selected widget */
+	FWidgetReference SelectedWidget;
+
+	/** The currently selected slate widget, this is refreshed every frame in case it changes */
+	TWeakPtr<SWidget> SelectedSlateWidget;
+
+	/** The wall clock time the user has been hovering over a single widget */
+	float HoverTime;
+
+	/** The current widget being hovered */
+	FWidgetReference HoveredWidget;
+
+	/** The current slate widget being hovered, this is refreshed every frame in case it changes */
+	TWeakPtr<SWidget> HoveredSlateWidget;
 };
