UMG - Nasty startup crash fixed (maybe) with BPW's including other BPWs, don't use old assets.  Partial renaming support.  More immediate runtime customization of some widget properties (Border, Button).  Started prototyping inline editable designer elements, needs more thought, dealing with passing events down with the invisible barrier is going to be a challenge.

diff --git a/Engine/Source/Editor/UMGEditor/Private/SUMGEditorTree.cpp b/Engine/Source/Editor/UMGEditor/Private/SUMGEditorTree.cpp
--- a/Engine/Source/Editor/UMGEditor/Private/SUMGEditorTree.cpp
+++ b/Engine/Source/Editor/UMGEditor/Private/SUMGEditorTree.cpp
@@ -267,14 +267,8 @@ void SUMGEditorTree::WidgetHierarchy_OnGetChildren(UWidget* InParent, TArray< UW
 
 TSharedRef< ITableRow > SUMGEditorTree::WidgetHierarchy_OnGenerateRow(UWidget* InItem, const TSharedRef<STableViewBase>& OwnerTable)
 {
-	return
-		SNew(STableRow< UWidget* >, OwnerTable)
-		.Padding(2.0f)
-//		.OnDragDetected(this, &SUMGEditorWidgetTemplates::OnDraggingWidgetTemplateItem)
-		[
-			SNew(SUMGEditorTreeItem, BlueprintEditor.Pin(), InItem)
-			.HighlightText(this, &SUMGEditorTree::GetSearchText)
-		];
+	return SNew(SUMGEditorTreeItem, OwnerTable, BlueprintEditor.Pin(), InItem)
+		.HighlightText(this, &SUMGEditorTree::GetSearchText);
 }
 
 void SUMGEditorTree::WidgetHierarchy_OnSelectionChanged(UWidget* SelectedItem, ESelectInfo::Type SelectInfo)
diff --git a/Engine/Source/Editor/UMGEditor/Private/SUMGEditorTreeItem.cpp b/Engine/Source/Editor/UMGEditor/Private/SUMGEditorTreeItem.cpp
--- a/Engine/Source/Editor/UMGEditor/Private/SUMGEditorTreeItem.cpp
+++ b/Engine/Source/Editor/UMGEditor/Private/SUMGEditorTreeItem.cpp
@@ -13,23 +13,77 @@
 #include "SKismetInspector.h"
 #include "BlueprintEditorUtils.h"
 
-void SUMGEditorTreeItem::Construct(const FArguments& InArgs, TSharedPtr<FBlueprintEditor> InBlueprintEditor, UWidget* InItem)
+#include "Kismet2NameValidators.h"
+
+#define LOCTEXT_NAMESPACE "UMG"
+
+void SUMGEditorTreeItem::Construct(const FArguments& InArgs, const TSharedRef< STableViewBase >& InOwnerTableView, TSharedPtr<FWidgetBlueprintEditor> InBlueprintEditor, UWidget* InItem)
 {
 	BlueprintEditor = InBlueprintEditor;
 	Item = InItem;
 
-	ChildSlot
-	[
-		SNew(STextBlock)
-		.ToolTipText(this, &SUMGEditorTreeItem::GetItemTooltipText)
-		.Text(this, &SUMGEditorTreeItem::GetItemText)
-		.HighlightText(InArgs._HighlightText)
-	];
+	STableRow< UWidget* >::Construct(
+		STableRow< UWidget* >::FArguments()
+		.Padding(2.0f)
+		.Content()
+		[
+			SNew(SInlineEditableTextBlock)
+			.ToolTipText(this, &SUMGEditorTreeItem::GetItemTooltipText)
+			.Font(this, &SUMGEditorTreeItem::GetItemFont)
+			.Text(this, &SUMGEditorTreeItem::GetItemText)
+			.HighlightText(InArgs._HighlightText)
+			.OnVerifyTextChanged(this, &SUMGEditorTreeItem::OnVerifyNameTextChanged)
+			.OnTextCommitted(this, &SUMGEditorTreeItem::OnNameTextCommited)
+			.IsSelected(this, &SUMGEditorTreeItem::IsSelectedExclusively)
+		],
+		InOwnerTableView);
+}
+
+bool SUMGEditorTreeItem::OnVerifyNameTextChanged(const FText& InText, FText& OutErrorMessage)
+{
+	FString NewName = InText.ToString();
+
+	UWidgetBlueprint* Blueprint = BlueprintEditor.Pin()->GetWidgetBlueprintObj();
+	UWidget* ExistingWidget = Blueprint->WidgetTree->FindWidget(NewName);
+
+	//TODO UMG check variables in the blueprint also
+	FKismetNameValidator Validator(Blueprint);
+
+	// NewName should be already validated. But one must make sure that NewTemplateName is also unique.
+	const bool bUniqueNameForVariable = ( EValidatorResult::Ok == Validator.IsValid(NewName) );
+	
+	if ( ( ExistingWidget != NULL && ExistingWidget != Item ) || !bUniqueNameForVariable )
+	{
+		OutErrorMessage = LOCTEXT("NameConflict", "NameConflict");
+		return false;
+	}
+
+	return true;
+}
+
+void SUMGEditorTreeItem::OnNameTextCommited(const FText& InText, ETextCommit::Type CommitInfo)
+{
+	UWidgetBlueprint* Blueprint = BlueprintEditor.Pin()->GetWidgetBlueprintObj();
+	FWidgetBlueprintEditorUtils::RenameWidget(Blueprint, Item->GetFName(), FName(*InText.ToString()));
+}
+
+FSlateFontInfo SUMGEditorTreeItem::GetItemFont() const
+{
+	if ( !Item->IsGeneratedName() && Item->bIsVariable )
+	{
+		// TODO UMG Hacky move into style area
+		return FSlateFontInfo(FPaths::EngineContentDir() / TEXT("Slate/Fonts/Roboto-Bold.ttf"), 9);
+	}
+	else
+	{
+		// TODO UMG Hacky move into style area
+		return FSlateFontInfo(FPaths::EngineContentDir() / TEXT("Slate/Fonts/Roboto-Normal.ttf"), 9);
+	}
 }
 
 FText SUMGEditorTreeItem::GetItemText() const
 {
-	return FText::FromString(Item->GetName());
+	return FText::FromString(Item->GetLabel());
 }
 
 FString SUMGEditorTreeItem::GetItemTooltipText() const
@@ -94,3 +148,5 @@ FReply SUMGEditorTreeItem::OnDrop(const FGeometry& MyGeometry, const FDragDropEv
 
 	return FReply::Unhandled();
 }
+
+#undef LOCTEXT_NAMESPACE
diff --git a/Engine/Source/Editor/UMGEditor/Private/SUMGEditorTreeItem.h b/Engine/Source/Editor/UMGEditor/Private/SUMGEditorTreeItem.h
--- a/Engine/Source/Editor/UMGEditor/Private/SUMGEditorTreeItem.h
+++ b/Engine/Source/Editor/UMGEditor/Private/SUMGEditorTreeItem.h
@@ -8,7 +8,7 @@
 /**
  * An widget item in the hierarchy tree view.
  */
-class SUMGEditorTreeItem : public SCompoundWidget
+class SUMGEditorTreeItem : public STableRow< UWidget* >
 {
 public:
 	SLATE_BEGIN_ARGS( SUMGEditorTreeItem ){}
@@ -18,7 +18,7 @@ public:
 
 	SLATE_END_ARGS()
 
-	void Construct(const FArguments& InArgs, TSharedPtr<FBlueprintEditor> InBlueprintEditor, UWidget* InItem);
+	void Construct(const FArguments& InArgs, const TSharedRef< STableViewBase >& InOwnerTableView, TSharedPtr<FWidgetBlueprintEditor> InBlueprintEditor, UWidget* InItem);
 
 	virtual void OnDragEnter(const FGeometry& MyGeometry, const FDragDropEvent& DragDropEvent) override;
 	virtual void OnDragLeave(const FDragDropEvent& DragDropEvent) override;
@@ -26,9 +26,16 @@ public:
 	virtual FReply OnDrop(const FGeometry& MyGeometry, const FDragDropEvent& DragDropEvent) override;
 	
 private:
+	/** Called when text is being committed to check for validity */
+	bool OnVerifyNameTextChanged(const FText& InText, FText& OutErrorMessage);
+
+	/* Called when text is committed on the node */
+	void OnNameTextCommited(const FText& InText, ETextCommit::Type CommitInfo);
+
+	FSlateFontInfo GetItemFont() const;
 	FText GetItemText() const;
 	FString GetItemTooltipText() const;
 
-	TWeakPtr<class FBlueprintEditor> BlueprintEditor;
+	TWeakPtr<FWidgetBlueprintEditor> BlueprintEditor;
 	UWidget* Item;
 };
diff --git a/Engine/Source/Editor/UMGEditor/Private/UMGEditorPrivatePCH.h b/Engine/Source/Editor/UMGEditor/Private/UMGEditorPrivatePCH.h
--- a/Engine/Source/Editor/UMGEditor/Private/UMGEditorPrivatePCH.h
+++ b/Engine/Source/Editor/UMGEditor/Private/UMGEditorPrivatePCH.h
@@ -29,3 +29,5 @@
 
 #include "WidgetTemplate.h"
 #include "WidgetTemplateDragDropOp.h"
+
+#include "WidgetBlueprintEditorUtils.h"
diff --git a/Engine/Source/Editor/UMGEditor/Private/WidgetBlueprintEditorUtils.cpp b/Engine/Source/Editor/UMGEditor/Private/WidgetBlueprintEditorUtils.cpp
--- a/Engine/Source/Editor/UMGEditor/Private/WidgetBlueprintEditorUtils.cpp
+++ b/Engine/Source/Editor/UMGEditor/Private/WidgetBlueprintEditorUtils.cpp
@@ -7,16 +7,23 @@
 #include "BlueprintEditorUtils.h"
 #include "K2Node_Variable.h"
 
-bool FWidgetBlueprintEditorUtils::RenameWidget(UBlueprint* Blueprint, const FName& OldName, const FName& NewName)
+bool FWidgetBlueprintEditorUtils::RenameWidget(UWidgetBlueprint* Blueprint, const FName& OldName, const FName& NewName)
 {
 	check(Blueprint);
 
 	bool bRenamed = false;
 
 	TSharedPtr<INameValidatorInterface> NameValidator = MakeShareable(new FKismetNameValidator(Blueprint));
-	const FString NewTemplateName = UTimelineTemplate::TimelineVariableNameToTemplateName(NewName);
+
 	// NewName should be already validated. But one must make sure that NewTemplateName is also unique.
-	const bool bUniqueNameForTemplate = (EValidatorResult::Ok == NameValidator->IsValid(NewTemplateName));
+	const bool bUniqueNameForTemplate = ( EValidatorResult::Ok == NameValidator->IsValid(NewName) );
+
+	const FString NewNameStr = NewName.ToString();
+	const FString OldNameStr = OldName.ToString();
+
+	UWidget* Widget = Blueprint->WidgetTree->FindWidget(OldNameStr);
+	check(Widget);
+	Widget->Rename(*NewNameStr);
 
 	UTimelineTemplate* Template = Blueprint->FindTimelineTemplateByVariableName(NewName);
 	if ((Template == NULL) && bUniqueNameForTemplate)
@@ -27,9 +34,6 @@ bool FWidgetBlueprintEditorUtils::RenameWidget(UBlueprint* Blueprint, const FNam
 			Blueprint->Modify();
 			Template->Modify();
 
-			const FString NewNameStr = NewName.ToString();
-			const FString OldNameStr = OldName.ToString();
-
 			TArray<UK2Node_Variable*> TimelineVarNodes;
 			FBlueprintEditorUtils::GetAllNodesOfClass<UK2Node_Variable>(Blueprint, TimelineVarNodes);
 			for(int32 It = 0; It < TimelineVarNodes.Num(); It++)
@@ -59,12 +63,12 @@ bool FWidgetBlueprintEditorUtils::RenameWidget(UBlueprint* Blueprint, const FNam
 
 			Blueprint->Timelines.Remove(Template);
 			
-			UObject* ExistingObject = StaticFindObject(NULL, Template->GetOuter(), *NewTemplateName, true);
+			UObject* ExistingObject = StaticFindObject(NULL, Template->GetOuter(), *NewName.ToString(), true);
 			if (ExistingObject != Template && ExistingObject != NULL)
 			{
 				ExistingObject->Rename(*MakeUniqueObjectName(ExistingObject->GetOuter(), ExistingObject->GetClass(), ExistingObject->GetFName()).ToString());
 			}
-			Template->Rename(*NewTemplateName, Template->GetOuter(), (Blueprint->bIsRegeneratingOnLoad ? REN_ForceNoResetLoaders : REN_None));
+			Template->Rename(*NewName.ToString(), Template->GetOuter(), ( Blueprint->bIsRegeneratingOnLoad ? REN_ForceNoResetLoaders : REN_None ));
 			Blueprint->Timelines.Add(Template);
 
 			// Validate child blueprints and adjust variable names to avoid a potential name collision
diff --git a/Engine/Source/Editor/UMGEditor/Private/WidgetBlueprintEditorUtils.h b/Engine/Source/Editor/UMGEditor/Private/WidgetBlueprintEditorUtils.h
--- a/Engine/Source/Editor/UMGEditor/Private/WidgetBlueprintEditorUtils.h
+++ b/Engine/Source/Editor/UMGEditor/Private/WidgetBlueprintEditorUtils.h
@@ -8,5 +8,5 @@
 class FWidgetBlueprintEditorUtils
 {
 public:
-	static bool RenameWidget(UBlueprint* Blueprint, const FName& OldName, const FName& NewName);
+	static bool RenameWidget(UWidgetBlueprint* Blueprint, const FName& OldName, const FName& NewName);
 };
diff --git a/Engine/Source/Runtime/UMG/Private/Components/Border.cpp b/Engine/Source/Runtime/UMG/Private/Components/Border.cpp
--- a/Engine/Source/Runtime/UMG/Private/Components/Border.cpp
+++ b/Engine/Source/Runtime/UMG/Private/Components/Border.cpp
@@ -72,6 +72,33 @@ void UBorder::SetContent(UWidget* InContent)
 	}
 }
 
+void UBorder::SetBrushColor(FLinearColor Color)
+{
+	BrushColor = Color;
+	if ( MyBorder.IsValid() )
+	{
+		MyBorder->SetColorAndOpacity(Color);
+	}
+}
+
+void UBorder::SetForegroundColor(FLinearColor InForegroundColor)
+{
+	ForegroundColor = InForegroundColor;
+	if ( MyBorder.IsValid() )
+	{
+		MyBorder->SetForegroundColor(InForegroundColor);
+	}
+}
+
+void UBorder::SetContentPadding(FMargin InContentPadding)
+{
+	ContentPadding = InContentPadding;
+	if ( MyBorder.IsValid() )
+	{
+		MyBorder->SetPadding(InContentPadding);
+	}
+}
+
 const FSlateBrush* UBorder::GetBorderBrush() const
 {
 	if ( Brush == NULL )
diff --git a/Engine/Source/Runtime/UMG/Private/Components/Button.cpp b/Engine/Source/Runtime/UMG/Private/Components/Button.cpp
--- a/Engine/Source/Runtime/UMG/Private/Components/Button.cpp
+++ b/Engine/Source/Runtime/UMG/Private/Components/Button.cpp
@@ -87,22 +87,38 @@ void UButton::SetContent(UWidget* InContent)
 
 void UButton::SetColorAndOpacity(FLinearColor Color)
 {
-	MyButton->SetColorAndOpacity(Color);
+	ColorAndOpacity = Color;
+	if ( MyButton.IsValid() )
+	{
+		MyButton->SetColorAndOpacity(Color);
+	}
 }
 
 void UButton::SetBackgroundColor(FLinearColor Color)
 {
-	MyButton->SetBorderBackgroundColor(Color);
+	BackgroundColor = Color;
+	if ( MyButton.IsValid() )
+	{
+		MyButton->SetBorderBackgroundColor(Color);
+	}
 }
 
 void UButton::SetForegroundColor(FLinearColor InForegroundColor)
 {
-	MyButton->SetForegroundColor(InForegroundColor);
+	ForegroundColor = InForegroundColor;
+	if ( MyButton.IsValid() )
+	{
+		MyButton->SetForegroundColor(InForegroundColor);
+	}
 }
 
 void UButton::SetContentPadding(FMargin InContentPadding)
 {
-	MyButton->SetContentPadding(InContentPadding);
+	ContentPadding = InContentPadding;
+	if ( MyButton.IsValid() )
+	{
+		MyButton->SetContentPadding(InContentPadding);
+	}
 }
 
 FReply UButton::HandleOnClicked()
diff --git a/Engine/Source/Runtime/UMG/Private/Components/ContentWidget.cpp b/Engine/Source/Runtime/UMG/Private/Components/ContentWidget.cpp
--- a/Engine/Source/Runtime/UMG/Private/Components/ContentWidget.cpp
+++ b/Engine/Source/Runtime/UMG/Private/Components/ContentWidget.cpp
@@ -15,8 +15,19 @@
 UContentWidget::UContentWidget(const FPostConstructInitializeProperties& PCIP)
 	: Super(PCIP)
 {
-	ContentSlot = ConstructObject<UPanelSlot>(UPanelSlot::StaticClass(), this);
+}
+
+void UContentWidget::PostInitProperties()
+{
+	Super::PostInitProperties();
+
+	if ( ContentSlot == NULL )
+	{
+		ContentSlot = ConstructObject<UPanelSlot>(UPanelSlot::StaticClass(), this);
+	}
+
 	ContentSlot->SetFlags(RF_Transactional);
+	ContentSlot->SetFlags(RF_DefaultSubObject);
 	ContentSlot->Parent = this;
 }
 
diff --git a/Engine/Source/Runtime/UMG/Private/Components/TextBlock.cpp b/Engine/Source/Runtime/UMG/Private/Components/TextBlock.cpp
--- a/Engine/Source/Runtime/UMG/Private/Components/TextBlock.cpp
+++ b/Engine/Source/Runtime/UMG/Private/Components/TextBlock.cpp
@@ -26,6 +26,19 @@ UTextBlock::UTextBlock(const FPostConstructInitializeProperties& PCIP)
 TSharedRef<SWidget> UTextBlock::RebuildWidget()
 {
 	MyTextBlock = SNew(STextBlock);
+
+#if WITH_EDITOR
+	MyEditorTextBlock = SNew(SInlineEditableTextBlock)
+		.OnTextCommitted(BIND_UOBJECT_DELEGATE(FOnTextCommitted, HandleTextCommitted))
+		//.IsSelected(InArgs._IsSelected)
+		;
+
+	//if ( IsDesignTime() )
+	//{
+	//	return MyEditorTextBlock.ToSharedRef();
+	//}
+#endif
+
 	return MyTextBlock.ToSharedRef();
 }
 
@@ -52,8 +65,33 @@ void UTextBlock::SyncronizeProperties()
 	MyTextBlock->SetTextStyle(StylePtr);
 	MyTextBlock->SetShadowOffset(ShadowOffset);
 	MyTextBlock->SetShadowColorAndOpacity(ShadowColorAndOpacityBinding);
+
+#if WITH_EDITOR
+	MyEditorTextBlock->SetText(TextBinding);
+	//MyEditorTextBlock->SetFont(FSlateFontInfo(FontPath, Font.Size));
+	//MyEditorTextBlock->SetColorAndOpacity(ColorAndOpacityBinding);
+	//MyEditorTextBlock->SetTextStyle(StylePtr);
+	//MyEditorTextBlock->SetShadowOffset(ShadowOffset);
+	//MyEditorTextBlock->SetShadowColorAndOpacity(ShadowColorAndOpacityBinding);
+#endif
+}
+
+#if WITH_EDITOR
+
+void UTextBlock::HandleTextCommitted(const FText& InText, ETextCommit::Type CommitteType)
+{
+	//TODO UMG How will this migrate to the template?  Seems to me we need the previews to have access to their templates!
+	//TODO UMG How will the user click the editable area?  There is an overlay blocking input so that other widgets don't get them.
+	//     Need a way to recognize one particular widget and forward things to them!
+}
+
+void UTextBlock::OnDesignerDoubleClicked()
+{
+	MyEditorTextBlock->EnterEditingMode();
 }
 
+#endif
+
 /////////////////////////////////////////////////////
 
 #undef LOCTEXT_NAMESPACE
diff --git a/Engine/Source/Runtime/UMG/Private/Components/Widget.cpp b/Engine/Source/Runtime/UMG/Private/Components/Widget.cpp
--- a/Engine/Source/Runtime/UMG/Private/Components/Widget.cpp
+++ b/Engine/Source/Runtime/UMG/Private/Components/Widget.cpp
@@ -120,6 +120,31 @@ TSharedRef<SWidget> UWidget::BuildDesignTimeWidget(TSharedRef<SWidget> WrapWidge
 
 #if WITH_EDITOR
 
+bool UWidget::IsGeneratedName() const
+{
+	FString Name = GetName();
+	FString BaseName = GetClass()->GetName() + "_";
+
+	if ( Name.StartsWith(BaseName) )
+	{
+		return true;
+	}
+
+	return false;
+}
+
+FString UWidget::GetLabel() const
+{
+	if ( IsGeneratedName() && !bIsVariable )
+	{
+		return "[" + GetClass()->GetName() + "]";
+	}
+	else
+	{
+		return GetName();
+	}
+}
+
 TSharedRef<SWidget> UWidget::GetToolboxPreviewWidget() const
 {
 	return SNew(SImage);
diff --git a/Engine/Source/Runtime/UMG/Private/WidgetBlueprint.cpp b/Engine/Source/Runtime/UMG/Private/WidgetBlueprint.cpp
--- a/Engine/Source/Runtime/UMG/Private/WidgetBlueprint.cpp
+++ b/Engine/Source/Runtime/UMG/Private/WidgetBlueprint.cpp
@@ -37,6 +37,11 @@ void UWidgetBlueprint::PostLoad()
 #endif
 }
 
+UClass* UWidgetBlueprint::RegenerateClass(UClass* ClassToRegenerate, UObject* PreviousCDO, TArray<UObject*>& ObjLoaded)
+{
+	return Super::RegenerateClass(ClassToRegenerate, PreviousCDO, ObjLoaded);
+}
+
 #if WITH_EDITOR
 UClass* UWidgetBlueprint::GetBlueprintClass() const
 {
diff --git a/Engine/Source/Runtime/UMG/Private/WidgetBlueprintGeneratedClass.cpp b/Engine/Source/Runtime/UMG/Private/WidgetBlueprintGeneratedClass.cpp
--- a/Engine/Source/Runtime/UMG/Private/WidgetBlueprintGeneratedClass.cpp
+++ b/Engine/Source/Runtime/UMG/Private/WidgetBlueprintGeneratedClass.cpp
@@ -10,7 +10,19 @@
 UWidgetBlueprintGeneratedClass::UWidgetBlueprintGeneratedClass(const FPostConstructInitializeProperties& PCIP)
 	: Super(PCIP)
 {
-	WidgetTree = ConstructObject<UWidgetTree>(UWidgetTree::StaticClass(), this);
+}
+
+void UWidgetBlueprintGeneratedClass::PostInitProperties()
+{
+	Super::PostInitProperties();
+
+	// Create a widget tree if one doesn't already exist.
+	if ( WidgetTree == NULL )
+	{
+		WidgetTree = ConstructObject<UWidgetTree>(UWidgetTree::StaticClass(), this);
+	}
+
+	WidgetTree->SetFlags(RF_DefaultSubObject);
 	WidgetTree->SetFlags(RF_Transactional);
 }
 
diff --git a/Engine/Source/Runtime/UMG/Private/WidgetTree.cpp b/Engine/Source/Runtime/UMG/Private/WidgetTree.cpp
--- a/Engine/Source/Runtime/UMG/Private/WidgetTree.cpp
+++ b/Engine/Source/Runtime/UMG/Private/WidgetTree.cpp
@@ -10,15 +10,7 @@ UWidgetTree::UWidgetTree(const FPostConstructInitializeProperties& PCIP)
 {
 }
 
-void UWidgetTree::RenameWidget(UWidget* Widget, FString& NewName)
-{
-	Widget->Modify();
-	Widget->Rename(*NewName);
-
-	// TODO UMG Update nodes in the blueprint!
-}
-
-UWidget* UWidgetTree::FindWidget(FString& Name) const
+UWidget* UWidgetTree::FindWidget(const FString& Name) const
 {
 	FString ExistingName;
 	for ( UWidget* Widget : WidgetTemplates )
diff --git a/Engine/Source/Runtime/UMG/Public/Blueprint/WidgetBlueprint.h b/Engine/Source/Runtime/UMG/Public/Blueprint/WidgetBlueprint.h
--- a/Engine/Source/Runtime/UMG/Public/Blueprint/WidgetBlueprint.h
+++ b/Engine/Source/Runtime/UMG/Public/Blueprint/WidgetBlueprint.h
@@ -56,6 +56,8 @@ public:
 	 
 	virtual void PostLoad() override;
 	virtual void PostLoadSubobjects(FObjectInstancingGraph* OuterInstanceGraph) override;
+
+	virtual UClass* RegenerateClass(UClass* ClassToRegenerate, UObject* PreviousCDO, TArray<UObject*>& ObjLoaded) override;
 	
 #if WITH_EDITOR
 	// UBlueprint interface
diff --git a/Engine/Source/Runtime/UMG/Public/Blueprint/WidgetBlueprintGeneratedClass.h b/Engine/Source/Runtime/UMG/Public/Blueprint/WidgetBlueprintGeneratedClass.h
--- a/Engine/Source/Runtime/UMG/Public/Blueprint/WidgetBlueprintGeneratedClass.h
+++ b/Engine/Source/Runtime/UMG/Public/Blueprint/WidgetBlueprintGeneratedClass.h
@@ -45,6 +45,8 @@ class UMG_API UWidgetBlueprintGeneratedClass : public UBlueprintGeneratedClass
 	/** This is transient data calculated at link time. */
 	TArray<UStructProperty*> WidgetNodeProperties;
 
+	virtual void PostInitProperties() override;
+
 	virtual void Link(FArchive& Ar, bool bRelinkExistingProperties) override;
 
 	void InitializeWidget(class UUserWidget* Actor) const;
diff --git a/Engine/Source/Runtime/UMG/Public/Blueprint/WidgetTree.h b/Engine/Source/Runtime/UMG/Public/Blueprint/WidgetTree.h
--- a/Engine/Source/Runtime/UMG/Public/Blueprint/WidgetTree.h
+++ b/Engine/Source/Runtime/UMG/Public/Blueprint/WidgetTree.h
@@ -17,9 +17,7 @@ public:
 
 public:
 
-	void RenameWidget(UWidget* Widget, FString& NewName);
-
-	UWidget* FindWidget(FString& Name) const;
+	UWidget* FindWidget(const FString& Name) const;
 
 	bool RemoveWidget(UWidget* Widget);
 
@@ -29,8 +27,11 @@ public:
 	template< class T >
 	T* ConstructWidget(TSubclassOf<UWidget> WidgetType)
 	{
+		// TODO UMG Editor only?
 		Modify();
 
+		// TODO UMG Don't have the widget tree responsible for construction and adding to the tree.
+
 		UWidget* Widget = (UWidget*)ConstructObject<UWidget>(WidgetType, this);
 		Widget->SetFlags(RF_Transactional);
 		WidgetTemplates.Add(Widget);
diff --git a/Engine/Source/Runtime/UMG/Public/Components/Border.h b/Engine/Source/Runtime/UMG/Public/Components/Border.h
--- a/Engine/Source/Runtime/UMG/Public/Components/Border.h
+++ b/Engine/Source/Runtime/UMG/Public/Components/Border.h
@@ -75,6 +75,18 @@ class UMG_API UBorder : public UContentWidget
 	UPROPERTY(EditDefaultsOnly, Category=Events)
 	FOnPointerEvent OnMouseDoubleClickEvent;
 
+	/**  */
+	UFUNCTION(BlueprintCallable, Category="Appearance")
+	void SetBrushColor(FLinearColor InColorAndOpacity);
+
+	/**  */
+	UFUNCTION(BlueprintCallable, Category="Appearance")
+	void SetForegroundColor(FLinearColor InForegroundColor);
+
+	/**  */
+	UFUNCTION(BlueprintCallable, Category="Appearance")
+	void SetContentPadding(FMargin InContentPadding);
+
 	// UContentWidget interface
 	virtual void SetContent(UWidget* Content) override;
 	// End UContentWidget interface
diff --git a/Engine/Source/Runtime/UMG/Public/Components/ContentWidget.h b/Engine/Source/Runtime/UMG/Public/Components/ContentWidget.h
--- a/Engine/Source/Runtime/UMG/Public/Components/ContentWidget.h
+++ b/Engine/Source/Runtime/UMG/Public/Components/ContentWidget.h
@@ -23,6 +23,10 @@ class UMG_API UContentWidget : public UPanelWidget
 	virtual void ReplaceChildAt(int32 Index, UWidget* Child) override;
 	// End UPanelWidget
 
+	// UObject
+	virtual void PostInitProperties() override;
+	// End UObject
+
 protected:
 
 	UPROPERTY()
diff --git a/Engine/Source/Runtime/UMG/Public/Components/TextBlock.h b/Engine/Source/Runtime/UMG/Public/Components/TextBlock.h
--- a/Engine/Source/Runtime/UMG/Public/Components/TextBlock.h
+++ b/Engine/Source/Runtime/UMG/Public/Components/TextBlock.h
@@ -62,11 +62,22 @@ public:
 	void SyncronizeProperties() override;
 	// End of UWidget interface
 
+#if WITH_EDITOR
+	void HandleTextCommitted(const FText& InText, ETextCommit::Type CommitteType);
+
+	virtual void OnDesignerDoubleClicked() override;
+#endif
+
 protected:
 	// UWidget interface
 	virtual TSharedRef<SWidget> RebuildWidget() override;
 	// End of UWidget interface
 
 protected:
+
 	TSharedPtr<STextBlock> MyTextBlock;
+
+#if WITH_EDITOR
+	TSharedPtr<SInlineEditableTextBlock> MyEditorTextBlock;
+#endif
 };
diff --git a/Engine/Source/Runtime/UMG/Public/Components/Widget.h b/Engine/Source/Runtime/UMG/Public/Components/Widget.h
--- a/Engine/Source/Runtime/UMG/Public/Components/Widget.h
+++ b/Engine/Source/Runtime/UMG/Public/Components/Widget.h
@@ -56,7 +56,7 @@ public:
 	/**
 	 * The parent slot of the UWidget.  Allows us to easily inline edit the layout controlling this widget.
 	 */
-	UPROPERTY(EditInline, EditDefaultsOnly, BlueprintReadWrite, Category=Layout, meta=( ShowOnlyInnerProperties ))
+	UPROPERTY(EditInline, EditAnywhere, BlueprintReadOnly, Category=Layout, meta=( ShowOnlyInnerProperties ))
 	UPanelSlot* Slot;
 
 	/** Sets whether this widget can be modified interactively by the user */
@@ -143,6 +143,17 @@ public:
 	virtual bool Modify(bool bAlwaysMarkDirty = true);
 	
 #if WITH_EDITOR
+	/** Is the label generated or provided by the user? */
+	bool IsGeneratedName() const;
+
+	/** Gets the label to display to the user for this widget. */
+	FString GetLabel() const;
+
+	virtual void OnDesignerSelected() { }
+	virtual void OnDesignerDeselected() { }
+	virtual void OnDesignerClicked() { }
+	virtual void OnDesignerDoubleClicked() { }
+
 	/** Gets a widget representing the tiny preview of the toolbox */
 	virtual TSharedRef<SWidget> GetToolboxPreviewWidget() const;
 	
