Fix crash opening two widget blueprints at once

diff --git a/Engine/Source/Editor/UMGEditor/Private/WidgetBlueprintEditor.cpp b/Engine/Source/Editor/UMGEditor/Private/WidgetBlueprintEditor.cpp
--- a/Engine/Source/Editor/UMGEditor/Private/WidgetBlueprintEditor.cpp
+++ b/Engine/Source/Editor/UMGEditor/Private/WidgetBlueprintEditor.cpp
@@ -334,7 +334,8 @@ UMovieScene* FWidgetBlueprintEditor::GetDefaultMovieScene()
 {
 	if( !DefaultMovieScene )
 	{
-		DefaultMovieScene = ConstructObject<UMovieScene>( UMovieScene::StaticClass(), GetTransientPackage(), "DefaultAnimationData" );
+		FName ObjectName = MakeUniqueObjectName( GetTransientPackage(), UMovieScene::StaticClass(), "DefaultAnimationData" );
+		DefaultMovieScene = ConstructObject<UMovieScene>( UMovieScene::StaticClass(), GetTransientPackage(), ObjectName );
 	}
 
 	return DefaultMovieScene;
