non unity fix

diff --git a/Engine/Source/Editor/UMGEditor/Private/SUMGEditorTreeItem.cpp b/Engine/Source/Editor/UMGEditor/Private/SUMGEditorTreeItem.cpp
--- a/Engine/Source/Editor/UMGEditor/Private/SUMGEditorTreeItem.cpp
+++ b/Engine/Source/Editor/UMGEditor/Private/SUMGEditorTreeItem.cpp
@@ -15,6 +15,8 @@
 
 #include "Kismet2NameValidators.h"
 
+#include "WidgetBlueprintEditor.h"
+
 #define LOCTEXT_NAMESPACE "UMG"
 
 void SUMGEditorTreeItem::Construct(const FArguments& InArgs, const TSharedRef< STableViewBase >& InOwnerTableView, TSharedPtr<FWidgetBlueprintEditor> InBlueprintEditor, UWidget* InItem)
