UMG - Renaming also renames the preview widget now.  General cleanup.

diff --git a/Engine/Source/Editor/UMGEditor/Private/SUMGEditorTree.cpp b/Engine/Source/Editor/UMGEditor/Private/SUMGEditorTree.cpp
--- a/Engine/Source/Editor/UMGEditor/Private/SUMGEditorTree.cpp
+++ b/Engine/Source/Editor/UMGEditor/Private/SUMGEditorTree.cpp
@@ -28,8 +28,6 @@ void SUMGEditorTree::Construct(const FArguments& InArgs, TSharedPtr<FWidgetBluep
 	UWidgetBlueprint* Blueprint = GetBlueprint();
 	Blueprint->OnChanged().AddSP(this, &SUMGEditorTree::OnBlueprintChanged);
 
-	FCoreDelegates::OnObjectPropertyChanged.AddRaw(this, &SUMGEditorTree::OnObjectPropertyChanged);
-
 	SAssignNew(WidgetTreeView, STreeView< UWidget* >)
 	.ItemHeight(20.0f)
 	.SelectionMode(ESelectionMode::Single)
@@ -79,8 +77,6 @@ SUMGEditorTree::~SUMGEditorTree()
 		Blueprint->OnChanged().RemoveAll(this);
 	}
 
-	FCoreDelegates::OnObjectPropertyChanged.RemoveAll(this);
-
 	if ( BlueprintEditor.IsValid() )
 	{
 		BlueprintEditor.Pin()->OnSelectedWidgetsChanged.RemoveAll(this);
@@ -171,14 +167,6 @@ void SUMGEditorTree::OnBlueprintChanged(UBlueprint* InBlueprint)
 	}
 }
 
-void SUMGEditorTree::OnObjectPropertyChanged(UObject* ObjectBeingModified, FPropertyChangedEvent& PropertyChangedEvent)
-{
-	if ( !ensure(ObjectBeingModified) )
-	{
-		return;
-	}
-}
-
 void SUMGEditorTree::ShowDetailsForObjects(TArray<UWidget*> TemplateWidgets)
 {
 	TSet<FWidgetReference> SelectedWidgets;
diff --git a/Engine/Source/Editor/UMGEditor/Private/SUMGEditorTree.h b/Engine/Source/Editor/UMGEditor/Private/SUMGEditorTree.h
--- a/Engine/Source/Editor/UMGEditor/Private/SUMGEditorTree.h
+++ b/Engine/Source/Editor/UMGEditor/Private/SUMGEditorTree.h
@@ -24,7 +24,9 @@ public:
 
 	virtual ~SUMGEditorTree();
 
+	// Begin SWidget
 	virtual void Tick(const FGeometry& AllottedGeometry, const double InCurrentTime, const float InDeltaTime);
+	// End SWidget
 
 private:
 	void BuildWrapWithMenu(FMenuBuilder& Menu);
@@ -43,8 +45,6 @@ private:
 	/** Called when the blueprint is structurally changed. */
 	void OnBlueprintChanged(UBlueprint* InBlueprint);
 
-	void OnObjectPropertyChanged(UObject* ObjectBeingModified, FPropertyChangedEvent& PropertyChangedEvent);
-
 	/** Called when the selected widget has changed.  The treeview then needs to match the new selection. */
 	void OnEditorSelectionChanged();
 
@@ -77,19 +77,30 @@ private:
 
 private:
 
+	/** Cached pointer to the blueprint editor that owns this tree. */
 	TWeakPtr<class FWidgetBlueprintEditor> BlueprintEditor;
 
+	/** The tree item source.  We only ever have one root widget, but the tree expects an array. */
 	TArray< UWidget* > RootWidgets;
 
+	/** The widget hierarchy slate treeview widget */
 	TSharedPtr< STreeView< UWidget* > > WidgetTreeView;
 
+	/** The filter used by the search box */
 	TSharedPtr<WidgetTextFilter> SearchBoxWidgetFilter;
 
+	/** The set of widgets currently passing the filter.  Used to quickly filter children when the tree is expanded lazily. */
 	TSet< UWidget* > WidgetsPassingFilter;
 
+	/**
+	  * A cache of the expanded items.  We store this off before a user filters the list so that we can 
+	  * restore the expansion state after the filter is cleared.
+	  */
 	TSet< UWidget* > CachedExpandedWidgets;
 
+	/** Caches a value to know if a filter is currently active. */
 	bool bIsFilterActive;
 
+	/** Has a full refresh of the tree been requested?  This happens when the user is filtering the tree */
 	bool bRefreshRequested;
 };
diff --git a/Engine/Source/Editor/UMGEditor/Private/SUMGEditorTreeItem.cpp b/Engine/Source/Editor/UMGEditor/Private/SUMGEditorTreeItem.cpp
--- a/Engine/Source/Editor/UMGEditor/Private/SUMGEditorTreeItem.cpp
+++ b/Engine/Source/Editor/UMGEditor/Private/SUMGEditorTreeItem.cpp
@@ -45,7 +45,6 @@ void SUMGEditorTreeItem::Construct(const FArguments& InArgs, const TSharedRef< S
 			.VAlign(VAlign_Center)
 			[
 				SNew(SInlineEditableTextBlock)
-				.ToolTipText(this, &SUMGEditorTreeItem::GetItemTooltipText)
 				.Font(this, &SUMGEditorTreeItem::GetItemFont)
 				.Text(this, &SUMGEditorTreeItem::GetItemText)
 				.HighlightText(InArgs._HighlightText)
@@ -64,10 +63,8 @@ bool SUMGEditorTreeItem::OnVerifyNameTextChanged(const FText& InText, FText& Out
 	UWidgetBlueprint* Blueprint = BlueprintEditor.Pin()->GetWidgetBlueprintObj();
 	UWidget* ExistingWidget = Blueprint->WidgetTree->FindWidget(NewName);
 
-	//TODO UMG check variables in the blueprint also
 	FKismetNameValidator Validator(Blueprint);
 
-	// NewName should be already validated. But one must make sure that NewTemplateName is also unique.
 	const bool bUniqueNameForVariable = ( EValidatorResult::Ok == Validator.IsValid(NewName) );
 	
 	if ( ( ExistingWidget != NULL && ExistingWidget != Item ) || !bUniqueNameForVariable )
@@ -81,8 +78,7 @@ bool SUMGEditorTreeItem::OnVerifyNameTextChanged(const FText& InText, FText& Out
 
 void SUMGEditorTreeItem::OnNameTextCommited(const FText& InText, ETextCommit::Type CommitInfo)
 {
-	UWidgetBlueprint* Blueprint = BlueprintEditor.Pin()->GetWidgetBlueprintObj();
-	FWidgetBlueprintEditorUtils::RenameWidget(Blueprint, Item->GetFName(), FName(*InText.ToString()));
+	FWidgetBlueprintEditorUtils::RenameWidget(BlueprintEditor.Pin().ToSharedRef(), Item->GetFName(), FName(*InText.ToString()));
 }
 
 FSlateFontInfo SUMGEditorTreeItem::GetItemFont() const
@@ -104,10 +100,7 @@ FText SUMGEditorTreeItem::GetItemText() const
 	return FText::FromString(Item->GetLabel());
 }
 
-FString SUMGEditorTreeItem::GetItemTooltipText() const
-{
-	return Item->GetDetailedInfo();
-}
+//@TODO UMG Allow items in the tree to be dragged, and reordered, and reparented.
 
 void SUMGEditorTreeItem::OnDragEnter(const FGeometry& MyGeometry, const FDragDropEvent& DragDropEvent)
 {
diff --git a/Engine/Source/Editor/UMGEditor/Private/SUMGEditorTreeItem.h b/Engine/Source/Editor/UMGEditor/Private/SUMGEditorTreeItem.h
--- a/Engine/Source/Editor/UMGEditor/Private/SUMGEditorTreeItem.h
+++ b/Engine/Source/Editor/UMGEditor/Private/SUMGEditorTreeItem.h
@@ -20,10 +20,12 @@ public:
 
 	void Construct(const FArguments& InArgs, const TSharedRef< STableViewBase >& InOwnerTableView, TSharedPtr<FWidgetBlueprintEditor> InBlueprintEditor, UWidget* InItem);
 
+	// Begin SWidget
 	virtual void OnDragEnter(const FGeometry& MyGeometry, const FDragDropEvent& DragDropEvent) override;
 	virtual void OnDragLeave(const FDragDropEvent& DragDropEvent) override;
 	virtual FReply OnDragOver(const FGeometry& MyGeometry, const FDragDropEvent& DragDropEvent) override;
 	virtual FReply OnDrop(const FGeometry& MyGeometry, const FDragDropEvent& DragDropEvent) override;
+	// End SWidget
 	
 private:
 	/** Called when text is being committed to check for validity */
@@ -32,10 +34,15 @@ private:
 	/* Called when text is committed on the node */
 	void OnNameTextCommited(const FText& InText, ETextCommit::Type CommitInfo);
 
+	/* Gets the font to use for the text item, bold for customized named items */
 	FSlateFontInfo GetItemFont() const;
+
+	/* @returns the widget name to use for the tree item */
 	FText GetItemText() const;
-	FString GetItemTooltipText() const;
 
+	/* The blueprint editor this tree view item's tree is associated with. */
 	TWeakPtr<FWidgetBlueprintEditor> BlueprintEditor;
+
+	/* The template widget that this tree item represents */
 	UWidget* Item;
 };
diff --git a/Engine/Source/Editor/UMGEditor/Private/WidgetBlueprintEditorUtils.cpp b/Engine/Source/Editor/UMGEditor/Private/WidgetBlueprintEditorUtils.cpp
--- a/Engine/Source/Editor/UMGEditor/Private/WidgetBlueprintEditorUtils.cpp
+++ b/Engine/Source/Editor/UMGEditor/Private/WidgetBlueprintEditorUtils.cpp
@@ -7,8 +7,9 @@
 #include "BlueprintEditorUtils.h"
 #include "K2Node_Variable.h"
 
-bool FWidgetBlueprintEditorUtils::RenameWidget(UWidgetBlueprint* Blueprint, const FName& OldName, const FName& NewName)
+bool FWidgetBlueprintEditorUtils::RenameWidget(TSharedRef<FWidgetBlueprintEditor> BlueprintEditor, const FName& OldName, const FName& NewName)
 {
+	UWidgetBlueprint* Blueprint = BlueprintEditor->GetWidgetBlueprintObj();
 	check(Blueprint);
 
 	bool bRenamed = false;
@@ -23,8 +24,19 @@ bool FWidgetBlueprintEditorUtils::RenameWidget(UWidgetBlueprint* Blueprint, cons
 
 	UWidget* Widget = Blueprint->WidgetTree->FindWidget(OldNameStr);
 	check(Widget);
+
+	// Rename Template
 	Widget->Rename(*NewNameStr);
 
+	// Rename Preview
+	UWidget* WidgetPreview = FWidgetReference::FromTemplate(BlueprintEditor, Widget).GetPreview();
+	if ( WidgetPreview )
+	{
+		WidgetPreview->Rename(*NewNameStr);
+	}
+
+	// TODO UMG RENAME VARIABLES IN THE GRAPH!
+
 	UTimelineTemplate* Template = Blueprint->FindTimelineTemplateByVariableName(NewName);
 	if ((Template == NULL) && bUniqueNameForTemplate)
 	{
@@ -79,5 +91,10 @@ bool FWidgetBlueprintEditorUtils::RenameWidget(UWidgetBlueprint* Blueprint, cons
 			bRenamed = true;
 		}
 	}
+
+	// Refresh references and flush editors
+	FBlueprintEditorUtils::MarkBlueprintAsStructurallyModified(Blueprint);
+	bRenamed = true;
+
 	return bRenamed;
 }
\ No newline at end of file
diff --git a/Engine/Source/Editor/UMGEditor/Private/WidgetBlueprintEditorUtils.h b/Engine/Source/Editor/UMGEditor/Private/WidgetBlueprintEditorUtils.h
--- a/Engine/Source/Editor/UMGEditor/Private/WidgetBlueprintEditorUtils.h
+++ b/Engine/Source/Editor/UMGEditor/Private/WidgetBlueprintEditorUtils.h
@@ -8,5 +8,5 @@
 class FWidgetBlueprintEditorUtils
 {
 public:
-	static bool RenameWidget(UWidgetBlueprint* Blueprint, const FName& OldName, const FName& NewName);
+	static bool RenameWidget(TSharedRef<FWidgetBlueprintEditor> BlueprintEditor, const FName& OldName, const FName& NewName);
 };
