Fix non-unity

diff --git a/Engine/Source/Editor/UMGEditor/Private/WidgetBlueprintEditorUtils.cpp b/Engine/Source/Editor/UMGEditor/Private/WidgetBlueprintEditorUtils.cpp
--- a/Engine/Source/Editor/UMGEditor/Private/WidgetBlueprintEditorUtils.cpp
+++ b/Engine/Source/Editor/UMGEditor/Private/WidgetBlueprintEditorUtils.cpp
@@ -3,6 +3,7 @@
 #include "UMGEditorPrivatePCH.h"
 
 #include "WidgetBlueprintEditorUtils.h"
+#include "WidgetBlueprintEditor.h"
 #include "Kismet2NameValidators.h"
 #include "BlueprintEditorUtils.h"
 #include "K2Node_Variable.h"
