UMG - Added margin track and changed a few things to support new track types being added to sequencer without needing to modify sequencer internals.

diff --git a/Engine/Source/Developer/MovieSceneTools/Private/PropertyTrackEditor.cpp b/Engine/Source/Developer/MovieSceneTools/Private/PropertyTrackEditor.cpp
--- a/Engine/Source/Developer/MovieSceneTools/Private/PropertyTrackEditor.cpp
+++ b/Engine/Source/Developer/MovieSceneTools/Private/PropertyTrackEditor.cpp
@@ -265,10 +265,13 @@ FPropertyTrackEditor::FPropertyTrackEditor( TSharedRef<ISequencer> InSequencer )
 {
 	// Get the object change listener for the sequencer and register a delegates for when properties change that we care about
 	ISequencerObjectChangeListener& ObjectChangeListener = InSequencer->GetObjectChangeListener();
-	ObjectChangeListener.GetOnAnimatablePropertyChanged( UFloatProperty::StaticClass() ).AddRaw( this, &FPropertyTrackEditor::OnAnimatedPropertyChanged<float, UMovieSceneFloatTrack> );
-	ObjectChangeListener.GetOnAnimatablePropertyChanged( UBoolProperty::StaticClass() ).AddRaw( this, &FPropertyTrackEditor::OnAnimatedPropertyChanged<bool, UMovieSceneBoolTrack> );
-	ObjectChangeListener.GetOnAnimatablePropertyChanged( UStructProperty::StaticClass() ).AddRaw( this, &FPropertyTrackEditor::OnAnimatedVectorPropertyChanged );
-	ObjectChangeListener.GetOnAnimatablePropertyChanged( UStructProperty::StaticClass() ).AddRaw( this, &FPropertyTrackEditor::OnAnimatedColorPropertyChanged );
+	ObjectChangeListener.GetOnAnimatablePropertyChanged( NAME_FloatProperty ).AddRaw( this, &FPropertyTrackEditor::OnAnimatedPropertyChanged<float, UMovieSceneFloatTrack> );
+	ObjectChangeListener.GetOnAnimatablePropertyChanged( NAME_BoolProperty ).AddRaw( this, &FPropertyTrackEditor::OnAnimatedPropertyChanged<bool, UMovieSceneBoolTrack> );
+	ObjectChangeListener.GetOnAnimatablePropertyChanged( NAME_Vector ).AddRaw( this, &FPropertyTrackEditor::OnAnimatedVectorPropertyChanged );
+	ObjectChangeListener.GetOnAnimatablePropertyChanged( NAME_Vector4 ).AddRaw( this, &FPropertyTrackEditor::OnAnimatedVectorPropertyChanged );
+	ObjectChangeListener.GetOnAnimatablePropertyChanged( NAME_Vector2D ).AddRaw( this, &FPropertyTrackEditor::OnAnimatedVectorPropertyChanged );
+	ObjectChangeListener.GetOnAnimatablePropertyChanged( NAME_Color ).AddRaw( this, &FPropertyTrackEditor::OnAnimatedColorPropertyChanged );
+	ObjectChangeListener.GetOnAnimatablePropertyChanged( NAME_LinearColor ).AddRaw( this, &FPropertyTrackEditor::OnAnimatedColorPropertyChanged );
 }
 
 FPropertyTrackEditor::~FPropertyTrackEditor()
@@ -277,9 +280,13 @@ FPropertyTrackEditor::~FPropertyTrackEditor()
 	if( Sequencer.IsValid() )
 	{
 		ISequencerObjectChangeListener& ObjectChangeListener = Sequencer->GetObjectChangeListener();
-		ObjectChangeListener.GetOnAnimatablePropertyChanged( UFloatProperty::StaticClass() ).RemoveAll( this );
-		ObjectChangeListener.GetOnAnimatablePropertyChanged( UBoolProperty::StaticClass() ).RemoveAll( this );
-		ObjectChangeListener.GetOnAnimatablePropertyChanged( UStructProperty::StaticClass() ).RemoveAll( this );
+		ObjectChangeListener.GetOnAnimatablePropertyChanged(NAME_FloatProperty).RemoveAll(this);
+		ObjectChangeListener.GetOnAnimatablePropertyChanged(NAME_BoolProperty).RemoveAll(this);
+		ObjectChangeListener.GetOnAnimatablePropertyChanged(NAME_Vector).RemoveAll(this);
+		ObjectChangeListener.GetOnAnimatablePropertyChanged(NAME_Vector4).RemoveAll(this);
+		ObjectChangeListener.GetOnAnimatablePropertyChanged(NAME_Vector2D).RemoveAll(this);
+		ObjectChangeListener.GetOnAnimatablePropertyChanged(NAME_Color).RemoveAll(this);
+		ObjectChangeListener.GetOnAnimatablePropertyChanged(NAME_LinearColor).RemoveAll(this);
 	}
 }
 
diff --git a/Engine/Source/Editor/PropertyEditor/Public/IPropertyChangeListener.h b/Engine/Source/Editor/PropertyEditor/Public/IPropertyChangeListener.h
--- a/Engine/Source/Editor/PropertyEditor/Public/IPropertyChangeListener.h
+++ b/Engine/Source/Editor/PropertyEditor/Public/IPropertyChangeListener.h
@@ -24,7 +24,7 @@ class IPropertyChangeListener
 {
 public:
 	/** Delegate called when a property changes */
-	DECLARE_MULTICAST_DELEGATE_TwoParams( FOnPropertyChanged, const TArray<UObject*>&, const TSharedRef< const class IPropertyHandle>  );
+	DECLARE_MULTICAST_DELEGATE_TwoParams( FOnPropertyChanged, const TArray<UObject*>&, TSharedRef< const class IPropertyHandle>  );
 
 	/**
 	 * Sets the object that should be listened to for changes
diff --git a/Engine/Source/Editor/Sequencer/Private/SAnimationOutlinerView.cpp b/Engine/Source/Editor/Sequencer/Private/SAnimationOutlinerView.cpp
--- a/Engine/Source/Editor/Sequencer/Private/SAnimationOutlinerView.cpp
+++ b/Engine/Source/Editor/Sequencer/Private/SAnimationOutlinerView.cpp
@@ -239,7 +239,7 @@ void SAnimationOutlinerView::OnSelectionChanged( TSharedPtr<FSequencerDisplayNod
 		TArray<UObject*> RuntimeObjects;
 		Sequencer.Pin()->GetRuntimeObjects( Sequencer.Pin()->GetFocusedMovieSceneInstance(), ObjectNode->GetObjectBinding(), RuntimeObjects );
 		
-		if( RuntimeObjects.Num() > 0 )
+		if( RuntimeObjects.Num() > 0 && Sequencer.Pin()->IsLevelEditorSequencer() )
 		{
 			const bool bNotifySelectionChanged = false;
 			const bool bDeselectBSP = true;
diff --git a/Engine/Source/Editor/Sequencer/Private/Sequencer.h b/Engine/Source/Editor/Sequencer/Private/Sequencer.h
--- a/Engine/Source/Editor/Sequencer/Private/Sequencer.h
+++ b/Engine/Source/Editor/Sequencer/Private/Sequencer.h
@@ -74,6 +74,7 @@ public:
 	virtual FGuid GetHandleToObject(UObject* Object) override;
 	virtual ISequencerObjectChangeListener& GetObjectChangeListener() override;
 	virtual void NotifyMovieSceneDataChanged() override;
+	virtual void UpdateRuntimeInstances() override;
 	virtual void AddSubMovieScene(UMovieScene* SubMovieScene) override;
 	virtual void FilterToShotSections(const TArray< TWeakObjectPtr<class UMovieSceneSection> >& ShotSections, bool bZoomToShotBounds = true) override;
 	virtual void FilterToSelectedShotSections(bool bZoomToShotBounds = true) override;
@@ -292,8 +293,6 @@ protected:
 	 */
 	void ResetPerMovieSceneData();
 
-	/** Updates runtime instances when a movie scene changes */
-	void UpdateRuntimeInstances();
 
 	/**
 	 * Destroys spawnables for all movie scenes in the stack
diff --git a/Engine/Source/Editor/Sequencer/Private/SequencerObjectChangeListener.cpp b/Engine/Source/Editor/Sequencer/Private/SequencerObjectChangeListener.cpp
--- a/Engine/Source/Editor/Sequencer/Private/SequencerObjectChangeListener.cpp
+++ b/Engine/Source/Editor/Sequencer/Private/SequencerObjectChangeListener.cpp
@@ -24,9 +24,25 @@ FSequencerObjectChangeListener::~FSequencerObjectChangeListener()
 	GEditor->OnActorMoved().RemoveAll( this );
 }
 
-void FSequencerObjectChangeListener::OnPropertyChanged( const TArray<UObject*>& ChangedObjects, const TSharedRef< const IPropertyHandle> PropertyValue, bool bRequireAutoKey )
+void FSequencerObjectChangeListener::OnPropertyChanged( const TArray<UObject*>& ChangedObjects,  TSharedRef< const IPropertyHandle> PropertyValue, bool bRequireAutoKey )
 {
-	ClassToPropertyChangedMap.FindRef( PropertyValue->GetProperty()->GetClass() ).Broadcast( ChangedObjects, *PropertyValue, bRequireAutoKey );
+	TSharedPtr< const IPropertyHandle> HandleToUse = PropertyValue;
+
+	TSharedPtr<const IPropertyHandle> ParentHandle = PropertyValue->GetParentHandle();
+
+	if( ParentHandle.IsValid() && ParentHandle->GetProperty() && ParentHandle->GetProperty()->IsA<UStructProperty>() )
+	{
+		// Handle the case where internal values of a struct are being set.  We want to key the whole struct in this case
+		// @todo Sequencer: Is this always the case?
+		HandleToUse = ParentHandle;
+	}
+
+	UProperty* Property = HandleToUse->GetProperty();
+	UStructProperty* StructProperty = Cast<UStructProperty>( Property );
+
+	FName PropertyName = StructProperty ? StructProperty->Struct->GetFName() : Property->GetFName();
+
+	ClassToPropertyChangedMap.FindRef( PropertyName ).Broadcast( ChangedObjects, *HandleToUse, bRequireAutoKey );
 }
 
 bool FSequencerObjectChangeListener::IsObjectValidForListening( UObject* Object ) const
@@ -36,9 +52,9 @@ bool FSequencerObjectChangeListener::IsObjectValidForListening( UObject* Object
 	return bListenForActorsOnly ? Object->IsA<AActor>() : true;
 }
 
-FOnAnimatablePropertyChanged& FSequencerObjectChangeListener::GetOnAnimatablePropertyChanged( TSubclassOf<UProperty> PropertyClass )
+FOnAnimatablePropertyChanged& FSequencerObjectChangeListener::GetOnAnimatablePropertyChanged( FName PropertyTypeName )
 {
-	return ClassToPropertyChangedMap.FindOrAdd( PropertyClass );
+	return ClassToPropertyChangedMap.FindOrAdd( PropertyTypeName );
 }
 
 FOnPropagateObjectChanges& FSequencerObjectChangeListener::GetOnPropagateObjectChanges()
@@ -84,7 +100,7 @@ void FSequencerObjectChangeListener::OnObjectPreEditChange( UObject* Object )
 
 void FSequencerObjectChangeListener::OnObjectPostEditChange( UObject* Object, FPropertyChangedEvent& PropertyChangedEvent )
 {
-	if( Object )
+	if( Object && PropertyChangedEvent.ChangeType != EPropertyChangeType::Interactive )
 	{
 		bool bIsObjectValid = IsObjectValidForListening( Object );
 
diff --git a/Engine/Source/Editor/Sequencer/Private/SequencerObjectChangeListener.h b/Engine/Source/Editor/Sequencer/Private/SequencerObjectChangeListener.h
--- a/Engine/Source/Editor/Sequencer/Private/SequencerObjectChangeListener.h
+++ b/Engine/Source/Editor/Sequencer/Private/SequencerObjectChangeListener.h
@@ -14,7 +14,7 @@ public:
 	~FSequencerObjectChangeListener();
 
 	/** ISequencerObjectChangeListener interface */
-	virtual FOnAnimatablePropertyChanged& GetOnAnimatablePropertyChanged( TSubclassOf<UProperty> PropertyClass ) override;
+	virtual FOnAnimatablePropertyChanged& GetOnAnimatablePropertyChanged( FName PropertyTypeName ) override;
 	virtual FOnPropagateObjectChanges& GetOnPropagateObjectChanges() override;
 	virtual void TriggerAllPropertiesChanged(UObject* Object) override;
 
@@ -45,7 +45,7 @@ private:
 	 *
 	 * @param Object	The object that PostEditChange was called on
 	 */
-	void OnPropertyChanged( const TArray<UObject*>& ChangedObjects, const TSharedRef< const IPropertyHandle> PropertyValue, bool bRequireAutoKey );
+	void OnPropertyChanged( const TArray<UObject*>& ChangedObjects, TSharedRef< const IPropertyHandle> PropertyValue, bool bRequireAutoKey );
 
 	/**
 	 * @return True if an object is valid for listening to property changes 
@@ -57,7 +57,7 @@ private:
 	TMap< TWeakObjectPtr<UObject>, TSharedPtr<class IPropertyChangeListener> > ActivePropertyChangeListeners;
 
 	/** A mapping of property classes to multi-cast delegate that is broadcast when properties of that type change */
-	TMap< TSubclassOf<UProperty>, FOnAnimatablePropertyChanged > ClassToPropertyChangedMap;
+	TMap< FName, FOnAnimatablePropertyChanged > ClassToPropertyChangedMap;
 
 	/** Delegate to call when object changes should be propagated */
 	FOnPropagateObjectChanges OnPropagateObjectChanges;
diff --git a/Engine/Source/Editor/Sequencer/Public/ISequencer.h b/Engine/Source/Editor/Sequencer/Public/ISequencer.h
--- a/Engine/Source/Editor/Sequencer/Public/ISequencer.h
+++ b/Engine/Source/Editor/Sequencer/Public/ISequencer.h
@@ -114,4 +114,6 @@ public:
 	virtual class ISequencerObjectChangeListener& GetObjectChangeListener() = 0;
 
 	virtual void NotifyMovieSceneDataChanged() = 0;
+
+	virtual void UpdateRuntimeInstances() = 0;
 };
diff --git a/Engine/Source/Editor/Sequencer/Public/ISequencerObjectChangeListener.h b/Engine/Source/Editor/Sequencer/Public/ISequencerObjectChangeListener.h
--- a/Engine/Source/Editor/Sequencer/Public/ISequencerObjectChangeListener.h
+++ b/Engine/Source/Editor/Sequencer/Public/ISequencerObjectChangeListener.h
@@ -17,7 +17,7 @@ public:
 	/**
 	 * A delegate for when a property of a specific UProperty class is changed.  
 	 */
-	virtual FOnAnimatablePropertyChanged& GetOnAnimatablePropertyChanged( TSubclassOf<UProperty> PropertyClass ) = 0;
+	virtual FOnAnimatablePropertyChanged& GetOnAnimatablePropertyChanged( FName PropertyTypeName ) = 0;
 
 	/**
 	 * A delegate for when object changes should be propagated to/from puppet actors
diff --git a/Engine/Source/Editor/UMGEditor/Private/Animation/MarginTrackEditor.cpp b/Engine/Source/Editor/UMGEditor/Private/Animation/MarginTrackEditor.cpp
new file mode 100644
--- /dev/null
+++ b/Engine/Source/Editor/UMGEditor/Private/Animation/MarginTrackEditor.cpp
@@ -0,0 +1,125 @@
+// Copyright 1998-2014 Epic Games, Inc. All Rights Reserved.
+
+#include "UMGEditorPrivatePCH.h"
+#include "MarginTrackEditor.h"
+#include "Developer/MovieSceneTools/Public/PropertySection.h"
+#include "Developer/MovieSceneTools/Public/MovieSceneToolHelpers.h"
+#include "Runtime/UMG/Public/Animation/MovieSceneMarginSection.h"
+#include "Runtime/UMG/Public/Animation/MovieSceneMarginTrack.h"
+#include "Editor/Sequencer/Public/ISectionLayoutBuilder.h"
+#include "Editor/Sequencer/Public/ISequencerObjectChangeListener.h"
+#include "Editor/PropertyEditor/Public/PropertyHandle.h"
+
+class FMarginPropertySection : public FPropertySection
+{
+public:
+	FMarginPropertySection( UMovieSceneSection& InSectionObject, FName SectionName )
+		: FPropertySection(InSectionObject, SectionName) {}
+
+	virtual void GenerateSectionLayout( class ISectionLayoutBuilder& LayoutBuilder ) const override
+	{
+		UMovieSceneMarginSection* MarginSection = Cast<UMovieSceneMarginSection>(&SectionObject);
+
+		LayoutBuilder.AddKeyArea("Left", NSLOCTEXT("FMarginPropertySection", "MarginLeft", "Left"), MakeShareable(new FFloatCurveKeyArea(MarginSection->GetLeftCurve())));
+		LayoutBuilder.AddKeyArea("Top", NSLOCTEXT("FMarginPropertySection", "MarginTop", "Top"), MakeShareable(new FFloatCurveKeyArea(MarginSection->GetTopCurve())));
+		LayoutBuilder.AddKeyArea("Right", NSLOCTEXT("FMarginPropertySection", "MarginRight", "Right"), MakeShareable(new FFloatCurveKeyArea(MarginSection->GetRightCurve())));
+		LayoutBuilder.AddKeyArea("Bottom", NSLOCTEXT("FMarginPropertySection", "MarginBottom", "Bottom"), MakeShareable(new FFloatCurveKeyArea(MarginSection->GetBottomCurve())));
+	}
+};
+
+FMarginTrackEditor::FMarginTrackEditor( TSharedRef<ISequencer> InSequencer )
+	: FMovieSceneTrackEditor( InSequencer ) 
+{
+	// Get the object change listener for the sequencer and register a delegates for when properties change that we care about
+	ISequencerObjectChangeListener& ObjectChangeListener = InSequencer->GetObjectChangeListener();
+	ObjectChangeListener.GetOnAnimatablePropertyChanged( "Margin" ).AddRaw( this, &FMarginTrackEditor::OnMarginChanged );
+}
+
+FMarginTrackEditor::~FMarginTrackEditor()
+{
+	TSharedPtr<ISequencer> Sequencer = GetSequencer();
+	if( Sequencer.IsValid() )
+	{
+		ISequencerObjectChangeListener& ObjectChangeListener = Sequencer->GetObjectChangeListener();
+		ObjectChangeListener.GetOnAnimatablePropertyChanged( "Margin" ).RemoveAll( this );
+	}
+}
+
+
+
+TSharedRef<FMovieSceneTrackEditor> FMarginTrackEditor::CreateTrackEditor( TSharedRef<ISequencer> InSequencer )
+{
+	return MakeShareable( new FMarginTrackEditor( InSequencer ) );
+}
+
+bool FMarginTrackEditor::SupportsType( TSubclassOf<UMovieSceneTrack> Type ) const
+{
+	return Type == UMovieSceneMarginTrack::StaticClass();
+}
+
+TSharedRef<ISequencerSection> FMarginTrackEditor::MakeSectionInterface( UMovieSceneSection& SectionObject, UMovieSceneTrack* Track )
+{
+	check( SupportsType( SectionObject.GetOuter()->GetClass() ) );
+
+	UClass* SectionClass = SectionObject.GetOuter()->GetClass();
+
+	TSharedRef<ISequencerSection> NewSection = MakeShareable( new FMarginPropertySection( SectionObject, Track->GetTrackName() ) );
+
+	return NewSection;
+}
+/*
+
+void FMarginTrackEditor::AddKey(const FGuid& ObjectGuid, UObject* AdditionalAsset)
+{
+	ISequencerObjectChangeListener& ObjectChangeListener = GetSequencer()->GetObjectChangeListener();
+
+	TArray<UObject*> OutObjects;
+	GetSequencer()->GetRuntimeObjects( GetSequencer()->GetFocusedMovieSceneInstance(), ObjectGuid, OutObjects);
+	for (int32 i = 0; i < OutObjects.Num(); ++i)
+	{
+		ObjectChangeListener.TriggerAllPropertiesChanged(OutObjects[i]);
+	}
+}*/
+
+void FMarginTrackEditor::OnMarginChanged( const TArray<UObject*>& InObjectsThatChanged, const IPropertyHandle& PropertyValue, bool bRequireAutoKey )
+{
+	FName PropertyName = PropertyValue.GetProperty()->GetFName();
+
+	AnimatablePropertyChanged
+	(
+		UMovieSceneMarginTrack::StaticClass(), 
+		bRequireAutoKey,
+		FOnKeyProperty::CreateRaw(this, &FMarginTrackEditor::OnKeyMargin, &InObjectsThatChanged, &PropertyValue, PropertyName ) 
+	);
+}
+
+
+void FMarginTrackEditor::OnKeyMargin( float KeyTime, const TArray<UObject*>* InObjectsThatChanged, const IPropertyHandle* PropertyValue, FName PropertyName )
+{
+	TArray<const void*> MarginValues;
+	PropertyValue->AccessRawData( MarginValues );
+
+	for( int32 ObjectIndex = 0; ObjectIndex < InObjectsThatChanged->Num(); ++ObjectIndex )
+	{
+		UObject* Object = (*InObjectsThatChanged)[ObjectIndex];
+		FMargin MarginValue = *(const FMargin*)MarginValues[ObjectIndex];
+
+		FGuid ObjectHandle = FindOrCreateHandleToObject( Object );
+		if (ObjectHandle.IsValid())
+		{
+			UMovieSceneTrack* Track = GetTrackForObject( ObjectHandle, UMovieSceneMarginTrack::StaticClass(), PropertyName );
+			if( ensure( Track ) )
+			{
+				UMovieSceneMarginTrack* MarginTrack = CastChecked<UMovieSceneMarginTrack>(Track);
+				MarginTrack->SetPropertyName( PropertyName );
+				// Find or add a new section at the auto-key time and changing the property same property
+				// AddKeyToSection is not actually a virtual, it's redefined in each class with a different type
+				bool bSuccessfulAdd = MarginTrack->AddKeyToSection( KeyTime, MarginValue );
+				if (bSuccessfulAdd)
+				{
+					MarginTrack->SetAsShowable();
+				}
+			}
+		}
+	}
+}
\ No newline at end of file
diff --git a/Engine/Source/Editor/UMGEditor/Private/Animation/MarginTrackEditor.h b/Engine/Source/Editor/UMGEditor/Private/Animation/MarginTrackEditor.h
new file mode 100644
--- /dev/null
+++ b/Engine/Source/Editor/UMGEditor/Private/Animation/MarginTrackEditor.h
@@ -0,0 +1,45 @@
+// Copyright 1998-2014 Epic Games, Inc. All Rights Reserved.
+
+#pragma once
+
+#include "Editor/Sequencer/Public/MovieSceneTrackEditor.h"
+
+class IPropertyHandle;
+
+class FMarginTrackEditor : public FMovieSceneTrackEditor
+{
+public:
+	/**
+	 * Constructor
+	 *
+	 * @param InSequencer	The sequencer instance to be used by this tool
+	 */
+	FMarginTrackEditor( TSharedRef<ISequencer> InSequencer );
+	~FMarginTrackEditor();
+
+	/**
+	 * Creates an instance of this class.  Called by a sequencer 
+	 *
+	 * @param OwningSequencer The sequencer instance to be used by this tool
+	 * @return The new instance of this class
+	 */
+	static TSharedRef<FMovieSceneTrackEditor> CreateTrackEditor( TSharedRef<ISequencer> OwningSequencer );
+
+	/** FMovieSceneTrackEditor Interface */
+	virtual bool SupportsType( TSubclassOf<UMovieSceneTrack> Type ) const override;
+	virtual TSharedRef<ISequencerSection> MakeSectionInterface( UMovieSceneSection& SectionObject, UMovieSceneTrack* Track ) override;
+
+private:
+	/**
+	 * Called by the details panel when an animatable property changes
+	 *
+	 * @param InObjectsThatChanged	List of objects that changed
+	 * @param PropertyValue			Handle to the property value which changed
+	 */
+	void OnMarginChanged( const TArray<UObject*>& InObjectsThatChanged, const IPropertyHandle& PropertyValue, bool bRequireAutoKey );
+
+	/** Called After OnMarginChanged if we actually can key the margin */
+	void OnKeyMargin( float KeyTime, const TArray<UObject*>* InObjectsThatChanged, const IPropertyHandle* PropertyValue, FName PropertyName );
+};
+
+
diff --git a/Engine/Source/Editor/UMGEditor/Private/UMGEditorModule.cpp b/Engine/Source/Editor/UMGEditor/Private/UMGEditorModule.cpp
--- a/Engine/Source/Editor/UMGEditor/Private/UMGEditorModule.cpp
+++ b/Engine/Source/Editor/UMGEditor/Private/UMGEditorModule.cpp
@@ -11,6 +11,9 @@
 #include "KismetCompilerModule.h"
 #include "WidgetBlueprintCompiler.h"
 
+#include "Editor/Sequencer/Public/ISequencerModule.h"
+#include "Animation/MarginTrackEditor.h"
+
 const FName UMGEditorAppIdentifier = FName(TEXT("UMGEditorApp"));
 
 class FUMGEditorModule : public IUMGEditorModule
@@ -40,6 +43,12 @@ public:
 		KismetCompilerModule.GetCompilers().Add(sd);
 		
 		FUMGBlueprintEditorExtensionHook::InstallHooks();
+
+
+		// Register with the sequencer module that we provide auto-key handlers.
+		ISequencerModule& SequencerModule = FModuleManager::Get().LoadModuleChecked<ISequencerModule>("Sequencer");
+		SequencerModule.RegisterTrackEditor(FOnCreateTrackEditor::CreateStatic(&FMarginTrackEditor::CreateTrackEditor));
+
 	}
 
 	/** Called before the module is unloaded, right before the module object is destroyed. */
diff --git a/Engine/Source/Editor/UMGEditor/Private/UMGSequencerObjectBindingManager.cpp b/Engine/Source/Editor/UMGEditor/Private/UMGSequencerObjectBindingManager.cpp
--- a/Engine/Source/Editor/UMGEditor/Private/UMGSequencerObjectBindingManager.cpp
+++ b/Engine/Source/Editor/UMGEditor/Private/UMGSequencerObjectBindingManager.cpp
@@ -3,10 +3,14 @@
 #include "UMGEditorPrivatePCH.h"
 #include "UMGSequencerObjectBindingManager.h"
 
-FUMGSequencerObjectBindingManager::FUMGSequencerObjectBindingManager( UWidgetBlueprint& InWidgetBlueprint )
-	: WidgetBlueprint( &InWidgetBlueprint )
+FUMGSequencerObjectBindingManager::FUMGSequencerObjectBindingManager( FWidgetBlueprintEditor& InWidgetBlueprintEditor )
+	: WidgetBlueprintEditor( InWidgetBlueprintEditor )
+{
+	WidgetBlueprintEditor.GetOnWidgetPreviewUpdated().AddRaw( this, &FUMGSequencerObjectBindingManager::OnWidgetPreviewUpdated );
+}
+
+FUMGSequencerObjectBindingManager::~FUMGSequencerObjectBindingManager()
 {
-	
 }
 
 FGuid FUMGSequencerObjectBindingManager::FindGuidForObject( const UMovieScene& MovieScene, UObject& Object ) const
@@ -14,10 +18,11 @@ FGuid FUMGSequencerObjectBindingManager::FindGuidForObject( const UMovieScene& M
 	return PreviewObjectToGuidMap.FindRef( &Object );
 }
 
-
 bool FUMGSequencerObjectBindingManager::CanPossessObject( UObject& Object ) const
 {
-	return !Object.IsIn( WidgetBlueprint.Get() );
+	UWidgetBlueprint* WidgetBlueprint = WidgetBlueprintEditor.GetWidgetBlueprintObj();
+
+	return !Object.IsIn( WidgetBlueprint );
 }
 
 void FUMGSequencerObjectBindingManager::BindPossessableObject( const FGuid& PossessableGuid, UObject& PossessedObject )
@@ -47,3 +52,31 @@ void FUMGSequencerObjectBindingManager::GetRuntimeObjects( const TSharedRef<FMov
 		}
 	}
 }
+
+void FUMGSequencerObjectBindingManager::OnWidgetPreviewUpdated()
+{
+	//
+	UUserWidget* PreviewWidget = WidgetBlueprintEditor.GetPreview();
+
+	UWidgetTree* WidgetTree = PreviewWidget->WidgetTree;
+
+	TMap< TWeakObjectPtr<UObject>, FGuid> OldPreviewObjectToGuidMap = PreviewObjectToGuidMap;
+	PreviewObjectToGuidMap.Empty();
+
+	for( auto It = OldPreviewObjectToGuidMap.CreateConstIterator(); It; ++It )
+	{
+		const TWeakObjectPtr<UObject> Object = It.Key();
+		if(Object.IsValid())
+		{
+			FName ObjectName = Object->GetFName();
+			UObject* NewObject = FindObject<UObject>( WidgetTree, *ObjectName.ToString() );
+			if( NewObject )
+			{
+				check( NewObject->GetClass() == Object->GetClass() );
+				PreviewObjectToGuidMap.Add( NewObject, It.Value() );
+			}
+		}
+	}
+
+	WidgetBlueprintEditor.GetSequencer()->UpdateRuntimeInstances();
+}
\ No newline at end of file
diff --git a/Engine/Source/Editor/UMGEditor/Private/UMGSequencerObjectBindingManager.h b/Engine/Source/Editor/UMGEditor/Private/UMGSequencerObjectBindingManager.h
--- a/Engine/Source/Editor/UMGEditor/Private/UMGSequencerObjectBindingManager.h
+++ b/Engine/Source/Editor/UMGEditor/Private/UMGSequencerObjectBindingManager.h
@@ -7,8 +7,9 @@
 class FUMGSequencerObjectBindingManager : public ISequencerObjectBindingManager
 {
 public:
-	FUMGSequencerObjectBindingManager( UWidgetBlueprint& InWidgetBlueprint );
-	
+	FUMGSequencerObjectBindingManager( FWidgetBlueprintEditor& InWidgetBlueprintEditor );
+	~FUMGSequencerObjectBindingManager();
+
 	/** ISequencerObjectBindingManager interface */
 	virtual bool AllowsSpawnableObjects() const override { return false; }
 	virtual FGuid FindGuidForObject( const UMovieScene& MovieScene, UObject& Object ) const override;
@@ -19,9 +20,11 @@ public:
 	virtual void UnbindPossessableObjects( const FGuid& PossessableGuid ) override;
 	virtual void GetRuntimeObjects( const TSharedRef<FMovieSceneInstance>& MovieSceneInstance, const FGuid& ObjectGuid, TArray<UObject*>& OutRuntimeObjects ) const override;
 	
+private:
+	void OnWidgetPreviewUpdated();
 private:
 	/** Mapping of preview objects to sequencer guids */
 	TMap< TWeakObjectPtr<UObject>, FGuid> PreviewObjectToGuidMap;
 	
-	TWeakObjectPtr<UWidgetBlueprint> WidgetBlueprint;
+	FWidgetBlueprintEditor& WidgetBlueprintEditor;
 };
diff --git a/Engine/Source/Editor/UMGEditor/Private/WidgetBlueprintEditor.cpp b/Engine/Source/Editor/UMGEditor/Private/WidgetBlueprintEditor.cpp
--- a/Engine/Source/Editor/UMGEditor/Private/WidgetBlueprintEditor.cpp
+++ b/Engine/Source/Editor/UMGEditor/Private/WidgetBlueprintEditor.cpp
@@ -184,6 +184,9 @@ void FWidgetBlueprintEditor::AddReferencedObjects( FReferenceCollector& Collecto
 	{
 		Collector.AddReferencedObject( DefaultMovieScene );
 	}
+
+	UUserWidget* Preview = GetPreview();
+	Collector.AddReferencedObject( Preview );
 }
 
 void FWidgetBlueprintEditor::NotifyPreChange(FEditPropertyChain* PropertyAboutToChange)
@@ -261,7 +264,7 @@ TSharedPtr<ISequencer>& FWidgetBlueprintEditor::GetSequencer()
 
 		UMovieScene* MovieScene = Blueprint->AnimationData.Num() ? Blueprint->AnimationData[0] : GetDefaultMovieScene();
 		
-		TSharedRef<FUMGSequencerObjectBindingManager> ObjectBindingManager = MakeShareable( new FUMGSequencerObjectBindingManager( *Blueprint ) );
+		TSharedRef<FUMGSequencerObjectBindingManager> ObjectBindingManager = MakeShareable( new FUMGSequencerObjectBindingManager( *this ) );
 		
 		Sequencer = FModuleManager::LoadModuleChecked< ISequencerModule >("Sequencer").CreateSequencer( MovieScene, ObjectBindingManager );
 	}
@@ -276,7 +279,6 @@ void FWidgetBlueprintEditor::DestroyPreview()
 	{
 		check(PreviewScene.GetWorld());
 
-		PreviewActor->ClearFlags(RF_Standalone);
 		PreviewActor->MarkPendingKill();
 	}
 
@@ -316,7 +318,7 @@ void FWidgetBlueprintEditor::UpdatePreview(UBlueprint* InBlueprint, bool bInForc
 		PreviewBlueprint = InBlueprint;
 
 		PreviewActor = ConstructObject<UUserWidget>(PreviewBlueprint->GeneratedClass, PreviewScene.GetWorld()->GetCurrentLevel());
-		PreviewActor->SetFlags(RF_Standalone | RF_Transactional);
+		PreviewActor->SetFlags(RF_Transactional);
 		
 		// Configure all the widgets to be set to design time.
 		PreviewActor->IsDesignTime(true);
@@ -328,6 +330,9 @@ void FWidgetBlueprintEditor::UpdatePreview(UBlueprint* InBlueprint, bool bInForc
 		// Store a reference to the preview actor.
 		PreviewWidgetActorPtr = PreviewActor;
 	}
+
+	OnWidgetPreviewUpdated.Broadcast();
+
 }
 
 UMovieScene* FWidgetBlueprintEditor::GetDefaultMovieScene()
diff --git a/Engine/Source/Editor/UMGEditor/Private/WidgetBlueprintEditor.h b/Engine/Source/Editor/UMGEditor/Private/WidgetBlueprintEditor.h
--- a/Engine/Source/Editor/UMGEditor/Private/WidgetBlueprintEditor.h
+++ b/Engine/Source/Editor/UMGEditor/Private/WidgetBlueprintEditor.h
@@ -16,6 +16,8 @@ class FWidgetBlueprintEditor : public FBlueprintEditor
 public:
 	DECLARE_MULTICAST_DELEGATE(FOnSelectedWidgetsChanged)
 
+	/** Called after the widget preview has been updated */
+	DECLARE_MULTICAST_DELEGATE(FOnWidgetPreviewUpdated)
 public:
 	FWidgetBlueprintEditor();
 
@@ -45,6 +47,9 @@ public:
 	/** @return The selected set of widgets */
 	const TSet<FWidgetReference>& GetSelectedWidgets() const;
 
+	/** @return Notification for when the preview widget has been updated */
+	FOnWidgetPreviewUpdated& GetOnWidgetPreviewUpdated() { return OnWidgetPreviewUpdated; }
+
 	/** Fires whenever the selected set of widgets changes */
 	FOnSelectedWidgetsChanged OnSelectedWidgetsChanged;
 
@@ -88,4 +93,7 @@ private:
 
 	/** The preview GUI object */
 	mutable TWeakObjectPtr<UUserWidget> PreviewWidgetActorPtr;
+	
+	/** Notification for when the preview widget has been updated  */
+	FOnWidgetPreviewUpdated OnWidgetPreviewUpdated;
 };
diff --git a/Engine/Source/Runtime/MovieSceneCoreTypes/Private/Instances/MovieSceneColorTrackInstance.cpp b/Engine/Source/Runtime/MovieSceneCoreTypes/Private/Instances/MovieSceneColorTrackInstance.cpp
--- a/Engine/Source/Runtime/MovieSceneCoreTypes/Private/Instances/MovieSceneColorTrackInstance.cpp
+++ b/Engine/Source/Runtime/MovieSceneCoreTypes/Private/Instances/MovieSceneColorTrackInstance.cpp
@@ -27,7 +27,7 @@ void FMovieSceneColorTrackInstance::Update( float Position, float LastPosition,
 			else
 			{
 				// @todo Sequencer - remove this legacy method of setting the property directly
-				UObject* PropertyOwner = NULL;
+				/*UObject* PropertyOwner = NULL;
 				UProperty* Property = NULL;
 				//@todo Sequencer - Major performance problems here.  This needs to be initialized and stored (not serialized) somewhere
 				void* Address = FMatineeUtils::GetPropertyAddress<void>( Object, ColorTrack->GetPropertyName(), Property, PropertyOwner );
@@ -54,7 +54,7 @@ void FMovieSceneColorTrackInstance::Update( float Position, float LastPosition,
 						// Let the property owner know that we changed one of its properties
 						PropertyOwner->PostInterpChange( Property );
 					}
-				}
+				}*/
 			}
 		}
 	}
diff --git a/Engine/Source/Runtime/UMG/Private/Animation/MovieSceneMarginSection.cpp b/Engine/Source/Runtime/UMG/Private/Animation/MovieSceneMarginSection.cpp
new file mode 100644
--- /dev/null
+++ b/Engine/Source/Runtime/UMG/Private/Animation/MovieSceneMarginSection.cpp
@@ -0,0 +1,66 @@
+// Copyright 1998-2014 Epic Games, Inc. All Rights Reserved.
+
+#include "UMGPrivatePCH.h"
+#include "MovieSceneMarginSection.h"
+
+
+UMovieSceneMarginSection::UMovieSceneMarginSection( const FPostConstructInitializeProperties& PCIP )
+	: Super( PCIP )
+{
+}
+
+void UMovieSceneMarginSection::MoveSection( float DeltaTime )
+{
+	Super::MoveSection( DeltaTime );
+
+	// Move all the curves in this section
+	TopCurve.ShiftCurve(DeltaTime);
+	LeftCurve.ShiftCurve(DeltaTime);
+	RightCurve.ShiftCurve(DeltaTime);
+	BottomCurve.ShiftCurve(DeltaTime);
+}
+
+void UMovieSceneMarginSection::DilateSection( float DilationFactor, float Origin )
+{
+	TopCurve.ScaleCurve(Origin, DilationFactor);
+	LeftCurve.ScaleCurve(Origin, DilationFactor);
+	RightCurve.ScaleCurve(Origin, DilationFactor);
+	BottomCurve.ScaleCurve(Origin, DilationFactor);
+
+	Super::DilateSection(DilationFactor, Origin);
+}
+
+FMargin UMovieSceneMarginSection::Eval( float Position ) const
+{
+	return FMargin(	TopCurve.Eval(Position),
+					LeftCurve.Eval(Position),
+					RightCurve.Eval(Position),
+					BottomCurve.Eval(Position));
+}
+
+void UMovieSceneMarginSection::AddKey( float Time, const FMargin& Value )
+{
+	AddKeyToCurve(TopCurve, Time, Value.Top);
+	AddKeyToCurve(LeftCurve, Time, Value.Left);
+	AddKeyToCurve(RightCurve, Time, Value.Right);
+	AddKeyToCurve(BottomCurve, Time, Value.Bottom);
+}
+
+bool UMovieSceneMarginSection::NewKeyIsNewData(float Time, const FMargin& Value) const
+{
+	return TopCurve.GetNumKeys() == 0 ||
+		LeftCurve.GetNumKeys() == 0 ||
+		RightCurve.GetNumKeys() == 0 ||
+		BottomCurve.GetNumKeys() == 0 ||
+		(Eval(Time) != Value);
+}
+
+
+
+void UMovieSceneMarginSection::AddKeyToCurve( FRichCurve& InCurve, float Time, float Value )
+{
+	if( IsTimeWithinSection(Time) )
+	{
+		InCurve.UpdateOrAddKey(Time, Value);
+	}
+}
diff --git a/Engine/Source/Runtime/UMG/Private/Animation/MovieSceneMarginTrack.cpp b/Engine/Source/Runtime/UMG/Private/Animation/MovieSceneMarginTrack.cpp
new file mode 100644
--- /dev/null
+++ b/Engine/Source/Runtime/UMG/Private/Animation/MovieSceneMarginTrack.cpp
@@ -0,0 +1,51 @@
+// Copyright 1998-2014 Epic Games, Inc. All Rights Reserved.
+
+#include "UMGPrivatePCH.h"
+#include "MovieSceneMarginSection.h"
+#include "MovieSceneMarginTrack.h"
+#include "IMovieScenePlayer.h"
+#include "MovieSceneMarginTrackInstance.h"
+#include "MovieSceneCommonHelpers.h"
+
+UMovieSceneMarginTrack::UMovieSceneMarginTrack(const FPostConstructInitializeProperties& PCIP)
+: Super(PCIP)
+{
+}
+
+UMovieSceneSection* UMovieSceneMarginTrack::CreateNewSection()
+{
+	return ConstructObject<UMovieSceneSection>( UMovieSceneMarginSection::StaticClass(), this );
+}
+
+TSharedPtr<IMovieSceneTrackInstance> UMovieSceneMarginTrack::CreateInstance()
+{
+	return MakeShareable( new FMovieSceneMarginTrackInstance( *this ) ); 
+}
+
+
+bool UMovieSceneMarginTrack::AddKeyToSection( float Time, FMargin Value )
+{
+	const UMovieSceneSection* NearestSection = MovieSceneHelpers::FindSectionAtTime( Sections, Time );
+	if (!NearestSection || CastChecked<UMovieSceneMarginSection>(NearestSection)->NewKeyIsNewData(Time, Value))
+	{
+		UMovieSceneMarginSection* NewSection = CastChecked<UMovieSceneMarginSection>( FindOrAddSection( Time ) );
+
+		NewSection->AddKey( Time, Value );
+
+		return true;
+	}
+	return false;
+}
+
+
+bool UMovieSceneMarginTrack::Eval( float Position, float LastPosition, FMargin& OutMargin ) const
+{
+	const UMovieSceneSection* Section = MovieSceneHelpers::FindSectionAtTime( Sections, Position );
+
+	if( Section )
+	{
+		OutMargin = CastChecked<UMovieSceneMarginSection>( Section )->Eval( Position );
+	}
+
+	return Section != NULL;
+}
\ No newline at end of file
diff --git a/Engine/Source/Runtime/UMG/Private/Animation/MovieSceneMarginTrackInstance.cpp b/Engine/Source/Runtime/UMG/Private/Animation/MovieSceneMarginTrackInstance.cpp
new file mode 100644
--- /dev/null
+++ b/Engine/Source/Runtime/UMG/Private/Animation/MovieSceneMarginTrackInstance.cpp
@@ -0,0 +1,52 @@
+// Copyright 1998-2014 Epic Games, Inc. All Rights Reserved.
+
+#include "UMGPrivatePCH.h"
+#include "MovieSceneMarginTrackInstance.h"
+
+FMovieSceneMarginTrackInstance::FMovieSceneMarginTrackInstance( UMovieSceneMarginTrack& InMarginTrack )
+{
+	MarginTrack = &InMarginTrack;
+}
+
+void FMovieSceneMarginTrackInstance::Update( float Position, float LastPosition, const TArray<UObject*>& RuntimeObjects, class IMovieScenePlayer& Player ) 
+{
+	FMargin MarginValue;
+	if( MarginTrack->Eval( Position, LastPosition, MarginValue ) )
+	{
+		for( int32 ObjIndex = 0; ObjIndex < RuntimeObjects.Num(); ++ObjIndex )
+		{
+			UObject* Object = RuntimeObjects[ObjIndex];
+
+			UFunction* Function = RuntimeObjectToFunctionMap.FindRef( Object );
+			if( Function )
+			{
+				Object->ProcessEvent( Function, &MarginValue );
+			}
+		}
+	}
+}
+
+void FMovieSceneMarginTrackInstance::RefreshInstance( const TArray<UObject*>& RuntimeObjects, class IMovieScenePlayer& Player )
+{
+	static const FString Set(TEXT("Set") );
+	
+	FString FunctionString = Set + MarginTrack->GetPropertyName().ToString();
+	
+	FName FunctionName = FName( *FunctionString );
+	
+	//Player.GetRuntimeObjects( )
+	for( UObject* Object : RuntimeObjects )
+	{
+		// If the mapping for the object already exists we dont need to search for it again
+		UFunction* Function = RuntimeObjectToFunctionMap.FindRef( Object );
+		
+		if( !Function )
+		{
+			Function = Object->FindFunction( FunctionName );
+			if( Function )
+			{
+				RuntimeObjectToFunctionMap.Add( Object, Function );
+			}
+		}
+	}
+}
diff --git a/Engine/Source/Runtime/UMG/Private/Animation/MovieSceneMarginTrackInstance.h b/Engine/Source/Runtime/UMG/Private/Animation/MovieSceneMarginTrackInstance.h
new file mode 100644
--- /dev/null
+++ b/Engine/Source/Runtime/UMG/Private/Animation/MovieSceneMarginTrackInstance.h
@@ -0,0 +1,23 @@
+// Copyright 1998-2014 Epic Games, Inc. All Rights Reserved.
+
+#pragma once
+
+#include "IMovieSceneTrackInstance.h"
+
+/**
+ * Instance of a UMovieSceneColorTrack
+ */
+class FMovieSceneMarginTrackInstance : public IMovieSceneTrackInstance
+{
+public:
+	FMovieSceneMarginTrackInstance( UMovieSceneMarginTrack& InMarginTrack );
+
+	/** IMovieSceneTrackInstance interface */
+	virtual void Update( float Position, float LastPosition, const TArray<UObject*>& RuntimeObjects, class IMovieScenePlayer& Player ) override;
+	virtual void RefreshInstance( const TArray<UObject*>& RuntimeObjects, class IMovieScenePlayer& Player ) override;
+private:
+	/** The track being instanced */
+	UMovieSceneMarginTrack* MarginTrack;
+	/** Mapping of objects to bound functions that will be called to update data on the track */
+	TMap< TWeakObjectPtr<UObject>, UFunction* > RuntimeObjectToFunctionMap;
+};
diff --git a/Engine/Source/Runtime/UMG/Private/WidgetBlueprintGeneratedClass.cpp b/Engine/Source/Runtime/UMG/Private/WidgetBlueprintGeneratedClass.cpp
--- a/Engine/Source/Runtime/UMG/Private/WidgetBlueprintGeneratedClass.cpp
+++ b/Engine/Source/Runtime/UMG/Private/WidgetBlueprintGeneratedClass.cpp
@@ -48,15 +48,10 @@ void UWidgetBlueprintGeneratedClass::Link(FArchive& Ar, bool bRelinkExistingProp
 
 void UWidgetBlueprintGeneratedClass::InitializeWidget(UUserWidget* UserWidget) const
 {
-	// Duplicate the graph, keeping track of what was duplicated
-	TMap<UObject*, UObject*> DuplicatedObjectList;
+	UWidgetTree* ClonedTree = DuplicateObject<UWidgetTree>( WidgetTree, UserWidget );
 
-	FObjectDuplicationParameters Parameters(const_cast<UWidgetTree*>( WidgetTree ), UserWidget);
-	Parameters.CreatedObjects = &DuplicatedObjectList;
+	UserWidget->WidgetTree = ClonedTree;
 
-	UWidgetTree* ClonedTree = CastChecked<UWidgetTree>(StaticDuplicateObjectEx(Parameters));
-
-	UUserWidget* WidgetActor = CastChecked<UUserWidget>(UserWidget);
 	UClass* ActorClass = UserWidget->GetClass();
 
 	for ( UWidget* Widget : ClonedTree->WidgetTemplates )
@@ -88,7 +83,7 @@ void UWidgetBlueprintGeneratedClass::InitializeWidget(UUserWidget* UserWidget) c
 			//TODO UMG Terrible performance, improve with Maps.
 			if ( Binding.ObjectName == VariableName )
 			{
-				UFunction* BoundFunction = WidgetActor->FindFunction(Binding.FunctionName);
+				UFunction* BoundFunction = UserWidget->FindFunction(Binding.FunctionName);
 
 				FString DelegateName = Binding.PropertyName.ToString() + "Delegate";
 
@@ -99,7 +94,7 @@ void UWidgetBlueprintGeneratedClass::InitializeWidget(UUserWidget* UserWidget) c
 						if ( DelegateProp->GetName() == DelegateName || DelegateProp->GetFName() == Binding.PropertyName )
 						{
 							FScriptDelegate* ScriptDelegate = DelegateProp->GetPropertyValuePtr_InContainer(Widget);
-							ScriptDelegate->BindUFunction(WidgetActor, Binding.FunctionName);
+							ScriptDelegate->BindUFunction(UserWidget, Binding.FunctionName);
 							break;
 						}
 					}
diff --git a/Engine/Source/Runtime/UMG/Public/Animation/MovieSceneMarginSection.h b/Engine/Source/Runtime/UMG/Public/Animation/MovieSceneMarginSection.h
new file mode 100644
--- /dev/null
+++ b/Engine/Source/Runtime/UMG/Public/Animation/MovieSceneMarginSection.h
@@ -0,0 +1,101 @@
+// Copyright 1998-2014 Epic Games, Inc. All Rights Reserved.
+
+#pragma once
+
+#include "MovieSceneSection.h"
+
+#include "MovieSceneMarginSection.generated.h"
+
+/**
+ * A section in a Margin track
+ */
+UCLASS(MinimalAPI)
+class UMovieSceneMarginSection : public UMovieSceneSection
+{
+	GENERATED_UCLASS_BODY()
+public:
+	/** UMovieSceneSection interface */
+	virtual void MoveSection( float DeltaPosition ) override;
+	virtual void DilateSection( float DilationFactor, float Origin ) override;
+
+	/**
+	 * Updates this section
+	 *
+	 * @param Position	The position in time within the movie scene
+	 */
+	virtual FMargin Eval( float Position ) const;
+
+	/** 
+	 * Adds a key to the section
+	 *
+	 * @param Time	The location in time where the key should be added
+	 * @param Value	The value of the key
+	 */
+	void AddKey( float Time, const FMargin& Value );
+	
+	/** 
+	 * Determines if a new key would be new data, or just a duplicate of existing data
+	 *
+	 * @param Time	The location in time where the key would be added
+	 * @param Value	The value of the new key
+	 * @return True if the new key would be new data, false if duplicate
+	 */
+	bool NewKeyIsNewData(float Time, const FMargin& Value) const;
+
+	/**
+	 * Gets the top curve
+	 *
+	 * @return The rich curve for this margin channel
+	 */
+	FRichCurve& GetTopCurve() { return TopCurve; }
+	const FRichCurve& GetTopCurve() const { return TopCurve; }
+
+	/**
+	 * Gets the left curve
+	 *
+	 * @return The rich curve for this margin channel
+	 */
+	FRichCurve& GetLeftCurve() { return LeftCurve; }
+	const FRichCurve& GetLeftCurve() const { return LeftCurve; }
+	/**
+	 * Gets the right curve
+	 *
+	 * @return The rich curve for this margin channel
+	 */
+	FRichCurve& GetRightCurve() { return RightCurve; }
+	const FRichCurve& GetRightCurve() const { return RightCurve; }
+	
+	/**
+	 * Gets the bottom curve
+	 *
+	 * @return The rich curve for this margin channel
+	 */
+	FRichCurve& GetBottomCurve() { return BottomCurve; }
+	const FRichCurve& GetBottomCurve() const { return BottomCurve; }
+
+private:
+	/**
+	 * Adds a key to a rich curve, finding an existing key to modify or adding a new one
+	 *
+	 * @param InCurve	The curve to add keys to
+	 * @param Time		The time where the key should be added
+	 * @param Value		The value at the given time
+	 */
+	void AddKeyToCurve( FRichCurve& InCurve, float Time, float Value );
+private:
+	/** Red curve data */
+	UPROPERTY()
+	FRichCurve TopCurve;
+
+	/** Green curve data */
+	UPROPERTY()
+	FRichCurve LeftCurve;
+
+	/** Blue curve data */
+	UPROPERTY()
+	FRichCurve RightCurve;
+
+	/** Alpha curve data */
+	UPROPERTY()
+	FRichCurve BottomCurve;
+};
diff --git a/Engine/Source/Runtime/UMG/Public/Animation/MovieSceneMarginTrack.h b/Engine/Source/Runtime/UMG/Public/Animation/MovieSceneMarginTrack.h
new file mode 100644
--- /dev/null
+++ b/Engine/Source/Runtime/UMG/Public/Animation/MovieSceneMarginTrack.h
@@ -0,0 +1,39 @@
+// Copyright 1998-2014 Epic Games, Inc. All Rights Reserved.
+
+#pragma once
+
+#include "MovieScenePropertyTrack.h"
+#include "MovieSceneMarginTrack.generated.h"
+
+/**
+ * Handles manipulation of FMargins in a movie scene
+ */
+UCLASS( MinimalAPI )
+class UMovieSceneMarginTrack : public UMovieScenePropertyTrack
+{
+	GENERATED_UCLASS_BODY()
+
+public:
+	/** UMovieSceneTrack interface */
+	virtual UMovieSceneSection* CreateNewSection() override;
+	virtual TSharedPtr<IMovieSceneTrackInstance> CreateInstance() override;
+
+	/**
+	 * Adds a key to a section.  Will create the section if it doesn't exist
+	 *
+	 * @param Time				The time relative to the owning movie scene where the section should be
+	 * @param Value				The value of the key
+	 * @return True if the key was successfully added.
+	 */
+	virtual bool AddKeyToSection( float Time, FMargin Value );
+	
+	/**
+	 * Evaluates the track at the playback position
+	 *
+	 * @param Position	The current playback position
+	 * @param LastPosition	The last plackback position
+	 * @param The color at the playback position
+	 * @return true if anything was evaluated. Note: if false is returned OutMargin remains unchanged
+	 */
+	virtual bool Eval( float Position, float LastPostion, FMargin& OutMargin ) const;
+};
diff --git a/Engine/Source/Runtime/UMG/Public/Blueprint/UserWidget.h b/Engine/Source/Runtime/UMG/Public/Blueprint/UserWidget.h
--- a/Engine/Source/Runtime/UMG/Public/Blueprint/UserWidget.h
+++ b/Engine/Source/Runtime/UMG/Public/Blueprint/UserWidget.h
@@ -95,71 +95,7 @@
 class UMG_API UUserWidget : public UWidget
 {
 	GENERATED_UCLASS_BODY()
-
-	/** Called when the visibility changes. */
-	UPROPERTY(BlueprintAssignable)
-	FOnVisibilityChangedEvent OnVisibilityChanged;
-
-	/** Controls whether the cursor is automatically visible when this widget is visible. */
-	UPROPERTY(EditDefaultsOnly, Category=Behavior)
-	uint32 bShowCursorWhenVisible : 1;
-
-	UPROPERTY(EditDefaultsOnly, Category=Behavior)
-	uint32 bModal : 1;
-
-	/**  */
-	UPROPERTY(EditDefaultsOnly, Category=Appearance)
-	bool bAbsoluteLayout;
-
-	UPROPERTY(EditAnywhere, Category=Appearance, meta=( EditCondition="!bAbsoluteLayout" ))
-	FMargin Padding;
-
-	/** How much space this slot should occupy in the direction of the panel. */
-	UPROPERTY(EditAnywhere, Category=Appearance, meta=( EditCondition="!bAbsoluteLayout" ))
-	FSlateChildSize Size;
-
-	/** Position. */
-	UPROPERTY(EditDefaultsOnly, Category=Appearance, meta=( EditCondition="bAbsoluteLayout" ))
-	FVector2D AbsolutePosition;
-
-	/** Size. */
-	UPROPERTY(EditDefaultsOnly, Category=Appearance, meta=( EditCondition="bAbsoluteLayout" ))
-	FVector2D AbsoluteSize;
-
-	/**
-	* Horizontal pivot position
-	*  Given a top aligned slot, where '+' represents the
-	*  anchor point defined by PositionAttr.
-	*
-	*   Left				Center				Right
-	+ _ _ _ _            _ _ + _ _          _ _ _ _ +
-	|		  |		   | 		   |	  |		    |
-	| _ _ _ _ |        | _ _ _ _ _ |	  | _ _ _ _ |
-	*
-	*  Note: FILL is NOT supported in absolute layout
-	*/
-	UPROPERTY(EditAnywhere, Category=Appearance)
-	TEnumAsByte<EHorizontalAlignment> HorizontalAlignment;
-
-	/**
-	* Vertical pivot position
-	*   Given a left aligned slot, where '+' represents the
-	*   anchor point defined by PositionAttr.
-	*
-	*   Top					Center			  Bottom
-	*	+_ _ _ _ _		 _ _ _ _ _		 _ _ _ _ _
-	*	|         |		| 		  |		|		  |
-	*	|         |     +		  |		|		  |
-	*	| _ _ _ _ |		| _ _ _ _ |		+ _ _ _ _ |
-	*
-	*  Note: FILL is NOT supported in absolute layout
-	*/
-	UPROPERTY(EditAnywhere, Category=Appearance)
-	TEnumAsByte<EVerticalAlignment> VerticalAlignment;
-
-	UPROPERTY()
-	TArray<UWidget*> Components;
-	
+public:
 	//UObject interface
 	virtual void PostInitProperties() override;
 	virtual class UWorld* GetWorld() const override;
@@ -261,6 +197,74 @@
 	/** @returns The uobject widget corresponding to a given name */
 	UWidget* GetHandleFromName(const FString& Name) const;
 
+public:
+	/** Called when the visibility changes. */
+	UPROPERTY(BlueprintAssignable)
+	FOnVisibilityChangedEvent OnVisibilityChanged;
+
+	/** Controls whether the cursor is automatically visible when this widget is visible. */
+	UPROPERTY(EditDefaultsOnly, Category=Behavior)
+	uint32 bShowCursorWhenVisible : 1;
+
+	UPROPERTY(EditDefaultsOnly, Category=Behavior)
+	uint32 bModal : 1;
+
+	/**  */
+	UPROPERTY(EditDefaultsOnly, Category=Appearance)
+	bool bAbsoluteLayout;
+
+	UPROPERTY(EditAnywhere, Category=Appearance, meta=( EditCondition="!bAbsoluteLayout" ))
+	FMargin Padding;
+
+	/** How much space this slot should occupy in the direction of the panel. */
+	UPROPERTY(EditAnywhere, Category=Appearance, meta=( EditCondition="!bAbsoluteLayout" ))
+	FSlateChildSize Size;
+
+	/** Position. */
+	UPROPERTY(EditDefaultsOnly, Category=Appearance, meta=( EditCondition="bAbsoluteLayout" ))
+	FVector2D AbsolutePosition;
+
+	/** Size. */
+	UPROPERTY(EditDefaultsOnly, Category=Appearance, meta=( EditCondition="bAbsoluteLayout" ))
+	FVector2D AbsoluteSize;
+
+	/**
+	* Horizontal pivot position
+	*  Given a top aligned slot, where '+' represents the
+	*  anchor point defined by PositionAttr.
+	*
+	*   Left				Center				Right
+	+ _ _ _ _            _ _ + _ _          _ _ _ _ +
+	|		  |		   | 		   |	  |		    |
+	| _ _ _ _ |        | _ _ _ _ _ |	  | _ _ _ _ |
+	*
+	*  Note: FILL is NOT supported in absolute layout
+	*/
+	UPROPERTY(EditAnywhere, Category=Appearance)
+	TEnumAsByte<EHorizontalAlignment> HorizontalAlignment;
+
+	/**
+	* Vertical pivot position
+	*   Given a left aligned slot, where '+' represents the
+	*   anchor point defined by PositionAttr.
+	*
+	*   Top					Center			  Bottom
+	*	+_ _ _ _ _		 _ _ _ _ _		 _ _ _ _ _
+	*	|         |		| 		  |		|		  |
+	*	|         |     +		  |		|		  |
+	*	| _ _ _ _ |		| _ _ _ _ |		+ _ _ _ _ |
+	*
+	*  Note: FILL is NOT supported in absolute layout
+	*/
+	UPROPERTY(EditAnywhere, Category=Appearance)
+	TEnumAsByte<EVerticalAlignment> VerticalAlignment;
+
+	UPROPERTY()
+	TArray<UWidget*> Components;
+
+	UPROPERTY()
+	class UWidgetTree* WidgetTree;
+
 protected:
 	virtual TSharedRef<SWidget> RebuildWidget() override;
 
