UMG - Now syncing selection from tree to designer.

diff --git a/Engine/Source/Editor/UMGEditor/Private/SUMGDesigner.cpp b/Engine/Source/Editor/UMGEditor/Private/SUMGDesigner.cpp
--- a/Engine/Source/Editor/UMGEditor/Private/SUMGDesigner.cpp
+++ b/Engine/Source/Editor/UMGEditor/Private/SUMGDesigner.cpp
@@ -119,6 +119,8 @@ void SUMGDesigner::Construct(const FArguments& InArgs, TSharedPtr<FWidgetBluepri
 			]
 		]
 	);
+
+	BlueprintEditor.Pin()->OnSelectedWidgetsChanged.AddRaw(this, &SUMGDesigner::OnEditorSelectionChanged);
 }
 
 SUMGDesigner::~SUMGDesigner()
@@ -128,6 +130,29 @@ SUMGDesigner::~SUMGDesigner()
 	{
 		Blueprint->OnChanged().RemoveAll(this);
 	}
+
+	if ( BlueprintEditor.IsValid() )
+	{
+		BlueprintEditor.Pin()->OnSelectedWidgetsChanged.RemoveAll(this);
+	}
+}
+
+void SUMGDesigner::OnEditorSelectionChanged()
+{
+	SelectedWidgets = BlueprintEditor.Pin()->GetSelectedWidgets();
+
+	if ( SelectedWidgets.Num() > 0 )
+	{
+		for ( FWidgetReference& Widget : SelectedWidgets )
+		{
+			SelectedWidget = Widget;
+			break;
+		}
+	}
+	else
+	{
+		SelectedWidget = FWidgetReference();
+	}
 }
 
 UWidgetBlueprint* SUMGDesigner::GetBlueprint() const
diff --git a/Engine/Source/Editor/UMGEditor/Private/SUMGDesigner.h b/Engine/Source/Editor/UMGEditor/Private/SUMGDesigner.h
--- a/Engine/Source/Editor/UMGEditor/Private/SUMGDesigner.h
+++ b/Engine/Source/Editor/UMGEditor/Private/SUMGDesigner.h
@@ -43,6 +43,8 @@ public:
 	void Register(TSharedRef<FDesignerExtension> Extension);
 
 private:
+	void OnEditorSelectionChanged();
+
 	/** @returns Gets the widget under the cursor based on a mouse pointer event. */
 	FWidgetReference GetWidgetAtCursor(const FGeometry& MyGeometry, const FPointerEvent& MouseEvent, FArrangedWidget& ArrangedWidget);
 
@@ -122,9 +124,11 @@ private:
 	FVector2D CachedDesignerWidgetLocation;
 
 	/** The currently selected widget */
+	TSet< FWidgetReference > SelectedWidgets;
 	FWidgetReference SelectedWidget;
 
 	/** The currently selected slate widget, this is refreshed every frame in case it changes */
+	TArray< TWeakPtr<SWidget> > SelectedSlateWidgets;
 	TWeakPtr<SWidget> SelectedSlateWidget;
 
 	/** The wall clock time the user has been hovering over a single widget */
