UMG - Can now use generic commands.  Adding copy, paste, delete (generic commands).  Adding a base visual that slots and widgets inherit from, more functionality will likely be moved to it over time.  For now it just makes it simpler to write a lot of the sequencer code.

diff --git a/Engine/Source/Editor/UMGEditor/Private/SDesignSurface.h b/Engine/Source/Editor/UMGEditor/Private/SDesignSurface.h
--- a/Engine/Source/Editor/UMGEditor/Private/SDesignSurface.h
+++ b/Engine/Source/Editor/UMGEditor/Private/SDesignSurface.h
@@ -19,6 +19,7 @@ public:
 	// SWidget interface
 	virtual FReply OnMouseWheel(const FGeometry& MyGeometry, const FPointerEvent& MouseEvent) override;
 	virtual int32 OnPaint(const FGeometry& AllottedGeometry, const FSlateRect& MyClippingRect, FSlateWindowElementList& OutDrawElements, int32 LayerId, const FWidgetStyle& InWidgetStyle, bool bParentEnabled) const override;
+	virtual bool SupportsKeyboardFocus() const override { return true; }
 	// End of Swidget interface
 
 protected:
diff --git a/Engine/Source/Editor/UMGEditor/Private/SUMGDesigner.cpp b/Engine/Source/Editor/UMGEditor/Private/SUMGDesigner.cpp
--- a/Engine/Source/Editor/UMGEditor/Private/SUMGDesigner.cpp
+++ b/Engine/Source/Editor/UMGEditor/Private/SUMGDesigner.cpp
@@ -239,71 +239,78 @@ FReply SUMGDesigner::OnMouseButtonDown(const FGeometry& MyGeometry, const FPoint
 {
 	CurrentHandle = HitTestDragHandles(MyGeometry, MouseEvent);
 
-	if ( CurrentHandle == DH_NONE )
+	if ( CurrentHandle != DH_NONE )
 	{
-		FArrangedWidget ArrangedWidget(SNullWidget::NullWidget, FGeometry());
-		SelectedWidget = GetWidgetAtCursor(MyGeometry, MouseEvent, ArrangedWidget);
+		BeginTransaction(LOCTEXT("ResizeWidget", "Resize Widget"));
+		return FReply::Handled().PreventThrottling().CaptureMouse(AsShared());
+	}
 
-		//TODO UMG Undoable Selection
+	//TODO UMG Undoable Selection
+	FArrangedWidget ArrangedWidget(SNullWidget::NullWidget, FGeometry());
+	SelectedWidget = GetWidgetAtCursor(MyGeometry, MouseEvent, ArrangedWidget);
+	SelectedWidgetContextMenuLocation = ArrangedWidget.Geometry.AbsoluteToLocal(MouseEvent.GetScreenSpacePosition());
 
-		if ( SelectedWidget.IsValid() )
-		{
-			//@TODO UMG primary FBlueprintEditor needs to be inherited and selection control needs to be centralized.
-			// Set the template as selected in the details panel
-			TSet<FWidgetReference> SelectedTemplates;
-			SelectedTemplates.Add(SelectedWidget);
-			BlueprintEditor.Pin()->SelectWidgets(SelectedTemplates);
-
-			// Remove all the current extension widgets
-			ExtensionWidgetCanvas->ClearChildren();
+	if ( SelectedWidget.IsValid() )
+	{
+		//@TODO UMG primary FBlueprintEditor needs to be inherited and selection control needs to be centralized.
+		// Set the template as selected in the details panel
+		TSet<FWidgetReference> SelectedTemplates;
+		SelectedTemplates.Add(SelectedWidget);
+		BlueprintEditor.Pin()->SelectWidgets(SelectedTemplates);
 
-			ExtensionWidgets.Reset();
+		// Remove all the current extension widgets
+		ExtensionWidgetCanvas->ClearChildren();
 
-			TArray<FWidgetReference> Selected;
-			Selected.Add(SelectedWidget);
+		ExtensionWidgets.Reset();
 
-			// Build extension widgets for new selection
-			for ( TSharedRef<FDesignerExtension>& Ext : DesignerExtensions )
-			{
-				Ext->BuildWidgetsForSelection(Selected, ExtensionWidgets);
-			}
+		TArray<FWidgetReference> Selected;
+		Selected.Add(SelectedWidget);
 
-			TSharedRef<SHorizontalBox> ExtensionBox = SNew(SHorizontalBox);
+		// Build extension widgets for new selection
+		for ( TSharedRef<FDesignerExtension>& Ext : DesignerExtensions )
+		{
+			Ext->BuildWidgetsForSelection(Selected, ExtensionWidgets);
+		}
 
-			// Add Widgets to designer surface
-			for ( TSharedRef<SWidget>& ExWidget : ExtensionWidgets )
-			{
-				ExtensionBox->AddSlot()
-					.AutoWidth()
-					[
-						ExWidget
-					];
-			}
+		TSharedRef<SHorizontalBox> ExtensionBox = SNew(SHorizontalBox);
 
-			ExtensionWidgetCanvas->AddSlot()
-				.Position(TAttribute<FVector2D>(this, &SUMGDesigner::GetCachedSelectionDesignerWidgetsLocation))
-				.Size(FVector2D(100, 25))
+		// Add Widgets to designer surface
+		for ( TSharedRef<SWidget>& ExWidget : ExtensionWidgets )
+		{
+			ExtensionBox->AddSlot()
+				.AutoWidth()
 				[
-					ExtensionBox
+					ExWidget
 				];
 		}
-	}
-	else
-	{
-		BeginTransaction(LOCTEXT("ResizeWidget", "Resize Widget"));
 
-		// Capture mouse for the drag handle
-		return FReply::Handled().PreventThrottling().CaptureMouse(AsShared());
+		ExtensionWidgetCanvas->AddSlot()
+			.Position(TAttribute<FVector2D>(this, &SUMGDesigner::GetCachedSelectionDesignerWidgetsLocation))
+			.Size(FVector2D(100, 25))
+			[
+				ExtensionBox
+			];
 	}
 
-	return FReply::Handled();
+	// Capture mouse for the drag handle and general mouse actions
+	return FReply::Handled().PreventThrottling().CaptureMouse(AsShared());
 }
 
 FReply SUMGDesigner::OnMouseButtonUp(const FGeometry& MyGeometry, const FPointerEvent& MouseEvent)
 {
-	EndTransaction();
+	if ( MouseEvent.GetEffectingButton() == EKeys::LeftMouseButton )
+	{
+		if ( CurrentHandle != DH_NONE )
+		{
+			EndTransaction();
+			CurrentHandle = DH_NONE;
+		}
+	}
+	else if ( MouseEvent.GetEffectingButton() == EKeys::RightMouseButton )
+	{
+		ShowContextMenu(MyGeometry, MouseEvent);
+	}
 
-	CurrentHandle = DH_NONE;
 	return FReply::Handled().ReleaseMouseCapture();
 }
 
@@ -368,6 +375,33 @@ void SUMGDesigner::OnMouseLeave(const FPointerEvent& MouseEvent)
 	HoverTime = 0;
 }
 
+FReply SUMGDesigner::OnKeyDown(const FGeometry& MyGeometry, const FKeyboardEvent& InKeyboardEvent)
+{
+	BlueprintEditor.Pin()->PasteDropLocation = FVector2D(0, 0);
+
+	if ( BlueprintEditor.Pin()->WidgetCommandList->ProcessCommandBindings(InKeyboardEvent) )
+	{
+		return FReply::Handled();
+	}
+
+	return FReply::Unhandled();
+}
+
+void SUMGDesigner::ShowContextMenu(const FGeometry& MyGeometry, const FPointerEvent& MouseEvent)
+{
+	FMenuBuilder MenuBuilder(true, NULL);
+
+	FWidgetBlueprintEditorUtils::CreateWidgetContextMenu(MenuBuilder, BlueprintEditor.Pin().ToSharedRef(), SelectedWidgetContextMenuLocation);
+
+	TSharedPtr<SWidget> MenuContent = MenuBuilder.MakeWidget();
+
+	if ( MenuContent.IsValid() )
+	{
+		FVector2D SummonLocation = MouseEvent.GetScreenSpacePosition();
+		FSlateApplication::Get().PushMenu(AsShared(), MenuContent.ToSharedRef(), SummonLocation, FPopupTransitionEffect(FPopupTransitionEffect::ContextMenu));
+	}
+}
+
 bool SUMGDesigner::GetArrangedWidget(TSharedRef<SWidget> Widget, FArrangedWidget& ArrangedWidget) const
 {
 	// We can't screenshot the widget unless there's a valid window handle to draw it in.
diff --git a/Engine/Source/Editor/UMGEditor/Private/SUMGDesigner.h b/Engine/Source/Editor/UMGEditor/Private/SUMGDesigner.h
--- a/Engine/Source/Editor/UMGEditor/Private/SUMGDesigner.h
+++ b/Engine/Source/Editor/UMGEditor/Private/SUMGDesigner.h
@@ -29,6 +29,8 @@ public:
 	virtual void OnMouseEnter(const FGeometry& MyGeometry, const FPointerEvent& MouseEvent) override;
 	virtual void OnMouseLeave(const FPointerEvent& MouseEvent) override;
 
+	virtual FReply OnKeyDown(const FGeometry& MyGeometry, const FKeyboardEvent& InKeyboardEvent) override;
+
 	virtual FCursorReply OnCursorQuery(const FGeometry& MyGeometry, const FPointerEvent& CursorEvent) const override;
 
 	virtual int32 OnPaint(const FGeometry& AllottedGeometry, const FSlateRect& MyClippingRect, FSlateWindowElementList& OutDrawElements, int32 LayerId, const FWidgetStyle& InWidgetStyle, bool bParentEnabled) const override;
@@ -43,6 +45,9 @@ public:
 	void Register(TSharedRef<FDesignerExtension> Extension);
 
 private:
+	/** Displays the context menu when you right click */
+	void ShowContextMenu(const FGeometry& MyGeometry, const FPointerEvent& MouseEvent);
+
 	void OnEditorSelectionChanged();
 
 	/** @returns Gets the widget under the cursor based on a mouse pointer event. */
@@ -123,10 +128,15 @@ private:
 
 	FVector2D CachedDesignerWidgetLocation;
 
-	/** The currently selected widget */
+	/** The currently selected set of widgets */
 	TSet< FWidgetReference > SelectedWidgets;
+
+	/** TODO UMG Remove, after getting multiselection working. */
 	FWidgetReference SelectedWidget;
 
+	/** The location in selected widget local space where the context menu was summoned. */
+	FVector2D SelectedWidgetContextMenuLocation;
+
 	/** The currently selected slate widget, this is refreshed every frame in case it changes */
 	TArray< TWeakPtr<SWidget> > SelectedSlateWidgets;
 	TWeakPtr<SWidget> SelectedSlateWidget;
diff --git a/Engine/Source/Editor/UMGEditor/Private/SUMGEditorTree.cpp b/Engine/Source/Editor/UMGEditor/Private/SUMGEditorTree.cpp
--- a/Engine/Source/Editor/UMGEditor/Private/SUMGEditorTree.cpp
+++ b/Engine/Source/Editor/UMGEditor/Private/SUMGEditorTree.cpp
@@ -93,6 +93,18 @@ void SUMGEditorTree::Tick(const FGeometry& AllottedGeometry, const double InCurr
 	}
 }
 
+FReply SUMGEditorTree::OnKeyDown(const FGeometry& MyGeometry, const FKeyboardEvent& InKeyboardEvent)
+{
+	BlueprintEditor.Pin()->PasteDropLocation = FVector2D(0, 0);
+
+	if ( BlueprintEditor.Pin()->WidgetCommandList->ProcessCommandBindings(InKeyboardEvent) )
+	{
+		return FReply::Handled();
+	}
+
+	return FReply::Unhandled();
+}
+
 void SUMGEditorTree::TransformWidgetToString(const UWidget* Widget, OUT TArray< FString >& Array)
 {
 	Array.Add(Widget->GetName());
@@ -179,55 +191,11 @@ void SUMGEditorTree::ShowDetailsForObjects(TArray<UWidget*> TemplateWidgets)
 	BlueprintEditor.Pin()->SelectWidgets(SelectedWidgets);
 }
 
-void SUMGEditorTree::BuildWrapWithMenu(FMenuBuilder& Menu)
-{
-	Menu.BeginSection("WrapWith", LOCTEXT("WidgetTree_WrapWith", "Wrap With..."));
-	{
-		for ( TObjectIterator<UClass> ClassIt; ClassIt; ++ClassIt )
-		{
-			UClass* WidgetClass = *ClassIt;
-			if ( WidgetClass->IsChildOf(UPanelWidget::StaticClass()) && WidgetClass->HasAnyClassFlags(CLASS_Abstract) == false )
-			{
-				Menu.AddMenuEntry(
-					WidgetClass->GetDisplayNameText(),
-					FText::GetEmpty(),
-					FSlateIcon(),
-					FUIAction(
-						FExecuteAction::CreateSP(this, &SUMGEditorTree::WrapSelectedWidgets, WidgetClass),
-						FCanExecuteAction()
-					)
-				);
-			}
-		}
-	}
-	Menu.EndSection();
-}
-
 TSharedPtr<SWidget> SUMGEditorTree::WidgetHierarchy_OnContextMenuOpening()
 {
 	FMenuBuilder MenuBuilder(true, NULL);
 
-	MenuBuilder.BeginSection("Actions");
-	{
-		const TArray<UWidget*> SelectionList = WidgetTreeView->GetSelectedItems();
-
-		MenuBuilder.AddMenuEntry(
-			LOCTEXT("WidgetTree_Delete", "Delete"),
-			LOCTEXT("WidgetTree_DeleteToolTip", "Deletes the current widget"),
-			FSlateIcon(),
-			FUIAction(
-			FExecuteAction::CreateSP(this, &SUMGEditorTree::DeleteSelected),
-			FCanExecuteAction()
-			)
-		);
-
-		MenuBuilder.AddSubMenu(
-			LOCTEXT("WidgetTree_WrapWith", "Wrap With..."),
-			LOCTEXT("WidgetTree_WrapWithToolTip", "Wraps the currently selected widgets inside of another container widget"),
-			FNewMenuDelegate::CreateSP(this, &SUMGEditorTree::BuildWrapWithMenu)
-		);
-	}
-	MenuBuilder.EndSection();
+	FWidgetBlueprintEditorUtils::CreateWidgetContextMenu(MenuBuilder, BlueprintEditor.Pin().ToSharedRef(), FVector2D(0, 0));
 
 	return MenuBuilder.MakeWidget();
 }
@@ -269,62 +237,22 @@ void SUMGEditorTree::WidgetHierarchy_OnSelectionChanged(UWidget* SelectedItem, E
 	}
 }
 
-void SUMGEditorTree::WrapSelectedWidgets(UClass* WidgetClass)
+FReply SUMGEditorTree::HandleDeleteSelected()
 {
-	const TArray<UWidget*> SelectionList = WidgetTreeView->GetSelectedItems();
-
-	TSharedPtr<FWidgetTemplateClass> Template = MakeShareable(new FWidgetTemplateClass(WidgetClass));
+	TSet<FWidgetReference> SelectedWidgets = BlueprintEditor.Pin()->GetSelectedWidgets();
+	//TArray<UWidget*> SelectedWidgets = WidgetTreeView->GetSelectedItems();
 
-	UWidgetBlueprint* BP = GetBlueprint();
-	UPanelWidget* NewWrapperWidget = CastChecked<UPanelWidget>( Template->Create(BP->WidgetTree) );
-
-	int32 OutIndex;
-	UPanelWidget* CurrentParent = BP->WidgetTree->FindWidgetParent(SelectionList[0], OutIndex);
-	if ( CurrentParent )
+	// Remove the selected items from the widget cache
+	for ( FWidgetReference& Item : SelectedWidgets )
 	{
-		CurrentParent->ReplaceChildAt(OutIndex, NewWrapperWidget);
-
-		NewWrapperWidget->AddChild(SelectionList[0], FVector2D());
+		CachedExpandedWidgets.Remove(Item.GetTemplate());
 	}
 
-	FBlueprintEditorUtils::MarkBlueprintAsStructurallyModified(BP);
-}
-
-FReply SUMGEditorTree::HandleDeleteSelected()
-{
-	DeleteSelected();
+	FWidgetBlueprintEditorUtils::DeleteWidgets(GetBlueprint(), SelectedWidgets);
 
 	return FReply::Handled();
 }
 
-void SUMGEditorTree::DeleteSelected()
-{
-	TArray<UWidget*> SelectedItems = WidgetTreeView->GetSelectedItems();
-	if ( SelectedItems.Num() > 0 )
-	{
-		UWidgetBlueprint* BP = GetBlueprint();
-
-		const FScopedTransaction Transaction(LOCTEXT("RemoveWidget", "Remove Widget"));
-		BP->WidgetTree->SetFlags(RF_Transactional);
-		BP->WidgetTree->Modify();
-
-		bool bRemoved = false;
-		for ( UWidget* Item : SelectedItems )
-		{
-			bRemoved = BP->WidgetTree->RemoveWidget(Item);
-
-			CachedExpandedWidgets.Remove(Item);
-		}
-
-		//TODO UMG There needs to be an event for widget removal so that caches can be updated, and selection
-
-		if ( bRemoved )
-		{
-			FBlueprintEditorUtils::MarkBlueprintAsStructurallyModified(BP);
-		}
-	}
-}
-
 bool SUMGEditorTree::FilterWidgetHierarchy(UWidget* CurrentWidget)
 {
 	bool bAnyChildrenPass = false;
diff --git a/Engine/Source/Editor/UMGEditor/Private/SUMGEditorTree.h b/Engine/Source/Editor/UMGEditor/Private/SUMGEditorTree.h
--- a/Engine/Source/Editor/UMGEditor/Private/SUMGEditorTree.h
+++ b/Engine/Source/Editor/UMGEditor/Private/SUMGEditorTree.h
@@ -26,6 +26,8 @@ public:
 
 	// Begin SWidget
 	virtual void Tick(const FGeometry& AllottedGeometry, const double InCurrentTime, const float InDeltaTime);
+
+	virtual FReply OnKeyDown(const FGeometry& MyGeometry, const FKeyboardEvent& InKeyboardEvent) override;
 	// End SWidget
 
 private:
diff --git a/Engine/Source/Editor/UMGEditor/Private/WidgetBlueprintEditor.cpp b/Engine/Source/Editor/UMGEditor/Private/WidgetBlueprintEditor.cpp
--- a/Engine/Source/Editor/UMGEditor/Private/WidgetBlueprintEditor.cpp
+++ b/Engine/Source/Editor/UMGEditor/Private/WidgetBlueprintEditor.cpp
@@ -85,6 +85,73 @@ void FWidgetBlueprintEditor::InitWidgetBlueprintEditor(const EToolkitMode::Type
 	}
 
 	UpdatePreview(GetWidgetBlueprintObj(), true);
+
+	WidgetCommandList = MakeShareable(new FUICommandList);
+
+	WidgetCommandList->MapAction(FGenericCommands::Get().Delete,
+		FExecuteAction::CreateSP(this, &FWidgetBlueprintEditor::DeleteSelectedWidgets),
+		FCanExecuteAction::CreateSP(this, &FWidgetBlueprintEditor::CanDeleteSelectedWidgets)
+		);
+
+	WidgetCommandList->MapAction(FGenericCommands::Get().Copy,
+		FExecuteAction::CreateSP(this, &FWidgetBlueprintEditor::CopySelectedWidgets),
+		FCanExecuteAction::CreateSP(this, &FWidgetBlueprintEditor::CanCopySelectedWidgets)
+		);
+
+	//WidgetCommandList->MapAction(FGenericCommands::Get().Cut,
+	//	FExecuteAction::CreateSP(this, &FWidgetBlueprintEditor::CutSelectedNodes),
+	//	FCanExecuteAction::CreateSP(this, &FWidgetBlueprintEditor::CanCutNodes)
+	//	);
+
+	WidgetCommandList->MapAction(FGenericCommands::Get().Paste,
+		FExecuteAction::CreateSP(this, &FWidgetBlueprintEditor::PasteWidgets),
+		FCanExecuteAction::CreateSP(this, &FWidgetBlueprintEditor::CanPasteWidgets)
+		);
+}
+
+bool FWidgetBlueprintEditor::CanDeleteSelectedWidgets()
+{
+	TSet<FWidgetReference> Widgets = GetSelectedWidgets();
+	return Widgets.Num() > 0;
+}
+
+void FWidgetBlueprintEditor::DeleteSelectedWidgets()
+{
+	TSet<FWidgetReference> Widgets = GetSelectedWidgets();
+	FWidgetBlueprintEditorUtils::DeleteWidgets(GetWidgetBlueprintObj(), Widgets);
+}
+
+bool FWidgetBlueprintEditor::CanCopySelectedWidgets()
+{
+	TSet<FWidgetReference> Widgets = GetSelectedWidgets();
+	return Widgets.Num() > 0;
+}
+
+void FWidgetBlueprintEditor::CopySelectedWidgets()
+{
+	TSet<FWidgetReference> Widgets = GetSelectedWidgets();
+	FWidgetBlueprintEditorUtils::CopyWidgets(GetWidgetBlueprintObj(), Widgets);
+}
+
+bool FWidgetBlueprintEditor::CanPasteWidgets()
+{
+	TSet<FWidgetReference> Widgets = GetSelectedWidgets();
+	if ( Widgets.Num() == 1 )
+	{
+		FWidgetReference Target = *Widgets.CreateIterator();
+		const bool bIsPanel = Cast<UPanelWidget>(Target.GetTemplate()) != NULL;
+		return bIsPanel;
+	}
+
+	return false;
+}
+
+void FWidgetBlueprintEditor::PasteWidgets()
+{
+	TSet<FWidgetReference> Widgets = GetSelectedWidgets();
+	FWidgetReference Target = *Widgets.CreateIterator();
+
+	FWidgetBlueprintEditorUtils::PasteWidgets(GetWidgetBlueprintObj(), Target, PasteDropLocation);
 }
 
 void FWidgetBlueprintEditor::Tick(float DeltaTime)
diff --git a/Engine/Source/Editor/UMGEditor/Private/WidgetBlueprintEditor.h b/Engine/Source/Editor/UMGEditor/Private/WidgetBlueprintEditor.h
--- a/Engine/Source/Editor/UMGEditor/Private/WidgetBlueprintEditor.h
+++ b/Engine/Source/Editor/UMGEditor/Private/WidgetBlueprintEditor.h
@@ -50,9 +50,26 @@ public:
 	/** @return Notification for when the preview widget has been updated */
 	FOnWidgetPreviewUpdated& GetOnWidgetPreviewUpdated() { return OnWidgetPreviewUpdated; }
 
+public:
 	/** Fires whenever the selected set of widgets changes */
 	FOnSelectedWidgetsChanged OnSelectedWidgetsChanged;
 
+	/** Command list for handling widget actions in the WidgetBlueprintEditor */
+	TSharedPtr< FUICommandList > WidgetCommandList;
+
+	/** Paste Metadata */
+	FVector2D PasteDropLocation;
+
+private:
+	bool CanDeleteSelectedWidgets();
+	void DeleteSelectedWidgets();
+
+	bool CanCopySelectedWidgets();
+	void CopySelectedWidgets();
+
+	bool CanPasteWidgets();
+	void PasteWidgets();
+
 private:
 	/** Updates the inspector to be viewing the currently selected set of widgets */
 	void RefreshDetails();
diff --git a/Engine/Source/Editor/UMGEditor/Private/WidgetBlueprintEditorUtils.cpp b/Engine/Source/Editor/UMGEditor/Private/WidgetBlueprintEditorUtils.cpp
--- a/Engine/Source/Editor/UMGEditor/Private/WidgetBlueprintEditorUtils.cpp
+++ b/Engine/Source/Editor/UMGEditor/Private/WidgetBlueprintEditorUtils.cpp
@@ -7,6 +7,59 @@
 #include "Kismet2NameValidators.h"
 #include "BlueprintEditorUtils.h"
 #include "K2Node_Variable.h"
+#include "WidgetTemplateClass.h"
+#include "Factories.h"
+#include "UnrealExporter.h"
+
+#define LOCTEXT_NAMESPACE "UMG"
+
+class FWidgetObjectTextFactory : public FCustomizableTextObjectFactory
+{
+public:
+	FWidgetObjectTextFactory()
+		: FCustomizableTextObjectFactory(GWarn)
+	{
+	}
+
+	// FCustomizableTextObjectFactory implementation
+
+	virtual bool CanCreateClass(UClass* ObjectClass) const override
+	{
+		const bool bIsWidget = ObjectClass->IsChildOf(UWidget::StaticClass());
+		const bool bIsSlot = ObjectClass->IsChildOf(UPanelSlot::StaticClass());
+
+		return bIsWidget || bIsSlot;
+	}
+
+	virtual void ProcessConstructedObject(UObject* NewObject) override
+	{
+		check(NewObject);
+
+		// Add it to the new object map
+		NewObjectMap.Add(NewObject->GetFName(), Cast<UWidget>(NewObject));
+
+		// If this is a scene component and it has a parent
+		UWidget* Widget = Cast<UWidget>(NewObject);
+		if ( Widget && Widget->Slot )
+		{
+			// Add an entry to the child->parent name map
+//			ParentMap.Add(NewObject->GetFName(), Widget->AttachParent->GetFName());
+
+			// Clear this so it isn't used when constructing the new SCS node
+			//Widget->AttachParent = NULL;
+		}
+	}
+
+	// FCustomizableTextObjectFactory (end)
+
+public:
+
+	// Child->Parent name map
+	TMap<FName, FName> ParentMap;
+
+	// Name->Instance object mapping
+	TMap<FName, UWidget*> NewObjectMap;
+};
 
 bool FWidgetBlueprintEditorUtils::RenameWidget(TSharedRef<FWidgetBlueprintEditor> BlueprintEditor, const FName& OldName, const FName& NewName)
 {
@@ -98,4 +151,190 @@ bool FWidgetBlueprintEditorUtils::RenameWidget(TSharedRef<FWidgetBlueprintEditor
 	bRenamed = true;
 
 	return bRenamed;
-}
\ No newline at end of file
+}
+
+void FWidgetBlueprintEditorUtils::CreateWidgetContextMenu(FMenuBuilder& MenuBuilder, TSharedRef<FWidgetBlueprintEditor> BlueprintEditor, FVector2D TargetLocation)
+{
+	BlueprintEditor->PasteDropLocation = TargetLocation;
+
+	TSet<FWidgetReference> Widgets = BlueprintEditor->GetSelectedWidgets();
+	UWidgetBlueprint* BP = BlueprintEditor->GetWidgetBlueprintObj();
+
+	MenuBuilder.PushCommandList(BlueprintEditor->WidgetCommandList.ToSharedRef());
+
+	MenuBuilder.BeginSection("Actions");
+	{
+		MenuBuilder.AddMenuEntry(FGenericCommands::Get().Copy);
+		MenuBuilder.AddMenuEntry(FGenericCommands::Get().Paste);
+		MenuBuilder.AddMenuEntry(FGenericCommands::Get().Delete);
+
+		MenuBuilder.AddSubMenu(
+			LOCTEXT("WidgetTree_WrapWith", "Wrap With..."),
+			LOCTEXT("WidgetTree_WrapWithToolTip", "Wraps the currently selected widgets inside of another container widget"),
+			FNewMenuDelegate::CreateStatic(&FWidgetBlueprintEditorUtils::BuildWrapWithMenu, BP, Widgets)
+			);
+	}
+	MenuBuilder.EndSection();
+
+	MenuBuilder.PopCommandList();
+}
+
+void FWidgetBlueprintEditorUtils::DeleteWidgets(UWidgetBlueprint* BP, TSet<FWidgetReference> Widgets)
+{
+	if ( Widgets.Num() > 0 )
+	{
+		const FScopedTransaction Transaction(LOCTEXT("RemoveWidget", "Remove Widget"));
+		BP->WidgetTree->SetFlags(RF_Transactional);
+		BP->WidgetTree->Modify();
+
+		bool bRemoved = false;
+		for ( FWidgetReference& Item : Widgets )
+		{
+			bRemoved = BP->WidgetTree->RemoveWidget(Item.GetTemplate());
+		}
+
+		//TODO UMG There needs to be an event for widget removal so that caches can be updated, and selection
+
+		if ( bRemoved )
+		{
+			FBlueprintEditorUtils::MarkBlueprintAsStructurallyModified(BP);
+		}
+	}
+}
+
+void FWidgetBlueprintEditorUtils::BuildWrapWithMenu(FMenuBuilder& Menu, UWidgetBlueprint* BP, TSet<FWidgetReference> Widgets)
+{
+	Menu.BeginSection("WrapWith", LOCTEXT("WidgetTree_WrapWith", "Wrap With..."));
+	{
+		for ( TObjectIterator<UClass> ClassIt; ClassIt; ++ClassIt )
+		{
+			UClass* WidgetClass = *ClassIt;
+			if ( WidgetClass->IsChildOf(UPanelWidget::StaticClass()) && WidgetClass->HasAnyClassFlags(CLASS_Abstract) == false )
+			{
+				Menu.AddMenuEntry(
+					WidgetClass->GetDisplayNameText(),
+					FText::GetEmpty(),
+					FSlateIcon(),
+					FUIAction(
+					FExecuteAction::CreateStatic(&FWidgetBlueprintEditorUtils::WrapWidgets, BP, Widgets, WidgetClass),
+					FCanExecuteAction()
+					)
+					);
+			}
+		}
+	}
+	Menu.EndSection();
+}
+
+void FWidgetBlueprintEditorUtils::WrapWidgets(UWidgetBlueprint* BP, TSet<FWidgetReference> Widgets, UClass* WidgetClass)
+{
+	TSharedPtr<FWidgetTemplateClass> Template = MakeShareable(new FWidgetTemplateClass(WidgetClass));
+
+	UPanelWidget* NewWrapperWidget = CastChecked<UPanelWidget>(Template->Create(BP->WidgetTree));
+
+	//TODO UMG When wrapping multiple widgets, how will that work?
+	for ( FWidgetReference& Item : Widgets )
+	{
+		int32 OutIndex;
+		UPanelWidget* CurrentParent = BP->WidgetTree->FindWidgetParent(Item.GetTemplate(), OutIndex);
+		if ( CurrentParent )
+		{
+			CurrentParent->ReplaceChildAt(OutIndex, NewWrapperWidget);
+
+			NewWrapperWidget->AddChild(Item.GetTemplate(), FVector2D());
+		}
+	}
+
+	FBlueprintEditorUtils::MarkBlueprintAsStructurallyModified(BP);
+}
+
+void FWidgetBlueprintEditorUtils::CopyWidgets(UWidgetBlueprint* BP, TSet<FWidgetReference> Widgets)
+{
+	TSet<UWidget*> CopyableWidets;
+	for ( const FWidgetReference& Widget : Widgets )
+	{
+		UWidget* ParentWidget = Widget.GetTemplate();
+		CopyableWidets.Add(ParentWidget);
+
+		UWidget::GatherAllChildren(ParentWidget, CopyableWidets);
+	}
+
+	FString ExportedText;
+	FWidgetBlueprintEditorUtils::ExportWidgetsToText(CopyableWidets, /*out*/ ExportedText);
+	FPlatformMisc::ClipboardCopy(*ExportedText);
+}
+
+void FWidgetBlueprintEditorUtils::ExportWidgetsToText(TSet<UWidget*> WidgetsToExport, /*out*/ FString& ExportedText)
+{
+	// Clear the mark state for saving.
+	UnMarkAllObjects(EObjectMark(OBJECTMARK_TagExp | OBJECTMARK_TagImp));
+
+	FStringOutputDevice Archive;
+	const FExportObjectInnerContext Context;
+
+	// Export each of the selected nodes
+	UObject* LastOuter = NULL;
+	for ( UWidget* Widget : WidgetsToExport )
+	{
+		// The nodes should all be from the same scope
+		UObject* ThisOuter = Widget->GetOuter();
+		check(( LastOuter == ThisOuter ) || ( LastOuter == NULL ));
+		LastOuter = ThisOuter;
+
+		UExporter::ExportToOutputDevice(&Context, Widget, NULL, Archive, TEXT("copy"), 0, PPF_ExportsNotFullyQualified | PPF_Copy | PPF_Delimited, false, ThisOuter);
+	}
+
+	ExportedText = Archive;
+}
+
+void FWidgetBlueprintEditorUtils::PasteWidgets(UWidgetBlueprint* BP, FWidgetReference ParentWidgetRef, FVector2D PasteLocation)
+{
+	const FScopedTransaction Transaction(FGenericCommands::Get().Paste->GetDescription());
+
+	UPanelWidget* ParentWidget = CastChecked<UPanelWidget>(ParentWidgetRef.GetTemplate());
+
+	if ( ParentWidget )
+	{
+		ParentWidget->Modify();
+	}
+
+	// Grab the text to paste from the clipboard.
+	FString TextToImport;
+	FPlatformMisc::ClipboardPaste(TextToImport);
+
+	// Import the nodes
+	TSet<UWidget*> PastedWidgets;
+	FWidgetBlueprintEditorUtils::ImportWidgetsFromText(BP, TextToImport, /*out*/ PastedWidgets);
+
+	for ( UWidget* NewWidget : PastedWidgets )
+	{
+		// Widgets with a null parent mean that they were the root most widget of their selection set when
+		// they were copied and thus we need to paste only the root most widgets.  All their children will be added
+		// automatically.
+		if ( NewWidget->GetParent() == NULL )
+		{
+			ParentWidget->AddChild(NewWidget, PasteLocation);
+		}
+
+		//TODO UMG HACK Because the WidgetTree isn't a real tree, add every widget.
+		NewWidget->SetFlags( RF_Transactional );
+		BP->WidgetTree->WidgetTemplates.Add( NewWidget );
+	}
+
+	//TODO UMG Reorder the live slot without rebuilding the structure
+	FBlueprintEditorUtils::MarkBlueprintAsStructurallyModified(BP);
+}
+
+void FWidgetBlueprintEditorUtils::ImportWidgetsFromText(UWidgetBlueprint* BP, const FString& TextToImport, /*out*/ TSet<UWidget*>& ImportedWidgetSet)
+{
+	// Turn the text buffer into objects
+	FWidgetObjectTextFactory Factory;
+	Factory.ProcessBuffer(BP->WidgetTree, RF_Transactional, TextToImport);
+
+	for ( auto& Entry : Factory.NewObjectMap )
+	{
+		ImportedWidgetSet.Add(Entry.Value);
+	}
+}
+
+#undef LOCTEXT_NAMESPACE
diff --git a/Engine/Source/Editor/UMGEditor/Private/WidgetBlueprintEditorUtils.h b/Engine/Source/Editor/UMGEditor/Private/WidgetBlueprintEditorUtils.h
--- a/Engine/Source/Editor/UMGEditor/Private/WidgetBlueprintEditorUtils.h
+++ b/Engine/Source/Editor/UMGEditor/Private/WidgetBlueprintEditorUtils.h
@@ -9,4 +9,22 @@ class FWidgetBlueprintEditorUtils
 {
 public:
 	static bool RenameWidget(TSharedRef<class FWidgetBlueprintEditor> BlueprintEditor, const FName& OldName, const FName& NewName);
+
+	static void CreateWidgetContextMenu(FMenuBuilder& MenuBuilder, TSharedRef<FWidgetBlueprintEditor> BlueprintEditor, FVector2D TargetLocation);
+
+	static void CopyWidgets(UWidgetBlueprint* BP, TSet<FWidgetReference> Widgets);
+
+	static void PasteWidgets(UWidgetBlueprint* BP, FWidgetReference ParentWidget, FVector2D PasteLocation);
+
+	static void DeleteWidgets(UWidgetBlueprint* BP, TSet<FWidgetReference> Widgets);
+
+public:
+	static void ExportWidgetsToText(TSet<UWidget*> WidgetsToExport, /*out*/ FString& ExportedText);
+
+	static void ImportWidgetsFromText(UWidgetBlueprint* BP, const FString& TextToImport, /*out*/ TSet<UWidget*>& ImportedWidgetSet);
+
+private:
+	static void BuildWrapWithMenu(FMenuBuilder& Menu, UWidgetBlueprint* BP, TSet<FWidgetReference> Widgets);
+
+	static void WrapWidgets(UWidgetBlueprint* BP, TSet<FWidgetReference> Widgets, UClass* WidgetClass);
 };
diff --git a/Engine/Source/Editor/UMGEditor/Public/WidgetReference.h b/Engine/Source/Editor/UMGEditor/Public/WidgetReference.h
--- a/Engine/Source/Editor/UMGEditor/Public/WidgetReference.h
+++ b/Engine/Source/Editor/UMGEditor/Public/WidgetReference.h
@@ -48,6 +48,7 @@ private:
 
 	TWeakPtr<class FWidgetBlueprintEditor> WidgetEditor;
 
+	//TODO UMG We probably want to make this a weak object reference.
 	UWidget* TemplateWidget;
 };
 
diff --git a/Engine/Source/Runtime/UMG/Private/Components/Visual.cpp b/Engine/Source/Runtime/UMG/Private/Components/Visual.cpp
new file mode 100644
--- /dev/null
+++ b/Engine/Source/Runtime/UMG/Private/Components/Visual.cpp
@@ -0,0 +1,11 @@
+// Copyright 1998-2014 Epic Games, Inc. All Rights Reserved.
+
+#include "UMGPrivatePCH.h"
+
+/////////////////////////////////////////////////////
+// UVisual
+
+UVisual::UVisual(const FPostConstructInitializeProperties& PCIP)
+	: Super(PCIP)
+{
+}
diff --git a/Engine/Source/Runtime/UMG/Private/Components/Widget.cpp b/Engine/Source/Runtime/UMG/Private/Components/Widget.cpp
--- a/Engine/Source/Runtime/UMG/Private/Components/Widget.cpp
+++ b/Engine/Source/Runtime/UMG/Private/Components/Widget.cpp
@@ -278,3 +278,31 @@ FSizeParam UWidget::ConvertSerializedSizeParamToRuntime(const FSlateChildSize& I
 		return FStretch(Input.Value);
 	}
 }
+
+void UWidget::GatherChildren(UWidget* Root, TSet<UWidget*>& Children)
+{
+	UPanelWidget* PanelRoot = Cast<UPanelWidget>(Root);
+	if ( PanelRoot )
+	{
+		for ( int32 ChildIndex = 0; ChildIndex < PanelRoot->GetChildrenCount(); ChildIndex++ )
+		{
+			UWidget* ChildWidget = PanelRoot->GetChildAt(ChildIndex);
+			Children.Add(ChildWidget);
+		}
+	}
+}
+
+void UWidget::GatherAllChildren(UWidget* Root, TSet<UWidget*>& Children)
+{
+	UPanelWidget* PanelRoot = Cast<UPanelWidget>(Root);
+	if ( PanelRoot )
+	{
+		for ( int32 ChildIndex = 0; ChildIndex < PanelRoot->GetChildrenCount(); ChildIndex++ )
+		{
+			UWidget* ChildWidget = PanelRoot->GetChildAt(ChildIndex);
+			Children.Add(ChildWidget);
+
+			GatherAllChildren(ChildWidget, Children);
+		}
+	}
+}
\ No newline at end of file
diff --git a/Engine/Source/Runtime/UMG/Private/WidgetTree.cpp b/Engine/Source/Runtime/UMG/Private/WidgetTree.cpp
--- a/Engine/Source/Runtime/UMG/Private/WidgetTree.cpp
+++ b/Engine/Source/Runtime/UMG/Private/WidgetTree.cpp
@@ -53,6 +53,7 @@ bool UWidgetTree::RemoveWidgetRecursive(UWidget* InRemovedWidget)
 	
 	//TODO UMG Make the Widget Tree actually a tree, instead of a list, it makes things like removal difficult.
 	
+	// If the widget being removed is a panel (has children) we need to delete all of its children first.
 	if ( UPanelWidget* InNonLeafRemovedWidget = Cast<UPanelWidget>(InRemovedWidget) )
 	{
 		while ( InNonLeafRemovedWidget->GetChildrenCount() > 0 )
@@ -60,16 +61,26 @@ bool UWidgetTree::RemoveWidgetRecursive(UWidget* InRemovedWidget)
 			RemoveWidget(InNonLeafRemovedWidget->GetChildAt(0));
 		}
 	}
-	
+
+	//TODO UMG This is slow, don't do this.  Need simpler way to manage finding parents
+
+	// Get the parent widget for this widget.
 	for ( UWidget* Template : WidgetTemplates )
 	{
 		UPanelWidget* NonLeafTemplate = Cast<UPanelWidget>(Template);
 		if ( NonLeafTemplate )
 		{
-			if ( NonLeafTemplate->RemoveChild(InRemovedWidget) )
+			for ( int32 ChildIndex = 0; ChildIndex < NonLeafTemplate->GetChildrenCount(); ChildIndex++ )
 			{
-				NonLeafTemplate->Modify();
-				return true;
+				if ( NonLeafTemplate->GetChildAt(ChildIndex) == InRemovedWidget )
+				{
+					NonLeafTemplate->Modify();
+
+					if ( NonLeafTemplate->RemoveChild(InRemovedWidget) )
+					{
+						return true;
+					}
+				}
 			}
 		}
 	}
@@ -79,6 +90,8 @@ bool UWidgetTree::RemoveWidgetRecursive(UWidget* InRemovedWidget)
 
 bool UWidgetTree::RemoveWidget(UWidget* InRemovedWidget)
 {
+	InRemovedWidget->Modify();
+
 	bool bRemoved = RemoveWidgetRecursive(InRemovedWidget);
 	int32 IndexRemoved = WidgetTemplates.Remove(InRemovedWidget);
 	
diff --git a/Engine/Source/Runtime/UMG/Public/Components/PanelSlot.h b/Engine/Source/Runtime/UMG/Public/Components/PanelSlot.h
--- a/Engine/Source/Runtime/UMG/Public/Components/PanelSlot.h
+++ b/Engine/Source/Runtime/UMG/Public/Components/PanelSlot.h
@@ -6,7 +6,7 @@
 
 /** The base class for all Slots in UMG. */
 UCLASS()
-class UMG_API UPanelSlot : public UObject
+class UMG_API UPanelSlot : public UVisual
 {
 	GENERATED_UCLASS_BODY()
 
diff --git a/Engine/Source/Runtime/UMG/Public/Components/Visual.h b/Engine/Source/Runtime/UMG/Public/Components/Visual.h
new file mode 100644
--- /dev/null
+++ b/Engine/Source/Runtime/UMG/Public/Components/Visual.h
@@ -0,0 +1,15 @@
+// Copyright 1998-2014 Epic Games, Inc. All Rights Reserved.
+
+#pragma once
+
+#include "Visual.generated.h"
+
+/** The base class for elements in UMG: slots and widgets. */
+UCLASS()
+class UMG_API UVisual : public UObject
+{
+	GENERATED_UCLASS_BODY()
+
+public:
+
+};
diff --git a/Engine/Source/Runtime/UMG/Public/Components/Widget.h b/Engine/Source/Runtime/UMG/Public/Components/Widget.h
--- a/Engine/Source/Runtime/UMG/Public/Components/Widget.h
+++ b/Engine/Source/Runtime/UMG/Public/Components/Widget.h
@@ -21,7 +21,7 @@
  * This is the base class for all wrapped Slate controls that are exposed to UMG.
  */
 UCLASS(Abstract)
-class UMG_API UWidget : public UObject
+class UMG_API UWidget : public UVisual
 {
 	GENERATED_UCLASS_BODY()
 
@@ -181,6 +181,9 @@ public:
 
 	static FSizeParam ConvertSerializedSizeParamToRuntime(const FSlateChildSize& Input);
 
+	static void GatherChildren(UWidget* Root, TSet<UWidget*>& Children);
+	static void GatherAllChildren(UWidget* Root, TSet<UWidget*>& Children);
+
 protected:
 	/** Function implemented by all subclasses of UWidget is called when the underlying SWidget needs to be constructed. */
 	virtual TSharedRef<SWidget> RebuildWidget();
diff --git a/Engine/Source/Runtime/UMG/Public/UMG.h b/Engine/Source/Runtime/UMG/Public/UMG.h
--- a/Engine/Source/Runtime/UMG/Public/UMG.h
+++ b/Engine/Source/Runtime/UMG/Public/UMG.h
@@ -24,6 +24,8 @@
 // Components
 #include "SlateWrapperTypes.h"
 
+#include "Visual.h"
+
 #include "PanelSlot.h"
 
 #include "Widget.h"
