UMG - Adding public interface for UMG Designer (stub).  Rebuilding extension widgets when selection changes instead of on mouse events in the designer.

diff --git a/Engine/Source/Editor/UMGEditor/Private/SUMGDesigner.cpp b/Engine/Source/Editor/UMGEditor/Private/SUMGDesigner.cpp
--- a/Engine/Source/Editor/UMGEditor/Private/SUMGDesigner.cpp
+++ b/Engine/Source/Editor/UMGEditor/Private/SUMGDesigner.cpp
@@ -153,6 +153,44 @@ void SUMGDesigner::OnEditorSelectionChanged()
 	{
 		SelectedWidget = FWidgetReference();
 	}
+
+	CreateExtensionWidgetsForSelection();
+}
+
+void SUMGDesigner::CreateExtensionWidgetsForSelection()
+{
+	// Remove all the current extension widgets
+	ExtensionWidgetCanvas->ClearChildren();
+
+	ExtensionWidgets.Reset();
+
+	TArray<FWidgetReference> Selected;
+	Selected.Add(SelectedWidget);
+
+	// Build extension widgets for new selection
+	for ( TSharedRef<FDesignerExtension>& Ext : DesignerExtensions )
+	{
+		Ext->BuildWidgetsForSelection(Selected, ExtensionWidgets);
+	}
+
+	TSharedRef<SHorizontalBox> ExtensionBox = SNew(SHorizontalBox);
+
+	// Add Widgets to designer surface
+	for ( TSharedRef<SWidget>& ExWidget : ExtensionWidgets )
+	{
+		ExtensionBox->AddSlot()
+			.AutoWidth()
+			[
+				ExWidget
+			];
+	}
+
+	ExtensionWidgetCanvas->AddSlot()
+		.Position(TAttribute<FVector2D>(this, &SUMGDesigner::GetCachedSelectionDesignerWidgetsLocation))
+		.Size(FVector2D(100, 25))
+		[
+			ExtensionBox
+		];
 }
 
 UWidgetBlueprint* SUMGDesigner::GetBlueprint() const
@@ -257,39 +295,6 @@ FReply SUMGDesigner::OnMouseButtonDown(const FGeometry& MyGeometry, const FPoint
 		TSet<FWidgetReference> SelectedTemplates;
 		SelectedTemplates.Add(SelectedWidget);
 		BlueprintEditor.Pin()->SelectWidgets(SelectedTemplates);
-
-		// Remove all the current extension widgets
-		ExtensionWidgetCanvas->ClearChildren();
-
-		ExtensionWidgets.Reset();
-
-		TArray<FWidgetReference> Selected;
-		Selected.Add(SelectedWidget);
-
-		// Build extension widgets for new selection
-		for ( TSharedRef<FDesignerExtension>& Ext : DesignerExtensions )
-		{
-			Ext->BuildWidgetsForSelection(Selected, ExtensionWidgets);
-		}
-
-		TSharedRef<SHorizontalBox> ExtensionBox = SNew(SHorizontalBox);
-
-		// Add Widgets to designer surface
-		for ( TSharedRef<SWidget>& ExWidget : ExtensionWidgets )
-		{
-			ExtensionBox->AddSlot()
-				.AutoWidth()
-				[
-					ExWidget
-				];
-		}
-
-		ExtensionWidgetCanvas->AddSlot()
-			.Position(TAttribute<FVector2D>(this, &SUMGDesigner::GetCachedSelectionDesignerWidgetsLocation))
-			.Size(FVector2D(100, 25))
-			[
-				ExtensionBox
-			];
 	}
 
 	// Capture mouse for the drag handle and general mouse actions
@@ -486,6 +491,15 @@ int32 SUMGDesigner::OnPaint(const FGeometry& AllottedGeometry, const FSlateRect&
 
 	LayerId += 1000;
 
+	TSet<FWidgetReference> Selected;
+	Selected.Add(SelectedWidget);
+
+	// Allow the extensions to paint anything they want.
+	for ( const TSharedRef<FDesignerExtension>& Ext : DesignerExtensions )
+	{
+		Ext->Paint(Selected, AllottedGeometry, MyClippingRect, OutDrawElements, LayerId, InWidgetStyle, bParentEnabled);
+	}
+
 	// Don't draw the hovered effect if it's also the selected widget
 	if ( HoveredSlateWidget.IsValid() && HoveredSlateWidget != SelectedSlateWidget )
 	{
diff --git a/Engine/Source/Editor/UMGEditor/Private/SUMGDesigner.h b/Engine/Source/Editor/UMGEditor/Private/SUMGDesigner.h
--- a/Engine/Source/Editor/UMGEditor/Private/SUMGDesigner.h
+++ b/Engine/Source/Editor/UMGEditor/Private/SUMGDesigner.h
@@ -6,13 +6,14 @@
 #include "SNodePanel.h"
 #include "SDesignSurface.h"
 #include "DesignerExtension.h"
+#include "IUMGDesigner.h"
 
 class FDesignerExtension;
 
 /**
  * The designer for widgets.  Allows for laying out widgets in a drag and drop environment.
  */
-class SUMGDesigner : public SDesignSurface
+class SUMGDesigner : public SDesignSurface, public IUMGDesigner
 {
 public:
 
@@ -45,6 +46,8 @@ public:
 	void Register(TSharedRef<FDesignerExtension> Extension);
 
 private:
+	void CreateExtensionWidgetsForSelection();
+
 	/** Displays the context menu when you right click */
 	void ShowContextMenu(const FGeometry& MyGeometry, const FPointerEvent& MouseEvent);
 
diff --git a/Engine/Source/Editor/UMGEditor/Public/DesignerExtension.h b/Engine/Source/Editor/UMGEditor/Public/DesignerExtension.h
--- a/Engine/Source/Editor/UMGEditor/Public/DesignerExtension.h
+++ b/Engine/Source/Editor/UMGEditor/Public/DesignerExtension.h
@@ -24,6 +24,10 @@ public:
 	/** Called every time the selection in the designer changes. */
 	virtual void BuildWidgetsForSelection(const TArray< FWidgetReference >& Selection, TArray< TSharedRef<SWidget> >& Widgets) = 0;
 
+	virtual void Paint(const TSet< FWidgetReference >& Selection, const FGeometry& AllottedGeometry, const FSlateRect& MyClippingRect, FSlateWindowElementList& OutDrawElements, int32 LayerId, const FWidgetStyle& InWidgetStyle, bool bParentEnabled) const
+	{
+	}
+
 	/** Gets the ID identifying this extension. */
 	FName GetExtensionId() const;
 
diff --git a/Engine/Source/Editor/UMGEditor/Public/IUMGDesigner.h b/Engine/Source/Editor/UMGEditor/Public/IUMGDesigner.h
new file mode 100644
--- /dev/null
+++ b/Engine/Source/Editor/UMGEditor/Public/IUMGDesigner.h
@@ -0,0 +1,13 @@
+// Copyright 1998-2014 Epic Games, Inc. All Rights Reserved.
+
+#pragma once
+
+/**
+ * The public interface implemented by the UMG Designer to allow extensions to call methods
+ * on the designer.
+ */
+class IUMGDesigner
+{
+public:
+
+};
