UMG - Now visualizing snap lines in the designer.  Not actually snapping yet.

diff --git a/Engine/Source/Editor/UMGEditor/Private/Extensions/CanvasSlotExtension.cpp b/Engine/Source/Editor/UMGEditor/Private/Extensions/CanvasSlotExtension.cpp
--- a/Engine/Source/Editor/UMGEditor/Private/Extensions/CanvasSlotExtension.cpp
+++ b/Engine/Source/Editor/UMGEditor/Private/Extensions/CanvasSlotExtension.cpp
@@ -4,6 +4,17 @@
 
 #define LOCTEXT_NAMESPACE "UMG"
 
+const float SnapDistance = 7;
+
+static float DistancePointToLine2D(const FVector2D& LinePointA, const FVector2D& LinePointB, const FVector2D& PointC)
+{
+	FVector2D AB = LinePointB - LinePointA;
+	FVector2D AC = PointC - LinePointA;
+
+	float Distance = FVector2D::CrossProduct(AB, AC) / FVector2D::Distance(LinePointA, LinePointB);
+	return FMath::Abs(Distance);
+}
+
 FCanvasSlotExtension::FCanvasSlotExtension()
 	: bDragging(false)
 {
@@ -48,6 +59,118 @@ void FCanvasSlotExtension::BuildWidgetsForSelection(const TArray< FWidgetReferen
 	Widgets.Add(MoveHandle.ToSharedRef());
 }
 
+bool FCanvasSlotExtension::GetCollisionPlanesForSlot(UCanvasPanel* Canvas, int32 SlotIndex, TArray<FVector2D>& Segments)
+{
+	Segments.SetNumUninitialized(8);
+
+	FGeometry ArrangedGeometry;
+	if ( Canvas->GetGeometryForSlot(SlotIndex, ArrangedGeometry) )
+	{
+		// Left Side
+		Segments[0] = ArrangedGeometry.Position;
+		Segments[1] = ArrangedGeometry.Position + FVector2D(0, ArrangedGeometry.Size.Y);
+
+		// Top Side
+		Segments[2] = ArrangedGeometry.Position;
+		Segments[3] = ArrangedGeometry.Position + FVector2D(ArrangedGeometry.Size.X, 0);
+
+		// Right Side
+		Segments[4] = ArrangedGeometry.Position + FVector2D(ArrangedGeometry.Size.X, 0);
+		Segments[5] = ArrangedGeometry.Position + ArrangedGeometry.Size;
+
+		// Bottom Side
+		Segments[6] = ArrangedGeometry.Position + FVector2D(0, ArrangedGeometry.Size.Y);
+		Segments[7] = ArrangedGeometry.Position + ArrangedGeometry.Size;
+
+		return true;
+	}
+
+	return false;
+}
+
+void FCanvasSlotExtension::Paint(const TSet< FWidgetReference >& Selection, const FGeometry& AllottedGeometry, const FSlateRect& MyClippingRect, FSlateWindowElementList& OutDrawElements, int32 LayerId, const FWidgetStyle& InWidgetStyle, bool bParentEnabled) const
+{
+	for ( const FWidgetReference& WidgetRef : Selection )
+	{
+		if ( !WidgetRef.IsValid() )
+		{
+			continue;
+		}
+
+		UWidget* Widget = WidgetRef.GetPreview();
+
+		if ( UCanvasPanelSlot* CanvasSlot = Cast<UCanvasPanelSlot>(Widget->Slot) )
+		{
+			UCanvasPanel* Canvas = CastChecked<UCanvasPanel>(CanvasSlot->Parent);
+
+			//TODO UMG Only show guide lines when near them and dragging
+			if ( bDragging )
+			{
+				FGeometry MyArrangedGeometry;
+				if ( !Canvas->GetGeometryForSlot(CanvasSlot, MyArrangedGeometry) )
+				{
+					continue;
+				}
+
+				TArray<FVector2D> LinePoints;
+				LinePoints.AddUninitialized(2);
+
+				for ( int32 SlotIndex = 0; SlotIndex < Canvas->GetChildrenCount(); SlotIndex++ )
+				{
+					// Ignore the slot being dragged.
+					if ( Canvas->Slots[SlotIndex] == CanvasSlot )
+					{
+						continue;
+					}
+
+					// Get the collision segments that we could potentially be docking against.
+					TArray<FVector2D> Segments;
+					if ( GetCollisionPlanesForSlot(Canvas, SlotIndex, Segments) )
+					{
+						for ( int32 SegmentBase = 0; SegmentBase < Segments.Num(); SegmentBase += 2 )
+						{
+							FVector2D PointA = Segments[SegmentBase];
+							FVector2D PointB = Segments[SegmentBase + 1];
+
+							FVector2D CollisionPoint = MyArrangedGeometry.Position;
+
+							//TODO Collide against all sides of the arranged geometry.
+							float Distance = DistancePointToLine2D(PointA, PointB, MyArrangedGeometry.Position);
+							if ( Distance <= SnapDistance )
+							{
+								FVector2D FarthestPoint = FVector2D::Distance(PointA, CollisionPoint) > FVector2D::Distance(PointB, CollisionPoint) ? PointA : PointB;
+								FVector2D NearestPoint = FVector2D::Distance(PointA, CollisionPoint) > FVector2D::Distance(PointB, CollisionPoint) ? PointB : PointA;
+
+								LinePoints[0] = FarthestPoint;
+								LinePoints[1] = FarthestPoint + (NearestPoint - FarthestPoint) * 100000;
+
+								LinePoints[0].X = FMath::Clamp(LinePoints[0].X, 0.0f, MyClippingRect.Right - MyClippingRect.Left);
+								LinePoints[0].Y = FMath::Clamp(LinePoints[0].Y, 0.0f, MyClippingRect.Bottom - MyClippingRect.Top);
+
+								LinePoints[1].X = FMath::Clamp(LinePoints[1].X, 0.0f, MyClippingRect.Right - MyClippingRect.Left);
+								LinePoints[1].Y = FMath::Clamp(LinePoints[1].Y, 0.0f, MyClippingRect.Bottom - MyClippingRect.Top);
+
+								FLinearColor Color(0.5f, 0.75, 1);
+								const bool bAntialias = true;
+
+								FSlateDrawElement::MakeLines(
+									OutDrawElements,
+									LayerId,
+									AllottedGeometry.ToPaintGeometry(),
+									LinePoints,
+									MyClippingRect,
+									ESlateDrawEffect::None,
+									Color,
+									bAntialias);
+							}
+						}
+					}
+				}
+			}
+		}
+	}
+}
+
 FReply FCanvasSlotExtension::HandleBeginDrag(const FGeometry& Geometry, const FPointerEvent& Event)
 {
 	bDragging = true;
@@ -74,8 +197,7 @@ FReply FCanvasSlotExtension::HandleDragging(const FGeometry& Geometry, const FPo
 	{
 		for ( FWidgetReference& Selection : SelectionCache )
 		{
-			MoveByAmount(Selection.GetPreview(), Event.GetCursorDelta());
-			MoveByAmount(Selection.GetTemplate(), Event.GetCursorDelta());
+			MoveByAmount(Selection, Event.GetCursorDelta());
 		}
 
 		return FReply::Handled();
@@ -84,33 +206,43 @@ FReply FCanvasSlotExtension::HandleDragging(const FGeometry& Geometry, const FPo
 	return FReply::Unhandled();
 }
 
-void FCanvasSlotExtension::MoveByAmount(UWidget* Widget, FVector2D Delta)
+void FCanvasSlotExtension::MoveByAmount(FWidgetReference& WidgetRef, FVector2D Delta)
 {
-	if ( !Delta.IsZero() )
+	if ( Delta.IsZero() )
 	{
-		UCanvasPanelSlot* CanvasSlot = Cast<UCanvasPanelSlot>(Widget->Slot);
-		UCanvasPanel* Parent = Cast<UCanvasPanel>(CanvasSlot->Parent);
+		return;
+	}
 
-		FMargin Offsets = CanvasSlot->LayoutData.Offsets;
-		Offsets.Left += Delta.X;
-		Offsets.Top += Delta.Y;
+	UWidget* Widget = WidgetRef.GetPreview();
 
-		// If the slot is stretched horizontally we need to move the right side as it no longer represents width, but
-		// now represents margin from the right stretched side.
-		if ( CanvasSlot->LayoutData.Anchors.IsStretchedHorizontal() )
-		{
-			Offsets.Right -= Delta.X;
-		}
+	UCanvasPanelSlot* CanvasSlot = CastChecked<UCanvasPanelSlot>(Widget->Slot);
+	UCanvasPanel* Parent = CastChecked<UCanvasPanel>(CanvasSlot->Parent);
 
-		// If the slot is stretched vertically we need to move the bottom side as it no longer represents width, but
-		// now represents margin from the bottom stretched side.
-		if ( CanvasSlot->LayoutData.Anchors.IsStretchedVertical() )
-		{
-			Offsets.Bottom -= Delta.Y;
-		}
+	FMargin Offsets = CanvasSlot->LayoutData.Offsets;
+	Offsets.Left += Delta.X;
+	Offsets.Top += Delta.Y;
+
+	// If the slot is stretched horizontally we need to move the right side as it no longer represents width, but
+	// now represents margin from the right stretched side.
+	if ( CanvasSlot->LayoutData.Anchors.IsStretchedHorizontal() )
+	{
+		Offsets.Right -= Delta.X;
+	}
 
-		CanvasSlot->SetOffset(Offsets);
+	// If the slot is stretched vertically we need to move the bottom side as it no longer represents width, but
+	// now represents margin from the bottom stretched side.
+	if ( CanvasSlot->LayoutData.Anchors.IsStretchedVertical() )
+	{
+		Offsets.Bottom -= Delta.Y;
 	}
+
+	CanvasSlot->SetOffset(Offsets);
+
+	// Update the Template widget to match
+	UWidget* TemplateWidget = WidgetRef.GetTemplate();
+	UCanvasPanelSlot* TemplateSlot = CastChecked<UCanvasPanelSlot>(TemplateWidget->Slot);
+
+	TemplateSlot->SetOffset(Offsets);
 }
 
 #undef LOCTEXT_NAMESPACE
diff --git a/Engine/Source/Editor/UMGEditor/Private/Extensions/CanvasSlotExtension.h b/Engine/Source/Editor/UMGEditor/Private/Extensions/CanvasSlotExtension.h
--- a/Engine/Source/Editor/UMGEditor/Private/Extensions/CanvasSlotExtension.h
+++ b/Engine/Source/Editor/UMGEditor/Private/Extensions/CanvasSlotExtension.h
@@ -12,6 +12,7 @@ public:
 	bool IsActive(const TArray< FWidgetReference >& Selection);
 	
 	virtual void BuildWidgetsForSelection(const TArray< FWidgetReference >& Selection, TArray< TSharedRef<SWidget> >& Widgets) override;
+	virtual void Paint(const TSet< FWidgetReference >& Selection, const FGeometry& AllottedGeometry, const FSlateRect& MyClippingRect, FSlateWindowElementList& OutDrawElements, int32 LayerId, const FWidgetStyle& InWidgetStyle, bool bParentEnabled) const override;
 
 private:
 
@@ -19,7 +20,9 @@ private:
 	FReply HandleEndDrag(const FGeometry& Geometry, const FPointerEvent& Event);
 	FReply HandleDragging(const FGeometry& Geometry, const FPointerEvent& Event);
 
-	void MoveByAmount(UWidget* Widget, FVector2D Delta);
+	void MoveByAmount(FWidgetReference& WidgetRef, FVector2D Delta);
+
+	static bool GetCollisionPlanesForSlot(UCanvasPanel* Canvas, int32 SlotIndex, TArray<FVector2D>& Segments);
 
 	bool bDragging;
 
diff --git a/Engine/Source/Editor/UMGEditor/Private/Extensions/HorizontalSlotExtension.cpp b/Engine/Source/Editor/UMGEditor/Private/Extensions/HorizontalSlotExtension.cpp
--- a/Engine/Source/Editor/UMGEditor/Private/Extensions/HorizontalSlotExtension.cpp
+++ b/Engine/Source/Editor/UMGEditor/Private/Extensions/HorizontalSlotExtension.cpp
@@ -66,8 +66,8 @@ FReply FHorizontalSlotExtension::HandleShift(int32 ShiftAmount)
 
 void FHorizontalSlotExtension::ShiftHorizontal(UWidget* Widget, int32 ShiftAmount)
 {
-	UHorizontalBoxSlot* Slot = Cast<UHorizontalBoxSlot>(Widget->Slot);
-	UHorizontalBox* Parent = Cast<UHorizontalBox>(Slot->Parent);
+	UHorizontalBoxSlot* Slot = CastChecked<UHorizontalBoxSlot>(Widget->Slot);
+	UHorizontalBox* Parent = CastChecked<UHorizontalBox>(Slot->Parent);
 
 	int32 CurrentIndex = Parent->GetChildIndex(Widget);
 	Parent->Slots.RemoveAt(CurrentIndex);
diff --git a/Engine/Source/Editor/UMGEditor/Private/Extensions/VerticalSlotExtension.cpp b/Engine/Source/Editor/UMGEditor/Private/Extensions/VerticalSlotExtension.cpp
--- a/Engine/Source/Editor/UMGEditor/Private/Extensions/VerticalSlotExtension.cpp
+++ b/Engine/Source/Editor/UMGEditor/Private/Extensions/VerticalSlotExtension.cpp
@@ -73,8 +73,8 @@ FReply FVerticalSlotExtension::HandleShiftVertical(int32 ShiftAmount)
 
 void FVerticalSlotExtension::ShiftVertical(UWidget* Widget, int32 ShiftAmount)
 {
-	UVerticalBoxSlot* Slot = Cast<UVerticalBoxSlot>(Widget->Slot);
-	UVerticalBox* Parent = Cast<UVerticalBox>(Slot->Parent);
+	UVerticalBoxSlot* Slot = CastChecked<UVerticalBoxSlot>(Widget->Slot);
+	UVerticalBox* Parent = CastChecked<UVerticalBox>(Slot->Parent);
 
 	int32 CurrentIndex = Parent->GetChildIndex(Widget);
 	Parent->Slots.RemoveAt(CurrentIndex);
diff --git a/Engine/Source/Editor/UMGEditor/Private/SUMGDesigner.cpp b/Engine/Source/Editor/UMGEditor/Private/SUMGDesigner.cpp
--- a/Engine/Source/Editor/UMGEditor/Private/SUMGDesigner.cpp
+++ b/Engine/Source/Editor/UMGEditor/Private/SUMGDesigner.cpp
@@ -173,6 +173,9 @@ void SUMGDesigner::CreateExtensionWidgetsForSelection()
 		Ext->BuildWidgetsForSelection(Selected, ExtensionWidgets);
 	}
 
+	//TODO UMG We need to have a better layout system for extensions.
+	//         currently everything is just laid out across the top.
+
 	TSharedRef<SHorizontalBox> ExtensionBox = SNew(SHorizontalBox);
 
 	// Add Widgets to designer surface
diff --git a/Engine/Source/Runtime/UMG/Private/Components/CanvasPanel.cpp b/Engine/Source/Runtime/UMG/Private/Components/CanvasPanel.cpp
--- a/Engine/Source/Runtime/UMG/Private/Components/CanvasPanel.cpp
+++ b/Engine/Source/Runtime/UMG/Private/Components/CanvasPanel.cpp
@@ -97,6 +97,12 @@ TSharedPtr<SConstraintCanvas> UCanvasPanel::GetCanvasWidget() const
 	return MyCanvas;
 }
 
+bool UCanvasPanel::GetGeometryForSlot(int32 SlotIndex, FGeometry& ArrangedGeometry) const
+{
+	UCanvasPanelSlot* Slot = Cast<UCanvasPanelSlot>(Slots[SlotIndex]);
+	return GetGeometryForSlot(Slot, ArrangedGeometry);
+}
+
 bool UCanvasPanel::GetGeometryForSlot(UCanvasPanelSlot* Slot, FGeometry& ArrangedGeometry) const
 {
 	if ( Slot->Content == NULL )
diff --git a/Engine/Source/Runtime/UMG/Public/Components/CanvasPanel.h b/Engine/Source/Runtime/UMG/Public/Components/CanvasPanel.h
--- a/Engine/Source/Runtime/UMG/Public/Components/CanvasPanel.h
+++ b/Engine/Source/Runtime/UMG/Public/Components/CanvasPanel.h
@@ -16,6 +16,9 @@ class UMG_API UCanvasPanel : public UPanelWidget
 	/** Gets the underlying native canvas widget if it has been constructed */
 	TSharedPtr<class SConstraintCanvas> GetCanvasWidget() const;
 
+	/** Computes the geometry for a particular slot based on the current geometry of the canvas. */
+	bool GetGeometryForSlot(int32 SlotIndex, FGeometry& ArrangedGeometry) const;
+
 	/** Computes the geometry for a particular slot based on the current geometry of the canvas. */
 	bool GetGeometryForSlot(UCanvasPanelSlot* Slot, FGeometry& ArrangedGeometry) const;
 
