UMG - Fixing drag/drop previews for widgets on the canvas.  Fixing a few crash bugs.  UPanelSlot now has virtual functions for setting the slots desired position and size, some panels may use it, some may not.

diff --git a/Engine/Source/Editor/UMGEditor/Private/Extensions/DesignerExtension.cpp b/Engine/Source/Editor/UMGEditor/Private/Extensions/DesignerExtension.cpp
--- a/Engine/Source/Editor/UMGEditor/Private/Extensions/DesignerExtension.cpp
+++ b/Engine/Source/Editor/UMGEditor/Private/Extensions/DesignerExtension.cpp
@@ -98,8 +98,11 @@ void FDesignerExtension::BeginTransaction(const FText& SessionName)
 
 	for ( FWidgetReference& Selection : SelectionCache )
 	{
-		Selection.GetPreview()->Modify();
-		Selection.GetTemplate()->Modify();
+		if ( Selection.IsValid() )
+		{
+			Selection.GetPreview()->Modify();
+			Selection.GetTemplate()->Modify();
+		}
 	}
 }
 
diff --git a/Engine/Source/Editor/UMGEditor/Private/SUMGDesigner.cpp b/Engine/Source/Editor/UMGEditor/Private/SUMGDesigner.cpp
--- a/Engine/Source/Editor/UMGEditor/Private/SUMGDesigner.cpp
+++ b/Engine/Source/Editor/UMGEditor/Private/SUMGDesigner.cpp
@@ -897,13 +897,21 @@ UWidget* SUMGDesigner::AddPreview(const FGeometry& MyGeometry, const FDragDropEv
 				Widget->IsDesignTime(true);
 			
 				FVector2D LocalPosition = ArrangedWidget.Geometry.AbsoluteToLocal(DragDropEvent.GetScreenSpacePosition());
-				Parent->AddChild(Widget, LocalPosition);
-				//@TODO UMG When we add a child blindly we need to default the slot size to the preferred size of the widget if the container supports such things.
-				//@TODO UMG We may need a desired size canvas, where the slots have no size, they only give you position, alternatively, maybe slots that don't clip, so center is still easy.
-			
-				DropPreviewParent = Parent;
-			
-				return Widget;
+				if ( UPanelSlot* Slot = Parent->AddChild(Widget) )
+				{
+					Slot->SetDesiredPosition(LocalPosition);
+					Slot->SetDesiredSize(FVector2D(150, 30));
+					//@TODO UMG When we add a child blindly we need to default the slot size to the preferred size of the widget if the container supports such things.
+					//@TODO UMG We may need a desired size canvas, where the slots have no size, they only give you position, alternatively, maybe slots that don't clip, so center is still easy.
+
+					DropPreviewParent = Parent;
+
+					return Widget;
+				}
+				else
+				{
+					// TODO UMG ERROR Slot can not be created because maybe the max children has been reached.
+				}
 			}
 			else if ( BP->WidgetTree->WidgetTemplates.Num() == 1 )
 			{
@@ -948,14 +956,18 @@ bool SUMGDesigner::AddToTemplate(const FGeometry& MyGeometry, const FDragDropEve
 				Widget->IsDesignTime(true);
 				
 				FVector2D LocalPosition = ArrangedWidget.Geometry.AbsoluteToLocal(DragDropEvent.GetScreenSpacePosition());
-				Parent->AddChild(Widget, LocalPosition);
-				//@TODO UMG When we add a child blindly we need to default the slot size to the preferred size of the widget if the container supports such things.
-				//@TODO UMG We may need a desired size canvas, where the slots have no size, they only give you position, alternatively, maybe slots that don't clip, so center is still easy.
-				
-				// Update the selected template to be the newly created one.
-				SelectedWidget = FWidgetReference::FromTemplate(BlueprintEditor.Pin(), Widget);
-				
-				return true;
+				if ( UPanelSlot* Slot = Parent->AddChild(Widget) )
+				{
+					Slot->SetDesiredPosition(LocalPosition);
+					Slot->SetDesiredSize(FVector2D(150, 30));
+					//@TODO UMG When we add a child blindly we need to default the slot size to the preferred size of the widget if the container supports such things.
+					//@TODO UMG We may need a desired size canvas, where the slots have no size, they only give you position, alternatively, maybe slots that don't clip, so center is still easy.
+
+					// Update the selected template to be the newly created one.
+					SelectedWidget = FWidgetReference::FromTemplate(BlueprintEditor.Pin(), Widget);
+
+					return true;
+				}
 			}
 		}
 		
diff --git a/Engine/Source/Editor/UMGEditor/Private/SUMGEditorTreeItem.cpp b/Engine/Source/Editor/UMGEditor/Private/SUMGEditorTreeItem.cpp
--- a/Engine/Source/Editor/UMGEditor/Private/SUMGEditorTreeItem.cpp
+++ b/Engine/Source/Editor/UMGEditor/Private/SUMGEditorTreeItem.cpp
@@ -151,7 +151,7 @@ FReply SUMGEditorTreeItem::OnDrop(const FGeometry& MyGeometry, const FDragDropEv
 
 			UWidget* Widget = DragDropOp->Template->Create(BP->WidgetTree);
 
-			Parent->AddChild(Widget, FVector2D(0,0));
+			Parent->AddChild(Widget);
 
 			FBlueprintEditorUtils::MarkBlueprintAsStructurallyModified(BP);
 		}
diff --git a/Engine/Source/Editor/UMGEditor/Private/Templates/WidgetTemplateButton.cpp b/Engine/Source/Editor/UMGEditor/Private/Templates/WidgetTemplateButton.cpp
--- a/Engine/Source/Editor/UMGEditor/Private/Templates/WidgetTemplateButton.cpp
+++ b/Engine/Source/Editor/UMGEditor/Private/Templates/WidgetTemplateButton.cpp
@@ -25,7 +25,7 @@ UWidget* FWidgetTemplateButton::Create(UWidgetTree* Tree)
 
 	Text->Text = FText::FromString(TEXT("Button Text"));
 	Text->Font.Size = 10;
-	Button->AddChild(Text, FVector2D(0, 0));
+	Button->AddChild(Text);
 
 	return Button;
 }
diff --git a/Engine/Source/Editor/UMGEditor/Private/Templates/WidgetTemplateCheckBox.cpp b/Engine/Source/Editor/UMGEditor/Private/Templates/WidgetTemplateCheckBox.cpp
--- a/Engine/Source/Editor/UMGEditor/Private/Templates/WidgetTemplateCheckBox.cpp
+++ b/Engine/Source/Editor/UMGEditor/Private/Templates/WidgetTemplateCheckBox.cpp
@@ -26,9 +26,9 @@ UWidget* FWidgetTemplateCheckBox::Create(UWidgetTree* Tree)
 	Text->Text = LOCTEXT("CheckboxText", "Checkbox Text");
 	Text->Font.Size = 10;
 
-	Horizontal->AddChild(Checkbox)
+	Horizontal->Add(Checkbox)
 		->Size = FSlateChildSize(ESlateSizeRule::Automatic);
-	Horizontal->AddChild(Text);
+	Horizontal->Add(Text);
 
 	return Horizontal;
 }
diff --git a/Engine/Source/Editor/UMGEditor/Private/WidgetBlueprintEditorUtils.cpp b/Engine/Source/Editor/UMGEditor/Private/WidgetBlueprintEditorUtils.cpp
--- a/Engine/Source/Editor/UMGEditor/Private/WidgetBlueprintEditorUtils.cpp
+++ b/Engine/Source/Editor/UMGEditor/Private/WidgetBlueprintEditorUtils.cpp
@@ -241,7 +241,7 @@ void FWidgetBlueprintEditorUtils::WrapWidgets(UWidgetBlueprint* BP, TSet<FWidget
 		{
 			CurrentParent->ReplaceChildAt(OutIndex, NewWrapperWidget);
 
-			NewWrapperWidget->AddChild(Item.GetTemplate(), FVector2D());
+			NewWrapperWidget->AddChild(Item.GetTemplate());
 		}
 	}
 
@@ -313,7 +313,8 @@ void FWidgetBlueprintEditorUtils::PasteWidgets(UWidgetBlueprint* BP, FWidgetRefe
 		// automatically.
 		if ( NewWidget->GetParent() == NULL )
 		{
-			ParentWidget->AddChild(NewWidget, PasteLocation);
+			UPanelSlot* Slot = ParentWidget->AddChild(NewWidget);
+			Slot->SetDesiredPosition(PasteLocation);
 		}
 
 		//TODO UMG HACK Because the WidgetTree isn't a real tree, add every widget.
diff --git a/Engine/Source/Runtime/UMG/Private/Components/CanvasPanelSlot.cpp b/Engine/Source/Runtime/UMG/Private/Components/CanvasPanelSlot.cpp
--- a/Engine/Source/Runtime/UMG/Private/Components/CanvasPanelSlot.cpp
+++ b/Engine/Source/Runtime/UMG/Private/Components/CanvasPanelSlot.cpp
@@ -9,7 +9,7 @@ UCanvasPanelSlot::UCanvasPanelSlot(const FPostConstructInitializeProperties& PCI
 	: Super(PCIP)
 	, Slot(NULL)
 {
-	LayoutData.Offsets = FMargin(0, 0, 1, 1);
+	LayoutData.Offsets = FMargin(0, 0, 100, 30);
 	LayoutData.Anchors = FAnchors(0.0f, 0.0f);
 	LayoutData.Alignment = FVector2D(0.0f, 0.0f);
 }
@@ -25,6 +25,16 @@ void UCanvasPanelSlot::BuildSlot(TSharedRef<SConstraintCanvas> Canvas)
 		];
 }
 
+void UCanvasPanelSlot::SetDesiredPosition(FVector2D InPosition)
+{
+	SetOffset(FMargin(InPosition.X, InPosition.Y, LayoutData.Offsets.Right, LayoutData.Offsets.Bottom));
+}
+
+void UCanvasPanelSlot::SetDesiredSize(FVector2D InSize)
+{
+	SetOffset(FMargin(LayoutData.Offsets.Left, LayoutData.Offsets.Top, InSize.X, InSize.Y));
+}
+
 void UCanvasPanelSlot::Resize(const FVector2D& Direction, const FVector2D& Amount)
 {
 	if ( Direction.X < 0 )
diff --git a/Engine/Source/Runtime/UMG/Private/Components/HorizontalBox.cpp b/Engine/Source/Runtime/UMG/Private/Components/HorizontalBox.cpp
--- a/Engine/Source/Runtime/UMG/Private/Components/HorizontalBox.cpp
+++ b/Engine/Source/Runtime/UMG/Private/Components/HorizontalBox.cpp
@@ -34,9 +34,9 @@ void UHorizontalBox::OnSlotRemoved(UPanelSlot* Slot)
 	}
 }
 
-UHorizontalBoxSlot* UHorizontalBox::AddChild(UWidget* Content)
+UHorizontalBoxSlot* UHorizontalBox::Add(UWidget* Content)
 {
-	return Cast<UHorizontalBoxSlot>( Super::AddChild(Content, FVector2D(0, 0)) );
+	return Cast<UHorizontalBoxSlot>( Super::AddChild(Content) );
 }
 
 TSharedRef<SWidget> UHorizontalBox::RebuildWidget()
diff --git a/Engine/Source/Runtime/UMG/Private/Components/PanelSlot.cpp b/Engine/Source/Runtime/UMG/Private/Components/PanelSlot.cpp
--- a/Engine/Source/Runtime/UMG/Private/Components/PanelSlot.cpp
+++ b/Engine/Source/Runtime/UMG/Private/Components/PanelSlot.cpp
@@ -10,6 +10,16 @@ UPanelSlot::UPanelSlot(const FPostConstructInitializeProperties& PCIP)
 {
 }
 
+void UPanelSlot::SetDesiredPosition(FVector2D InPosition)
+{
+
+}
+
+void UPanelSlot::SetDesiredSize(FVector2D InSize)
+{
+
+}
+
 void UPanelSlot::Resize(const FVector2D& Direction, const FVector2D& Amount)
 {
 
diff --git a/Engine/Source/Runtime/UMG/Private/Components/PanelWidget.cpp b/Engine/Source/Runtime/UMG/Private/Components/PanelWidget.cpp
--- a/Engine/Source/Runtime/UMG/Private/Components/PanelWidget.cpp
+++ b/Engine/Source/Runtime/UMG/Private/Components/PanelWidget.cpp
@@ -52,7 +52,7 @@ bool UPanelWidget::RemoveChildAt(int32 Index)
 	return true;
 }
 
-UPanelSlot* UPanelWidget::AddChild(UWidget* Content, FVector2D Position)
+UPanelSlot* UPanelWidget::AddChild(UWidget* Content)
 {
 	if ( !bCanHaveMultipleChildren && GetChildrenCount() > 0 )
 	{
diff --git a/Engine/Source/Runtime/UMG/Private/Components/VerticalBox.cpp b/Engine/Source/Runtime/UMG/Private/Components/VerticalBox.cpp
--- a/Engine/Source/Runtime/UMG/Private/Components/VerticalBox.cpp
+++ b/Engine/Source/Runtime/UMG/Private/Components/VerticalBox.cpp
@@ -34,9 +34,9 @@ void UVerticalBox::OnSlotRemoved(UPanelSlot* Slot)
 	}
 }
 
-UVerticalBoxSlot* UVerticalBox::AddChild(UWidget* Content)
+UVerticalBoxSlot* UVerticalBox::Add(UWidget* Content)
 {
-	return Cast<UVerticalBoxSlot>( Super::AddChild(Content, FVector2D(0, 0)) );
+	return Cast<UVerticalBoxSlot>( Super::AddChild(Content) );
 }
 
 TSharedRef<SWidget> UVerticalBox::RebuildWidget()
diff --git a/Engine/Source/Runtime/UMG/Public/Components/CanvasPanelSlot.h b/Engine/Source/Runtime/UMG/Public/Components/CanvasPanelSlot.h
--- a/Engine/Source/Runtime/UMG/Public/Components/CanvasPanelSlot.h
+++ b/Engine/Source/Runtime/UMG/Public/Components/CanvasPanelSlot.h
@@ -36,6 +36,10 @@ class UMG_API UCanvasPanelSlot : public UPanelSlot
 
 	void BuildSlot(TSharedRef<SConstraintCanvas> Canvas);
 
+	virtual void SetDesiredPosition(FVector2D InPosition) override;
+
+	virtual void SetDesiredSize(FVector2D InSize) override;
+
 	virtual void Resize(const FVector2D& Direction, const FVector2D& Amount) override;
 
 	virtual bool CanResize(const FVector2D& Direction) const override;
diff --git a/Engine/Source/Runtime/UMG/Public/Components/HorizontalBox.h b/Engine/Source/Runtime/UMG/Public/Components/HorizontalBox.h
--- a/Engine/Source/Runtime/UMG/Public/Components/HorizontalBox.h
+++ b/Engine/Source/Runtime/UMG/Public/Components/HorizontalBox.h
@@ -11,7 +11,7 @@ class UMG_API UHorizontalBox : public UPanelWidget
 
 	/**  */
 	UFUNCTION(BlueprintCallable, Category="Widget")
-	UHorizontalBoxSlot* AddChild(UWidget* Content);
+	UHorizontalBoxSlot* Add(UWidget* Content);
 
 #if WITH_EDITOR
 	// UWidget interface
diff --git a/Engine/Source/Runtime/UMG/Public/Components/PanelSlot.h b/Engine/Source/Runtime/UMG/Public/Components/PanelSlot.h
--- a/Engine/Source/Runtime/UMG/Public/Components/PanelSlot.h
+++ b/Engine/Source/Runtime/UMG/Public/Components/PanelSlot.h
@@ -20,6 +20,10 @@ public:
 	
 	bool IsDesignTime() const;
 
+	virtual void SetDesiredPosition(FVector2D InPosition);
+
+	virtual void SetDesiredSize(FVector2D InSize);
+
 	virtual void Resize(const FVector2D& Direction, const FVector2D& Amount);
 
 	virtual bool CanResize(const FVector2D& Direction) const;
diff --git a/Engine/Source/Runtime/UMG/Public/Components/PanelWidget.h b/Engine/Source/Runtime/UMG/Public/Components/PanelWidget.h
--- a/Engine/Source/Runtime/UMG/Public/Components/PanelWidget.h
+++ b/Engine/Source/Runtime/UMG/Public/Components/PanelWidget.h
@@ -28,7 +28,7 @@ public:
 
 	bool RemoveChildAt(int32 Index);
 
-	UPanelSlot* AddChild(UWidget* Content, FVector2D Position);
+	UPanelSlot* AddChild(UWidget* Content);
 
 	void ReplaceChildAt(int32 Index, UWidget* Content);
 
diff --git a/Engine/Source/Runtime/UMG/Public/Components/VerticalBox.h b/Engine/Source/Runtime/UMG/Public/Components/VerticalBox.h
--- a/Engine/Source/Runtime/UMG/Public/Components/VerticalBox.h
+++ b/Engine/Source/Runtime/UMG/Public/Components/VerticalBox.h
@@ -15,7 +15,7 @@ class UMG_API UVerticalBox : public UPanelWidget
 
 	/**  */
 	UFUNCTION(BlueprintCallable, Category="Widget")
-	UVerticalBoxSlot* AddChild(UWidget* Content);
+	UVerticalBoxSlot* Add(UWidget* Content);
 
 #if WITH_EDITOR
 	// UWidget interface
