UMG - Further improving the alignment guides to use every side of the widget.

diff --git a/Engine/Source/Editor/UMGEditor/Private/Extensions/CanvasSlotExtension.cpp b/Engine/Source/Editor/UMGEditor/Private/Extensions/CanvasSlotExtension.cpp
--- a/Engine/Source/Editor/UMGEditor/Private/Extensions/CanvasSlotExtension.cpp
+++ b/Engine/Source/Editor/UMGEditor/Private/Extensions/CanvasSlotExtension.cpp
@@ -59,33 +59,47 @@ void FCanvasSlotExtension::BuildWidgetsForSelection(const TArray< FWidgetReferen
 	Widgets.Add(MoveHandle.ToSharedRef());
 }
 
-bool FCanvasSlotExtension::GetCollisionPlanesForSlot(UCanvasPanel* Canvas, int32 SlotIndex, TArray<FVector2D>& Segments)
+bool FCanvasSlotExtension::GetCollisionSegmentsForSlot(UCanvasPanel* Canvas, int32 SlotIndex, TArray<FVector2D>& Segments)
 {
-	Segments.SetNumUninitialized(8);
-
 	FGeometry ArrangedGeometry;
 	if ( Canvas->GetGeometryForSlot(SlotIndex, ArrangedGeometry) )
 	{
-		// Left Side
-		Segments[0] = ArrangedGeometry.Position;
-		Segments[1] = ArrangedGeometry.Position + FVector2D(0, ArrangedGeometry.Size.Y);
+		GetCollisionSegmentsFromGeometry(ArrangedGeometry, Segments);
+		return true;
+	}
+	return false;
+}
+
+bool FCanvasSlotExtension::GetCollisionSegmentsForSlot(UCanvasPanel* Canvas, UCanvasPanelSlot* Slot, TArray<FVector2D>& Segments)
+{
+	FGeometry ArrangedGeometry;
+	if ( Canvas->GetGeometryForSlot(Slot, ArrangedGeometry) )
+	{
+		GetCollisionSegmentsFromGeometry(ArrangedGeometry, Segments);
+		return true;
+	}
+	return false;
+}
 
-		// Top Side
-		Segments[2] = ArrangedGeometry.Position;
-		Segments[3] = ArrangedGeometry.Position + FVector2D(ArrangedGeometry.Size.X, 0);
+void FCanvasSlotExtension::GetCollisionSegmentsFromGeometry(FGeometry ArrangedGeometry, TArray<FVector2D>& Segments)
+{
+	Segments.SetNumUninitialized(8);
 
-		// Right Side
-		Segments[4] = ArrangedGeometry.Position + FVector2D(ArrangedGeometry.Size.X, 0);
-		Segments[5] = ArrangedGeometry.Position + ArrangedGeometry.Size;
+	// Left Side
+	Segments[0] = ArrangedGeometry.Position;
+	Segments[1] = ArrangedGeometry.Position + FVector2D(0, ArrangedGeometry.Size.Y);
 
-		// Bottom Side
-		Segments[6] = ArrangedGeometry.Position + FVector2D(0, ArrangedGeometry.Size.Y);
-		Segments[7] = ArrangedGeometry.Position + ArrangedGeometry.Size;
+	// Top Side
+	Segments[2] = ArrangedGeometry.Position;
+	Segments[3] = ArrangedGeometry.Position + FVector2D(ArrangedGeometry.Size.X, 0);
 
-		return true;
-	}
+	// Right Side
+	Segments[4] = ArrangedGeometry.Position + FVector2D(ArrangedGeometry.Size.X, 0);
+	Segments[5] = ArrangedGeometry.Position + ArrangedGeometry.Size;
 
-	return false;
+	// Bottom Side
+	Segments[6] = ArrangedGeometry.Position + FVector2D(0, ArrangedGeometry.Size.Y);
+	Segments[7] = ArrangedGeometry.Position + ArrangedGeometry.Size;
 }
 
 void FCanvasSlotExtension::Paint(const TSet< FWidgetReference >& Selection, const FGeometry& AllottedGeometry, const FSlateRect& MyClippingRect, FSlateWindowElementList& OutDrawElements, int32 LayerId, const FWidgetStyle& InWidgetStyle, bool bParentEnabled) const
@@ -115,53 +129,63 @@ void FCanvasSlotExtension::Paint(const TSet< FWidgetReference >& Selection, cons
 				TArray<FVector2D> LinePoints;
 				LinePoints.AddUninitialized(2);
 
-				for ( int32 SlotIndex = 0; SlotIndex < Canvas->GetChildrenCount(); SlotIndex++ )
+				// Get the collision segments that we could potentially be docking against.
+				TArray<FVector2D> MySegments;
+				if ( GetCollisionSegmentsForSlot(Canvas, CanvasSlot, MySegments) )
 				{
-					// Ignore the slot being dragged.
-					if ( Canvas->Slots[SlotIndex] == CanvasSlot )
+					for ( int32 MySegmentIndex = 0; MySegmentIndex < MySegments.Num(); MySegmentIndex += 2 )
 					{
-						continue;
-					}
+						FVector2D MySegmentBase = MySegments[MySegmentIndex];
 
-					// Get the collision segments that we could potentially be docking against.
-					TArray<FVector2D> Segments;
-					if ( GetCollisionPlanesForSlot(Canvas, SlotIndex, Segments) )
-					{
-						for ( int32 SegmentBase = 0; SegmentBase < Segments.Num(); SegmentBase += 2 )
+						for ( int32 SlotIndex = 0; SlotIndex < Canvas->GetChildrenCount(); SlotIndex++ )
 						{
-							FVector2D PointA = Segments[SegmentBase];
-							FVector2D PointB = Segments[SegmentBase + 1];
-
-							FVector2D CollisionPoint = MyArrangedGeometry.Position;
+							// Ignore the slot being dragged.
+							if ( Canvas->Slots[SlotIndex] == CanvasSlot )
+							{
+								continue;
+							}
 
-							//TODO Collide against all sides of the arranged geometry.
-							float Distance = DistancePointToLine2D(PointA, PointB, MyArrangedGeometry.Position);
-							if ( Distance <= SnapDistance )
+							// Get the collision segments that we could potentially be docking against.
+							TArray<FVector2D> Segments;
+							if ( GetCollisionSegmentsForSlot(Canvas, SlotIndex, Segments) )
 							{
-								FVector2D FarthestPoint = FVector2D::Distance(PointA, CollisionPoint) > FVector2D::Distance(PointB, CollisionPoint) ? PointA : PointB;
-								FVector2D NearestPoint = FVector2D::Distance(PointA, CollisionPoint) > FVector2D::Distance(PointB, CollisionPoint) ? PointB : PointA;
-
-								LinePoints[0] = FarthestPoint;
-								LinePoints[1] = FarthestPoint + (NearestPoint - FarthestPoint) * 100000;
-
-								LinePoints[0].X = FMath::Clamp(LinePoints[0].X, 0.0f, MyClippingRect.Right - MyClippingRect.Left);
-								LinePoints[0].Y = FMath::Clamp(LinePoints[0].Y, 0.0f, MyClippingRect.Bottom - MyClippingRect.Top);
-
-								LinePoints[1].X = FMath::Clamp(LinePoints[1].X, 0.0f, MyClippingRect.Right - MyClippingRect.Left);
-								LinePoints[1].Y = FMath::Clamp(LinePoints[1].Y, 0.0f, MyClippingRect.Bottom - MyClippingRect.Top);
-
-								FLinearColor Color(0.5f, 0.75, 1);
-								const bool bAntialias = true;
-
-								FSlateDrawElement::MakeLines(
-									OutDrawElements,
-									LayerId,
-									AllottedGeometry.ToPaintGeometry(),
-									LinePoints,
-									MyClippingRect,
-									ESlateDrawEffect::None,
-									Color,
-									bAntialias);
+								for ( int32 SegmentBase = 0; SegmentBase < Segments.Num(); SegmentBase += 2 )
+								{
+									FVector2D PointA = Segments[SegmentBase];
+									FVector2D PointB = Segments[SegmentBase + 1];
+
+									FVector2D CollisionPoint = MySegmentBase;
+
+									//TODO Collide against all sides of the arranged geometry.
+									float Distance = DistancePointToLine2D(PointA, PointB, CollisionPoint);
+									if ( Distance <= SnapDistance )
+									{
+										FVector2D FarthestPoint = FVector2D::Distance(PointA, CollisionPoint) > FVector2D::Distance(PointB, CollisionPoint) ? PointA : PointB;
+										FVector2D NearestPoint = FVector2D::Distance(PointA, CollisionPoint) > FVector2D::Distance(PointB, CollisionPoint) ? PointB : PointA;
+
+										LinePoints[0] = FarthestPoint;
+										LinePoints[1] = FarthestPoint + ( NearestPoint - FarthestPoint ) * 100000;
+
+										LinePoints[0].X = FMath::Clamp(LinePoints[0].X, 0.0f, MyClippingRect.Right - MyClippingRect.Left);
+										LinePoints[0].Y = FMath::Clamp(LinePoints[0].Y, 0.0f, MyClippingRect.Bottom - MyClippingRect.Top);
+
+										LinePoints[1].X = FMath::Clamp(LinePoints[1].X, 0.0f, MyClippingRect.Right - MyClippingRect.Left);
+										LinePoints[1].Y = FMath::Clamp(LinePoints[1].Y, 0.0f, MyClippingRect.Bottom - MyClippingRect.Top);
+
+										FLinearColor Color(0.5f, 0.75, 1);
+										const bool bAntialias = true;
+
+										FSlateDrawElement::MakeLines(
+											OutDrawElements,
+											LayerId,
+											AllottedGeometry.ToPaintGeometry(),
+											LinePoints,
+											MyClippingRect,
+											ESlateDrawEffect::None,
+											Color,
+											bAntialias);
+									}
+								}
 							}
 						}
 					}
diff --git a/Engine/Source/Editor/UMGEditor/Private/Extensions/CanvasSlotExtension.h b/Engine/Source/Editor/UMGEditor/Private/Extensions/CanvasSlotExtension.h
--- a/Engine/Source/Editor/UMGEditor/Private/Extensions/CanvasSlotExtension.h
+++ b/Engine/Source/Editor/UMGEditor/Private/Extensions/CanvasSlotExtension.h
@@ -22,7 +22,9 @@ private:
 
 	void MoveByAmount(FWidgetReference& WidgetRef, FVector2D Delta);
 
-	static bool GetCollisionPlanesForSlot(UCanvasPanel* Canvas, int32 SlotIndex, TArray<FVector2D>& Segments);
+	static bool GetCollisionSegmentsForSlot(UCanvasPanel* Canvas, int32 SlotIndex, TArray<FVector2D>& Segments);
+	static bool GetCollisionSegmentsForSlot(UCanvasPanel* Canvas, UCanvasPanelSlot* Slot, TArray<FVector2D>& Segments);
+	static void GetCollisionSegmentsFromGeometry(FGeometry ArrangedGeometry, TArray<FVector2D>& Segments);
 
 	bool bDragging;
 
diff --git a/Engine/Source/Editor/UMGEditor/Private/SUMGDesigner.cpp b/Engine/Source/Editor/UMGEditor/Private/SUMGDesigner.cpp
--- a/Engine/Source/Editor/UMGEditor/Private/SUMGDesigner.cpp
+++ b/Engine/Source/Editor/UMGEditor/Private/SUMGDesigner.cpp
@@ -911,6 +911,7 @@ UWidget* SUMGDesigner::AddPreview(const FGeometry& MyGeometry, const FDragDropEv
 				else
 				{
 					// TODO UMG ERROR Slot can not be created because maybe the max children has been reached.
+					//          Maybe we can traverse the hierarchy and add it to the first parent that will accept it?
 				}
 			}
 			else if ( BP->WidgetTree->WidgetTemplates.Num() == 1 )
