UMG - Now allowing the drag drop behavior in the designer to move a widget.  Still very rough at the moment, doesn't preserve size characteristics, still need to find a good way to migrate slot preferences.  A little closer to mkaing the widget tree a widget tree, now that parent pointers are accessible.

diff --git a/Engine/Source/Editor/UMGEditor/Private/SUMGDesigner.cpp b/Engine/Source/Editor/UMGEditor/Private/SUMGDesigner.cpp
--- a/Engine/Source/Editor/UMGEditor/Private/SUMGDesigner.cpp
+++ b/Engine/Source/Editor/UMGEditor/Private/SUMGDesigner.cpp
@@ -13,6 +13,31 @@
 
 #define LOCTEXT_NAMESPACE "UMG"
 
+const float HoveredAnimationTime = 0.150f;
+
+class FSelectedWidgetDragDropOp : public FDecoratedDragDropOp
+{
+public:
+	DRAG_DROP_OPERATOR_TYPE(FSelectedWidgetDragDropOp, FDecoratedDragDropOp)
+
+	FWidgetReference Widget;
+
+	static TSharedRef<FSelectedWidgetDragDropOp> New(FWidgetReference InWidget);
+
+	//virtual TSharedPtr<SWidget> GetDefaultDecorator() const override;
+};
+
+TSharedRef<FSelectedWidgetDragDropOp> FSelectedWidgetDragDropOp::New(FWidgetReference InWidget)
+{
+	TSharedRef<FSelectedWidgetDragDropOp> Operation = MakeShareable(new FSelectedWidgetDragDropOp());
+	Operation->Widget = InWidget;
+	Operation->DefaultHoverText = FText::FromString( InWidget.GetTemplate()->GetLabel() );
+	Operation->CurrentHoverText = FText::FromString( InWidget.GetTemplate()->GetLabel() );
+	Operation->Construct();
+
+	return Operation;
+}
+
 //////////////////////////////////////////////////////////////////////////
 
 static bool LocateWidgetsUnderCursor_Helper(FArrangedWidget& Candidate, FVector2D InAbsoluteCursorLocation, FArrangedChildren& OutWidgetsUnderCursor, bool bIgnoreEnabledStatus)
@@ -77,6 +102,9 @@ void SUMGDesigner::Construct(const FArguments& InArgs, TSharedPtr<FWidgetBluepri
 
 	HoverTime = 0;
 
+	bMouseDown = false;
+	bMovingExistingWidget = false;
+
 	DragDirections.Init((int32)DH_MAX);
 	DragDirections[DH_TOP_LEFT] = FVector2D(-1, -1);
 	DragDirections[DH_TOP_CENTER] = FVector2D(0, -1);
@@ -298,6 +326,8 @@ FReply SUMGDesigner::OnMouseButtonDown(const FGeometry& MyGeometry, const FPoint
 		TSet<FWidgetReference> SelectedTemplates;
 		SelectedTemplates.Add(SelectedWidget);
 		BlueprintEditor.Pin()->SelectWidgets(SelectedTemplates);
+
+		bMouseDown = true;
 	}
 
 	// Capture mouse for the drag handle and general mouse actions
@@ -308,6 +338,9 @@ FReply SUMGDesigner::OnMouseButtonUp(const FGeometry& MyGeometry, const FPointer
 {
 	if ( MouseEvent.GetEffectingButton() == EKeys::LeftMouseButton )
 	{
+		bMouseDown = false;
+		bMovingExistingWidget = false;
+
 		if ( CurrentHandle != DH_NONE )
 		{
 			EndTransaction();
@@ -356,9 +389,19 @@ FReply SUMGDesigner::OnMouseMove(const FGeometry& MyGeometry, const FPointerEven
 		}
 	}
 
-	// Update the hovered widget under the mouse
+	if ( bMouseDown )
+	{
+		if ( SelectedWidget.IsValid() && !bMovingExistingWidget )
+		{
+			bMovingExistingWidget = true;
+			//Drag selected widgets
+			return FReply::Handled().DetectDrag(AsShared(), EKeys::LeftMouseButton);
+		}
+	}
+	
 	if ( CurrentHandle == DH_NONE )
 	{
+		// Update the hovered widget under the mouse
 		FArrangedWidget ArrangedWidget(SNullWidget::NullWidget, FGeometry());
 		FWidgetReference NewHoveredWidget = GetWidgetAtCursor(MyGeometry, MouseEvent, ArrangedWidget);
 		if ( !( NewHoveredWidget == HoveredWidget ) )
@@ -511,8 +554,6 @@ int32 SUMGDesigner::OnPaint(const FGeometry& AllottedGeometry, const FSlateRect&
 		FArrangedWidget ArrangedWidget(SNullWidget::NullWidget, FGeometry());
 		GetArrangedWidgetRelativeToWindow(Widget, ArrangedWidget);
 
-		const float HoveredAnimationTime = 0.25f;
-
 		// Draw hovered effect
 		// Azure = 0x007FFF
 		const FLinearColor HoveredTint(0, 0.5, 1, FMath::Clamp(HoverTime / HoveredAnimationTime, 0.0f, 1.0f));
@@ -825,6 +866,16 @@ void SUMGDesigner::Tick(const FGeometry& AllottedGeometry, const double InCurren
 	SCompoundWidget::Tick(AllottedGeometry, InCurrentTime, InDeltaTime);
 }
 
+FReply SUMGDesigner::OnDragDetected(const FGeometry& MyGeometry, const FPointerEvent& MouseEvent)
+{
+	if ( SelectedWidget.IsValid() )
+	{
+		return FReply::Handled().BeginDragDrop(FSelectedWidgetDragDropOp::New(SelectedWidget));
+	}
+
+	return FReply::Unhandled();
+}
+
 void SUMGDesigner::OnDragEnter(const FGeometry& MyGeometry, const FDragDropEvent& DragDropEvent)
 {
 	//@TODO UMG Drop Feedback
@@ -846,7 +897,8 @@ void SUMGDesigner::OnDragLeave(const FDragDropEvent& DragDropEvent)
 		}
 
 		UWidgetBlueprint* BP = GetBlueprint();
-		BP->WidgetTree->RemoveWidget(DropPreviewWidget);
+		const bool bIsRecursive = false;
+		BP->WidgetTree->RemoveWidget(DropPreviewWidget, bIsRecursive);
 		DropPreviewWidget = NULL;
 	}
 }
@@ -862,11 +914,13 @@ FReply SUMGDesigner::OnDragOver(const FGeometry& MyGeometry, const FDragDropEven
 			DropPreviewParent->RemoveChild(DropPreviewWidget);
 		}
 		
-		BP->WidgetTree->RemoveWidget(DropPreviewWidget);
+		const bool bIsRecursive = false;
+		BP->WidgetTree->RemoveWidget(DropPreviewWidget, bIsRecursive);
 		DropPreviewWidget = NULL;
 	}
 	
-	DropPreviewWidget = AddPreview(MyGeometry, DragDropEvent);
+	const bool bIsPreview = true;
+	DropPreviewWidget = ProcessDropAndAddWidget(MyGeometry, DragDropEvent, bIsPreview);
 	if ( DropPreviewWidget )
 	{
 		//@TODO UMG Drop Feedback
@@ -876,22 +930,46 @@ FReply SUMGDesigner::OnDragOver(const FGeometry& MyGeometry, const FDragDropEven
 	return FReply::Unhandled();
 }
 
-UWidget* SUMGDesigner::AddPreview(const FGeometry& MyGeometry, const FDragDropEvent& DragDropEvent)
+UWidget* SUMGDesigner::ProcessDropAndAddWidget(const FGeometry& MyGeometry, const FDragDropEvent& DragDropEvent, bool bIsPreview)
 {
+	// In order to prevent the GetWidgetAtCursor code from picking the widget we're about to move, we need to mark it
+	// as the drop preview widget before any other code can run.
+	TSharedPtr<FSelectedWidgetDragDropOp> SelectedDragDropOp = DragDropEvent.GetOperationAs<FSelectedWidgetDragDropOp>();
+	if ( SelectedDragDropOp.IsValid() )
+	{
+		DropPreviewWidget = SelectedDragDropOp->Widget.GetPreview();
+	}
+
+	FArrangedWidget ArrangedWidget(SNullWidget::NullWidget, FGeometry());
+	FWidgetReference WidgetUnderCursor = GetWidgetAtCursor(MyGeometry, DragDropEvent, ArrangedWidget);
+
+	UWidgetBlueprint* BP = GetBlueprint();
+
+	UWidget* Target = NULL;
+	if ( WidgetUnderCursor.IsValid() )
+	{
+		Target = bIsPreview ? WidgetUnderCursor.GetPreview() : WidgetUnderCursor.GetTemplate();
+	}
+
 	TSharedPtr<FWidgetTemplateDragDropOp> DragDropOp = DragDropEvent.GetOperationAs<FWidgetTemplateDragDropOp>();
 	if ( DragDropOp.IsValid() )
-	{
-		FArrangedWidget ArrangedWidget(SNullWidget::NullWidget, FGeometry());
-		FWidgetReference Selection = GetWidgetAtCursor(MyGeometry, DragDropEvent, ArrangedWidget);
-		
-		UWidgetBlueprint* BP = GetBlueprint();
-		
-		if ( Selection.IsValid() )
+	{		
+		if ( Target )
 		{
-			UWidget* Target = Selection.GetPreview();
 			if ( Target->IsA(UPanelWidget::StaticClass()) )
 			{
 				UPanelWidget* Parent = Cast<UPanelWidget>(Target);
+
+				FScopedTransaction Transaction(LOCTEXT("Designer_AddWidget", "Add Widget"));
+
+				if ( !bIsPreview )
+				{
+					Parent->SetFlags(RF_Transactional);
+					Parent->Modify();
+
+					BP->WidgetTree->SetFlags(RF_Transactional);
+					BP->WidgetTree->Modify();
+				}
 				
 				UWidget* Widget = DragDropOp->Template->Create(BP->WidgetTree);
 				Widget->IsDesignTime(true);
@@ -906,6 +984,11 @@ UWidget* SUMGDesigner::AddPreview(const FGeometry& MyGeometry, const FDragDropEv
 
 					DropPreviewParent = Parent;
 
+					if ( !bIsPreview )
+					{
+						Transaction.Cancel();
+					}
+
 					return Widget;
 				}
 				else
@@ -913,6 +996,11 @@ UWidget* SUMGDesigner::AddPreview(const FGeometry& MyGeometry, const FDragDropEv
 					// TODO UMG ERROR Slot can not be created because maybe the max children has been reached.
 					//          Maybe we can traverse the hierarchy and add it to the first parent that will accept it?
 				}
+
+				if ( !bIsPreview )
+				{
+					Transaction.Cancel();
+				}
 			}
 			else if ( BP->WidgetTree->WidgetTemplates.Num() == 1 )
 			{
@@ -921,76 +1009,106 @@ UWidget* SUMGDesigner::AddPreview(const FGeometry& MyGeometry, const FDragDropEv
 				
 				DropPreviewParent = NULL;
 				
+				return Widget;
+			}
+		}
+		else
+		{
+			if ( BP->WidgetTree->WidgetTemplates.Num() == 0 )
+			{
+				FScopedTransaction Transaction(LOCTEXT("Designer_AddWidget", "Add Widget"));
+
+				if ( !bIsPreview )
+				{
+					BP->WidgetTree->SetFlags(RF_Transactional);
+					BP->WidgetTree->Modify();
+				}
+
+				// TODO UMG This method isn't great, maybe the user widget should just be a canvas.
+
+				// Add it to the root if there are no other widgets to add it to.
+				UWidget* Widget = DragDropOp->Template->Create(BP->WidgetTree);
+				Widget->IsDesignTime(true);
+
+				SelectedWidget = FWidgetReference::FromTemplate(BlueprintEditor.Pin(), Widget);
+
+				DropPreviewParent = NULL;
+
+				if ( !bIsPreview )
+				{
+					Transaction.Cancel();
+				}
+
 				return Widget;
 			}
 		}
 	}
-	
-	return NULL;
-}
 
-bool SUMGDesigner::AddToTemplate(const FGeometry& MyGeometry, const FDragDropEvent& DragDropEvent)
-{
-	TSharedPtr<FWidgetTemplateDragDropOp> DragDropOp = DragDropEvent.GetOperationAs<FWidgetTemplateDragDropOp>();
-	if ( DragDropOp.IsValid() )
+	// Attempt to deal with moving widgets from a drag operation.
+	if ( SelectedDragDropOp.IsValid() )
 	{
-		FArrangedWidget ArrangedWidget(SNullWidget::NullWidget, FGeometry());
-		FWidgetReference Selection = GetWidgetAtCursor(MyGeometry, DragDropEvent, ArrangedWidget);
-		
-		UWidgetBlueprint* BP = GetBlueprint();
-		
-		if ( Selection.IsValid() )
+		if ( Target )
 		{
-			UWidget* Target = Selection.GetTemplate();
 			if ( Target->IsA(UPanelWidget::StaticClass()) )
 			{
-				UPanelWidget* Parent = Cast<UPanelWidget>(Target);
+				UPanelWidget* NewParent = Cast<UPanelWidget>(Target);
 
-				const FScopedTransaction Transaction(LOCTEXT("Designer_AddWidget", "Add Widget"));
-				Parent->SetFlags(RF_Transactional);
-				Parent->Modify();
+				FScopedTransaction Transaction(LOCTEXT("Designer_MoveWidget", "Move Widget"));
+
+				if ( !bIsPreview )
+				{
+					NewParent->SetFlags(RF_Transactional);
+					NewParent->Modify();
+
+					BP->WidgetTree->SetFlags(RF_Transactional);
+					BP->WidgetTree->Modify();
+				}
+
+				UWidget* Widget = bIsPreview ? SelectedDragDropOp->Widget.GetPreview() : SelectedDragDropOp->Widget.GetTemplate();
+
+				if ( Widget->GetParent() )
+				{
+					if ( !bIsPreview )
+					{
+						Widget->GetParent()->Modify();
+					}
+
+					Widget->GetParent()->RemoveChild(Widget);
+				}
 
-				BP->WidgetTree->SetFlags(RF_Transactional);
-				BP->WidgetTree->Modify();
-				
-				UWidget* Widget = DragDropOp->Template->Create(BP->WidgetTree);
-				Widget->IsDesignTime(true);
-				
 				FVector2D LocalPosition = ArrangedWidget.Geometry.AbsoluteToLocal(DragDropEvent.GetScreenSpacePosition());
-				if ( UPanelSlot* Slot = Parent->AddChild(Widget) )
+				if ( UPanelSlot* Slot = NewParent->AddChild(Widget) )
 				{
+					//TODO UMG Migrate existing slot info
 					Slot->SetDesiredPosition(LocalPosition);
 					Slot->SetDesiredSize(FVector2D(150, 30));
 					//@TODO UMG When we add a child blindly we need to default the slot size to the preferred size of the widget if the container supports such things.
 					//@TODO UMG We may need a desired size canvas, where the slots have no size, they only give you position, alternatively, maybe slots that don't clip, so center is still easy.
 
-					// Update the selected template to be the newly created one.
-					SelectedWidget = FWidgetReference::FromTemplate(BlueprintEditor.Pin(), Widget);
+					DropPreviewParent = NewParent;
 
-					return true;
+					if ( !bIsPreview )
+					{
+						Transaction.Cancel();
+					}
+
+					return Widget;
+				}
+				else
+				{
+					// TODO UMG ERROR Slot can not be created because maybe the max children has been reached.
+					//          Maybe we can traverse the hierarchy and add it to the first parent that will accept it?
 				}
-			}
-		}
-		
-		if ( BP->WidgetTree->WidgetTemplates.Num() == 0 )
-		{
-			const FScopedTransaction Transaction(LOCTEXT("Designer_AddWidget", "Add Widget"));
-			BP->WidgetTree->SetFlags(RF_Transactional);
-			BP->WidgetTree->Modify();
 
-			// TODO UMG This method isn't great, maybe the user widget should just be a canvas.
-			
-			// Add it to the root if there are no other widgets to add it to.
-			UWidget* Widget = DragDropOp->Template->Create(BP->WidgetTree);
-			Widget->IsDesignTime(true);
-			
-			SelectedWidget = FWidgetReference::FromTemplate(BlueprintEditor.Pin(), Widget);
-			
-			return true;
+				if ( !bIsPreview )
+				{
+					Transaction.Cancel();
+				}
+			}
 		}
 	}
 	
-	return false;
+	return NULL;
 }
 
 FReply SUMGDesigner::OnDrop(const FGeometry& MyGeometry, const FDragDropEvent& DragDropEvent)
@@ -1004,11 +1122,14 @@ FReply SUMGDesigner::OnDrop(const FGeometry& MyGeometry, const FDragDropEvent& D
 			DropPreviewParent->RemoveChild(DropPreviewWidget);
 		}
 		
-		BP->WidgetTree->RemoveWidget(DropPreviewWidget);
+		const bool bIsRecursive = false;
+		BP->WidgetTree->RemoveWidget(DropPreviewWidget, bIsRecursive);
 		DropPreviewWidget = NULL;
 	}
 	
-	if ( AddToTemplate(MyGeometry, DragDropEvent) )
+	const bool bIsPreview = false;
+	UWidget* Widget = ProcessDropAndAddWidget(MyGeometry, DragDropEvent, bIsPreview);
+	if ( Widget )
 	{
 		FBlueprintEditorUtils::MarkBlueprintAsStructurallyModified(BP);
 
diff --git a/Engine/Source/Editor/UMGEditor/Private/SUMGDesigner.h b/Engine/Source/Editor/UMGEditor/Private/SUMGDesigner.h
--- a/Engine/Source/Editor/UMGEditor/Private/SUMGDesigner.h
+++ b/Engine/Source/Editor/UMGEditor/Private/SUMGDesigner.h
@@ -37,6 +37,7 @@ public:
 	virtual int32 OnPaint(const FGeometry& AllottedGeometry, const FSlateRect& MyClippingRect, FSlateWindowElementList& OutDrawElements, int32 LayerId, const FWidgetStyle& InWidgetStyle, bool bParentEnabled) const override;
 	virtual void Tick(const FGeometry& AllottedGeometry, const double InCurrentTime, const float InDeltaTime) override;
 
+	virtual FReply OnDragDetected(const FGeometry& MyGeometry, const FPointerEvent& MouseEvent) override;
 	virtual void OnDragEnter(const FGeometry& MyGeometry, const FDragDropEvent& DragDropEvent) override;
 	virtual void OnDragLeave(const FDragDropEvent& DragDropEvent) override;
 	virtual FReply OnDragOver(const FGeometry& MyGeometry, const FDragDropEvent& DragDropEvent) override;
@@ -103,10 +104,8 @@ private:
 private:
 	void DrawDragHandles(const FPaintGeometry& SelectionGeometry, const FSlateRect& MyClippingRect, FSlateWindowElementList& OutDrawElements, int32 LayerId, const FWidgetStyle& InWidgetStyle) const;
 	DragHandle HitTestDragHandles(const FGeometry& AllottedGeometry, const FPointerEvent& PointerEvent) const;
-	
-	UWidget* AddPreview(const FGeometry& MyGeometry, const FDragDropEvent& DragDropEvent);
-	
-	bool AddToTemplate(const FGeometry& MyGeometry, const FDragDropEvent& DragDropEvent);
+
+	UWidget* ProcessDropAndAddWidget(const FGeometry& MyGeometry, const FDragDropEvent& DragDropEvent, bool bIsPreview);
 	
 private:
 	/** A reference to the BP Editor that owns this designer */
@@ -144,6 +143,12 @@ private:
 	TArray< TWeakPtr<SWidget> > SelectedSlateWidgets;
 	TWeakPtr<SWidget> SelectedSlateWidget;
 
+	/**  */
+	bool bMouseDown;
+
+	/** An existing widget is being moved in its current container, or in to a new container. */
+	bool bMovingExistingWidget;
+
 	/** The wall clock time the user has been hovering over a single widget */
 	float HoverTime;
 
diff --git a/Engine/Source/Runtime/UMG/Private/Components/Widget.cpp b/Engine/Source/Runtime/UMG/Private/Components/Widget.cpp
--- a/Engine/Source/Runtime/UMG/Private/Components/Widget.cpp
+++ b/Engine/Source/Runtime/UMG/Private/Components/Widget.cpp
@@ -68,7 +68,7 @@ bool UWidget::IsHovered() const
 	return false;
 }
 
-UWidget* UWidget::GetParent() const
+UPanelWidget* UWidget::GetParent() const
 {
 	if ( Slot )
 	{
diff --git a/Engine/Source/Runtime/UMG/Private/WidgetTree.cpp b/Engine/Source/Runtime/UMG/Private/WidgetTree.cpp
--- a/Engine/Source/Runtime/UMG/Private/WidgetTree.cpp
+++ b/Engine/Source/Runtime/UMG/Private/WidgetTree.cpp
@@ -47,52 +47,34 @@ UPanelWidget* UWidgetTree::FindWidgetParent(UWidget* Widget, int32& OutChildInde
 	return NULL;
 }
 
-bool UWidgetTree::RemoveWidgetRecursive(UWidget* InRemovedWidget)
+bool UWidgetTree::RemoveWidget(UWidget* InRemovedWidget, bool bIsRecursive)
 {
-	UWidget* Parent = NULL;
-	
-	//TODO UMG Make the Widget Tree actually a tree, instead of a list, it makes things like removal difficult.
-	
-	// If the widget being removed is a panel (has children) we need to delete all of its children first.
-	if ( UPanelWidget* InNonLeafRemovedWidget = Cast<UPanelWidget>(InRemovedWidget) )
+	InRemovedWidget->Modify();
+
+	bool bRemoved = false;
+
+	UPanelWidget* InRemovedWidgetParent = InRemovedWidget->GetParent();
+	if ( InRemovedWidgetParent )
 	{
-		while ( InNonLeafRemovedWidget->GetChildrenCount() > 0 )
+		InRemovedWidgetParent->Modify();
+
+		if ( InRemovedWidgetParent->RemoveChild(InRemovedWidget) )
 		{
-			RemoveWidget(InNonLeafRemovedWidget->GetChildAt(0));
+			bRemoved = true;
 		}
 	}
 
-	//TODO UMG This is slow, don't do this.  Need simpler way to manage finding parents
-
-	// Get the parent widget for this widget.
-	for ( UWidget* Template : WidgetTemplates )
+	if ( bIsRecursive )
 	{
-		UPanelWidget* NonLeafTemplate = Cast<UPanelWidget>(Template);
-		if ( NonLeafTemplate )
+		if ( UPanelWidget* InNonLeafRemovedWidget = Cast<UPanelWidget>(InRemovedWidget) )
 		{
-			for ( int32 ChildIndex = 0; ChildIndex < NonLeafTemplate->GetChildrenCount(); ChildIndex++ )
+			while ( InNonLeafRemovedWidget->GetChildrenCount() > 0 )
 			{
-				if ( NonLeafTemplate->GetChildAt(ChildIndex) == InRemovedWidget )
-				{
-					NonLeafTemplate->Modify();
-
-					if ( NonLeafTemplate->RemoveChild(InRemovedWidget) )
-					{
-						return true;
-					}
-				}
+				RemoveWidget(InNonLeafRemovedWidget->GetChildAt(0));
 			}
 		}
 	}
-	
-	return false;
-}
-
-bool UWidgetTree::RemoveWidget(UWidget* InRemovedWidget)
-{
-	InRemovedWidget->Modify();
 
-	bool bRemoved = RemoveWidgetRecursive(InRemovedWidget);
 	int32 IndexRemoved = WidgetTemplates.Remove(InRemovedWidget);
 	
 	return bRemoved || IndexRemoved != -1;
diff --git a/Engine/Source/Runtime/UMG/Public/Blueprint/WidgetTree.h b/Engine/Source/Runtime/UMG/Public/Blueprint/WidgetTree.h
--- a/Engine/Source/Runtime/UMG/Public/Blueprint/WidgetTree.h
+++ b/Engine/Source/Runtime/UMG/Public/Blueprint/WidgetTree.h
@@ -19,7 +19,7 @@ public:
 
 	UWidget* FindWidget(const FString& Name) const;
 
-	bool RemoveWidget(UWidget* Widget);
+	bool RemoveWidget(UWidget* Widget, bool bIsRecursive = true);
 
 	class UPanelWidget* FindWidgetParent(UWidget* Widget, int32& OutChildIndex);
 
@@ -38,7 +38,4 @@ public:
 
 		return (T*)Widget;
 	}
-	
-private:
-	bool RemoveWidgetRecursive(UWidget* InRemovedWidget);
 };
diff --git a/Engine/Source/Runtime/UMG/Public/Components/Widget.h b/Engine/Source/Runtime/UMG/Public/Components/Widget.h
--- a/Engine/Source/Runtime/UMG/Public/Components/Widget.h
+++ b/Engine/Source/Runtime/UMG/Public/Components/Widget.h
@@ -122,7 +122,7 @@ public:
 
 	/** Gets the parent widget */
 	UFUNCTION(BlueprintCallable, Category="Widget")
-	UWidget* GetParent() const;
+	class UPanelWidget* GetParent() const;
 
 	/**
 	 * Gets the underlying slate widget or constructs it if it doesn't exist.  This function is
