UMG - Now renaming variable nodes in the blueprint when widgets are named.

diff --git a/Engine/Source/Editor/UMGEditor/Private/WidgetBlueprintCompiler.cpp b/Engine/Source/Editor/UMGEditor/Private/WidgetBlueprintCompiler.cpp
--- a/Engine/Source/Editor/UMGEditor/Private/WidgetBlueprintCompiler.cpp
+++ b/Engine/Source/Editor/UMGEditor/Private/WidgetBlueprintCompiler.cpp
@@ -596,6 +596,7 @@ void FWidgetBlueprintCompiler::CreateClassVariablesFromBlueprint()
 
 	ValidateWidgetNames();
 
+	// Build the set of variables based on the variable widgets in the widget tree.
 	for ( UWidget* Widget : Blueprint->WidgetTree->WidgetTemplates )
 	{
 		// Skip non-variable widgets
diff --git a/Engine/Source/Editor/UMGEditor/Private/WidgetBlueprintEditorUtils.cpp b/Engine/Source/Editor/UMGEditor/Private/WidgetBlueprintEditorUtils.cpp
--- a/Engine/Source/Editor/UMGEditor/Private/WidgetBlueprintEditorUtils.cpp
+++ b/Engine/Source/Editor/UMGEditor/Private/WidgetBlueprintEditorUtils.cpp
@@ -79,77 +79,59 @@ bool FWidgetBlueprintEditorUtils::RenameWidget(TSharedRef<FWidgetBlueprintEditor
 	UWidget* Widget = Blueprint->WidgetTree->FindWidget(OldNameStr);
 	check(Widget);
 
-	// Rename Template
-	Widget->Rename(*NewNameStr);
-
-	// Rename Preview
-	UWidget* WidgetPreview = FWidgetReference::FromTemplate(BlueprintEditor, Widget).GetPreview();
-	if ( WidgetPreview )
+	if ( Widget )
 	{
-		WidgetPreview->Rename(*NewNameStr);
-	}
-
-	// TODO UMG RENAME VARIABLES IN THE GRAPH!
-
-	UTimelineTemplate* Template = Blueprint->FindTimelineTemplateByVariableName(NewName);
-	if ((Template == NULL) && bUniqueNameForTemplate)
-	{
-		Template = Blueprint->FindTimelineTemplateByVariableName(OldName);
-		if (Template)
+		// Rename Template
+		Blueprint->Modify();
+		Widget->Modify();
+		Widget->Rename(*NewNameStr);
+
+		// Rename Preview
+		UWidget* WidgetPreview = FWidgetReference::FromTemplate(BlueprintEditor, Widget).GetPreview();
+		if ( WidgetPreview )
 		{
-			Blueprint->Modify();
-			Template->Modify();
+			WidgetPreview->Rename(*NewNameStr);
+		}
 
-			TArray<UK2Node_Variable*> TimelineVarNodes;
-			FBlueprintEditorUtils::GetAllNodesOfClass<UK2Node_Variable>(Blueprint, TimelineVarNodes);
-			for(int32 It = 0; It < TimelineVarNodes.Num(); It++)
+		TArray<UK2Node_Variable*> WidgetVarNodes;
+		FBlueprintEditorUtils::GetAllNodesOfClass<UK2Node_Variable>(Blueprint, WidgetVarNodes);
+		for ( int32 It = 0; It < WidgetVarNodes.Num(); It++ )
+		{
+			UK2Node_Variable* TestNode = WidgetVarNodes[It];
+			if ( TestNode && ( OldName == TestNode->GetVarName() ) )
 			{
-				UK2Node_Variable* TestNode = TimelineVarNodes[It];
-				if(TestNode && (OldName == TestNode->GetVarName()))
+				UEdGraphPin* TestPin = TestNode->FindPin(OldNameStr);
+				if ( TestPin && ( Widget->GetClass() == TestPin->PinType.PinSubCategoryObject.Get() ) )
 				{
-					UEdGraphPin* TestPin = TestNode->FindPin(OldNameStr);
-					if(TestPin && (UTimelineComponent::StaticClass() == TestPin->PinType.PinSubCategoryObject.Get()))
+					TestNode->Modify();
+					if ( TestNode->VariableReference.IsSelfContext() )
 					{
-						TestNode->Modify();
-						if(TestNode->VariableReference.IsSelfContext())
-						{
-							TestNode->VariableReference.SetSelfMember(NewName);
-						}
-						else
-						{
-							//TODO:
-							UClass* ParentClass = TestNode->VariableReference.GetMemberParentClass((UClass*)NULL);
-							TestNode->VariableReference.SetExternalMember(NewName, ParentClass);
-						}
-						TestPin->Modify();
-						TestPin->PinName = NewNameStr;
+						TestNode->VariableReference.SetSelfMember(NewName);
 					}
+					else
+					{
+						//TODO:
+						UClass* ParentClass = TestNode->VariableReference.GetMemberParentClass((UClass*)NULL);
+						TestNode->VariableReference.SetExternalMember(NewName, ParentClass);
+					}
+					TestPin->Modify();
+					TestPin->PinName = NewNameStr;
 				}
 			}
+		}
 
-			Blueprint->Timelines.Remove(Template);
-			
-			UObject* ExistingObject = StaticFindObject(NULL, Template->GetOuter(), *NewName.ToString(), true);
-			if (ExistingObject != Template && ExistingObject != NULL)
-			{
-				ExistingObject->Rename(*MakeUniqueObjectName(ExistingObject->GetOuter(), ExistingObject->GetClass(), ExistingObject->GetFName()).ToString());
-			}
-			Template->Rename(*NewName.ToString(), Template->GetOuter(), ( Blueprint->bIsRegeneratingOnLoad ? REN_ForceNoResetLoaders : REN_None ));
-			Blueprint->Timelines.Add(Template);
+		// Validate child blueprints and adjust variable names to avoid a potential name collision
+		FBlueprintEditorUtils::ValidateBlueprintChildVariables(Blueprint, NewName);
 
-			// Validate child blueprints and adjust variable names to avoid a potential name collision
-			FBlueprintEditorUtils::ValidateBlueprintChildVariables(Blueprint, NewName);
+		// Refresh references and flush editors
+		FBlueprintEditorUtils::MarkBlueprintAsStructurallyModified(Blueprint);
+		bRenamed = true;
 
-			// Refresh references and flush editors
-			FBlueprintEditorUtils::MarkBlueprintAsStructurallyModified(Blueprint);
-			bRenamed = true;
-		}
+		// Refresh references and flush editors
+		FBlueprintEditorUtils::MarkBlueprintAsStructurallyModified(Blueprint);
+		bRenamed = true;
 	}
 
-	// Refresh references and flush editors
-	FBlueprintEditorUtils::MarkBlueprintAsStructurallyModified(Blueprint);
-	bRenamed = true;
-
 	return bRenamed;
 }
 
