UMG - When moving objects between contains we now attempt to copy all relevant properties into the new container.  Making slots more robust when you remove items, it attempts to cleanup references.

diff --git a/Engine/Source/Editor/UMGEditor/Private/Extensions/CanvasSlotExtension.cpp b/Engine/Source/Editor/UMGEditor/Private/Extensions/CanvasSlotExtension.cpp
--- a/Engine/Source/Editor/UMGEditor/Private/Extensions/CanvasSlotExtension.cpp
+++ b/Engine/Source/Editor/UMGEditor/Private/Extensions/CanvasSlotExtension.cpp
@@ -115,75 +115,76 @@ void FCanvasSlotExtension::Paint(const TSet< FWidgetReference >& Selection, cons
 
 		if ( UCanvasPanelSlot* CanvasSlot = Cast<UCanvasPanelSlot>(Widget->Slot) )
 		{
-			UCanvasPanel* Canvas = CastChecked<UCanvasPanel>(CanvasSlot->Parent);
-
-			//TODO UMG Only show guide lines when near them and dragging
-			if ( bDragging )
+			if ( UCanvasPanel* Canvas = CastChecked<UCanvasPanel>(CanvasSlot->Parent) )
 			{
-				FGeometry MyArrangedGeometry;
-				if ( !Canvas->GetGeometryForSlot(CanvasSlot, MyArrangedGeometry) )
+				//TODO UMG Only show guide lines when near them and dragging
+				if ( bDragging )
 				{
-					continue;
-				}
+					FGeometry MyArrangedGeometry;
+					if ( !Canvas->GetGeometryForSlot(CanvasSlot, MyArrangedGeometry) )
+					{
+						continue;
+					}
 
-				TArray<FVector2D> LinePoints;
-				LinePoints.AddUninitialized(2);
+					TArray<FVector2D> LinePoints;
+					LinePoints.AddUninitialized(2);
 
-				// Get the collision segments that we could potentially be docking against.
-				TArray<FVector2D> MySegments;
-				if ( GetCollisionSegmentsForSlot(Canvas, CanvasSlot, MySegments) )
-				{
-					for ( int32 MySegmentIndex = 0; MySegmentIndex < MySegments.Num(); MySegmentIndex += 2 )
+					// Get the collision segments that we could potentially be docking against.
+					TArray<FVector2D> MySegments;
+					if ( GetCollisionSegmentsForSlot(Canvas, CanvasSlot, MySegments) )
 					{
-						FVector2D MySegmentBase = MySegments[MySegmentIndex];
-
-						for ( int32 SlotIndex = 0; SlotIndex < Canvas->GetChildrenCount(); SlotIndex++ )
+						for ( int32 MySegmentIndex = 0; MySegmentIndex < MySegments.Num(); MySegmentIndex += 2 )
 						{
-							// Ignore the slot being dragged.
-							if ( Canvas->Slots[SlotIndex] == CanvasSlot )
-							{
-								continue;
-							}
+							FVector2D MySegmentBase = MySegments[MySegmentIndex];
 
-							// Get the collision segments that we could potentially be docking against.
-							TArray<FVector2D> Segments;
-							if ( GetCollisionSegmentsForSlot(Canvas, SlotIndex, Segments) )
+							for ( int32 SlotIndex = 0; SlotIndex < Canvas->GetChildrenCount(); SlotIndex++ )
 							{
-								for ( int32 SegmentBase = 0; SegmentBase < Segments.Num(); SegmentBase += 2 )
+								// Ignore the slot being dragged.
+								if ( Canvas->Slots[SlotIndex] == CanvasSlot )
 								{
-									FVector2D PointA = Segments[SegmentBase];
-									FVector2D PointB = Segments[SegmentBase + 1];
-
-									FVector2D CollisionPoint = MySegmentBase;
+									continue;
+								}
 
-									//TODO Collide against all sides of the arranged geometry.
-									float Distance = DistancePointToLine2D(PointA, PointB, CollisionPoint);
-									if ( Distance <= SnapDistance )
+								// Get the collision segments that we could potentially be docking against.
+								TArray<FVector2D> Segments;
+								if ( GetCollisionSegmentsForSlot(Canvas, SlotIndex, Segments) )
+								{
+									for ( int32 SegmentBase = 0; SegmentBase < Segments.Num(); SegmentBase += 2 )
 									{
-										FVector2D FarthestPoint = FVector2D::Distance(PointA, CollisionPoint) > FVector2D::Distance(PointB, CollisionPoint) ? PointA : PointB;
-										FVector2D NearestPoint = FVector2D::Distance(PointA, CollisionPoint) > FVector2D::Distance(PointB, CollisionPoint) ? PointB : PointA;
-
-										LinePoints[0] = FarthestPoint;
-										LinePoints[1] = FarthestPoint + ( NearestPoint - FarthestPoint ) * 100000;
-
-										LinePoints[0].X = FMath::Clamp(LinePoints[0].X, 0.0f, MyClippingRect.Right - MyClippingRect.Left);
-										LinePoints[0].Y = FMath::Clamp(LinePoints[0].Y, 0.0f, MyClippingRect.Bottom - MyClippingRect.Top);
-
-										LinePoints[1].X = FMath::Clamp(LinePoints[1].X, 0.0f, MyClippingRect.Right - MyClippingRect.Left);
-										LinePoints[1].Y = FMath::Clamp(LinePoints[1].Y, 0.0f, MyClippingRect.Bottom - MyClippingRect.Top);
-
-										FLinearColor Color(0.5f, 0.75, 1);
-										const bool bAntialias = true;
-
-										FSlateDrawElement::MakeLines(
-											OutDrawElements,
-											LayerId,
-											AllottedGeometry.ToPaintGeometry(),
-											LinePoints,
-											MyClippingRect,
-											ESlateDrawEffect::None,
-											Color,
-											bAntialias);
+										FVector2D PointA = Segments[SegmentBase];
+										FVector2D PointB = Segments[SegmentBase + 1];
+
+										FVector2D CollisionPoint = MySegmentBase;
+
+										//TODO Collide against all sides of the arranged geometry.
+										float Distance = DistancePointToLine2D(PointA, PointB, CollisionPoint);
+										if ( Distance <= SnapDistance )
+										{
+											FVector2D FarthestPoint = FVector2D::Distance(PointA, CollisionPoint) > FVector2D::Distance(PointB, CollisionPoint) ? PointA : PointB;
+											FVector2D NearestPoint = FVector2D::Distance(PointA, CollisionPoint) > FVector2D::Distance(PointB, CollisionPoint) ? PointB : PointA;
+
+											LinePoints[0] = FarthestPoint;
+											LinePoints[1] = FarthestPoint + ( NearestPoint - FarthestPoint ) * 100000;
+
+											LinePoints[0].X = FMath::Clamp(LinePoints[0].X, 0.0f, MyClippingRect.Right - MyClippingRect.Left);
+											LinePoints[0].Y = FMath::Clamp(LinePoints[0].Y, 0.0f, MyClippingRect.Bottom - MyClippingRect.Top);
+
+											LinePoints[1].X = FMath::Clamp(LinePoints[1].X, 0.0f, MyClippingRect.Right - MyClippingRect.Left);
+											LinePoints[1].Y = FMath::Clamp(LinePoints[1].Y, 0.0f, MyClippingRect.Bottom - MyClippingRect.Top);
+
+											FLinearColor Color(0.5f, 0.75, 1);
+											const bool bAntialias = true;
+
+											FSlateDrawElement::MakeLines(
+												OutDrawElements,
+												LayerId,
+												AllottedGeometry.ToPaintGeometry(),
+												LinePoints,
+												MyClippingRect,
+												ESlateDrawEffect::None,
+												Color,
+												bAntialias);
+										}
 									}
 								}
 							}
diff --git a/Engine/Source/Editor/UMGEditor/Private/SUMGDesigner.cpp b/Engine/Source/Editor/UMGEditor/Private/SUMGDesigner.cpp
--- a/Engine/Source/Editor/UMGEditor/Private/SUMGDesigner.cpp
+++ b/Engine/Source/Editor/UMGEditor/Private/SUMGDesigner.cpp
@@ -20,6 +20,8 @@ class FSelectedWidgetDragDropOp : public FDecoratedDragDropOp
 public:
 	DRAG_DROP_OPERATOR_TYPE(FSelectedWidgetDragDropOp, FDecoratedDragDropOp)
 
+	TMap<FName, FString> ExportedSlotProperties;
+
 	FWidgetReference Widget;
 
 	static TSharedRef<FSelectedWidgetDragDropOp> New(FWidgetReference InWidget);
@@ -35,6 +37,8 @@ TSharedRef<FSelectedWidgetDragDropOp> FSelectedWidgetDragDropOp::New(FWidgetRefe
 	Operation->CurrentHoverText = FText::FromString( InWidget.GetTemplate()->GetLabel() );
 	Operation->Construct();
 
+	FWidgetBlueprintEditorUtils::ExportPropertiesToText(InWidget.GetTemplate()->Slot, Operation->ExportedSlotProperties);
+
 	return Operation;
 }
 
@@ -940,11 +944,23 @@ UWidget* SUMGDesigner::ProcessDropAndAddWidget(const FGeometry& MyGeometry, cons
 		DropPreviewWidget = SelectedDragDropOp->Widget.GetPreview();
 	}
 
+	UWidgetBlueprint* BP = GetBlueprint();
+
+	if ( DropPreviewWidget )
+	{
+		if ( DropPreviewParent )
+		{
+			DropPreviewParent->RemoveChild(DropPreviewWidget);
+		}
+
+		const bool bIsRecursive = false;
+		BP->WidgetTree->RemoveWidget(DropPreviewWidget, bIsRecursive);
+		DropPreviewWidget = NULL;
+	}
+
 	FArrangedWidget ArrangedWidget(SNullWidget::NullWidget, FGeometry());
 	FWidgetReference WidgetUnderCursor = GetWidgetAtCursor(MyGeometry, DragDropEvent, ArrangedWidget);
 
-	UWidgetBlueprint* BP = GetBlueprint();
-
 	UWidget* Target = NULL;
 	if ( WidgetUnderCursor.IsValid() )
 	{
@@ -1079,9 +1095,11 @@ UWidget* SUMGDesigner::ProcessDropAndAddWidget(const FGeometry& MyGeometry, cons
 				FVector2D LocalPosition = ArrangedWidget.Geometry.AbsoluteToLocal(DragDropEvent.GetScreenSpacePosition());
 				if ( UPanelSlot* Slot = NewParent->AddChild(Widget) )
 				{
+					FWidgetBlueprintEditorUtils::ImportPropertiesFromText(Slot, SelectedDragDropOp->ExportedSlotProperties);
+
 					//TODO UMG Migrate existing slot info
 					Slot->SetDesiredPosition(LocalPosition);
-					Slot->SetDesiredSize(FVector2D(150, 30));
+					//Slot->SetDesiredSize(FVector2D(150, 30));
 					//@TODO UMG When we add a child blindly we need to default the slot size to the preferred size of the widget if the container supports such things.
 					//@TODO UMG We may need a desired size canvas, where the slots have no size, they only give you position, alternatively, maybe slots that don't clip, so center is still easy.
 
diff --git a/Engine/Source/Editor/UMGEditor/Private/WidgetBlueprintEditorUtils.cpp b/Engine/Source/Editor/UMGEditor/Private/WidgetBlueprintEditorUtils.cpp
--- a/Engine/Source/Editor/UMGEditor/Private/WidgetBlueprintEditorUtils.cpp
+++ b/Engine/Source/Editor/UMGEditor/Private/WidgetBlueprintEditorUtils.cpp
@@ -320,4 +320,39 @@ void FWidgetBlueprintEditorUtils::ImportWidgetsFromText(UWidgetBlueprint* BP, co
 	}
 }
 
+void FWidgetBlueprintEditorUtils::ExportPropertiesToText(UObject* Object, TMap<FName, FString>& ExportedProperties)
+{
+	if ( Object )
+	{
+		for ( TFieldIterator<UProperty> PropertyIt(Object->GetClass(), EFieldIteratorFlags::ExcludeSuper); PropertyIt; ++PropertyIt )
+		{
+			UProperty* Property = *PropertyIt;
+
+			// Don't serialize out object properties, we just want value data.
+			if ( !Property->IsA<UObjectProperty>() )
+			{
+				FString ValueText;
+				if ( Property->ExportText_InContainer(0, ValueText, Object, Object, Object, PPF_IncludeTransient) )
+				{
+					ExportedProperties.Add(Property->GetFName(), ValueText);
+				}
+			}
+		}
+	}
+}
+
+void FWidgetBlueprintEditorUtils::ImportPropertiesFromText(UObject* Object, const TMap<FName, FString>& ExportedProperties)
+{
+	if ( Object )
+	{
+		for ( const auto& Entry : ExportedProperties )
+		{
+			if ( UProperty* Property = FindField<UProperty>(Object->GetClass(), Entry.Key) )
+			{
+				Property->ImportText(*Entry.Value, Property->ContainerPtrToValuePtr<uint8>(Object), 0, Object);
+			}
+		}
+	}
+}
+
 #undef LOCTEXT_NAMESPACE
diff --git a/Engine/Source/Editor/UMGEditor/Private/WidgetBlueprintEditorUtils.h b/Engine/Source/Editor/UMGEditor/Private/WidgetBlueprintEditorUtils.h
--- a/Engine/Source/Editor/UMGEditor/Private/WidgetBlueprintEditorUtils.h
+++ b/Engine/Source/Editor/UMGEditor/Private/WidgetBlueprintEditorUtils.h
@@ -23,6 +23,12 @@ public:
 
 	static void ImportWidgetsFromText(UWidgetBlueprint* BP, const FString& TextToImport, /*out*/ TSet<UWidget*>& ImportedWidgetSet);
 
+	/** Exports the individual properties of an object to text and stores them in a map. */
+	static void ExportPropertiesToText(UObject* Object, TMap<FName, FString>& ExportedProperties);
+
+	/** Attempts to import any property in the map and apply it to a property with the same name on the object. */
+	static void ImportPropertiesFromText(UObject* Object, const TMap<FName, FString>& ExportedProperties);
+
 private:
 	static void BuildWrapWithMenu(FMenuBuilder& Menu, UWidgetBlueprint* BP, TSet<FWidgetReference> Widgets);
 
diff --git a/Engine/Source/Runtime/UMG/Private/Components/PanelWidget.cpp b/Engine/Source/Runtime/UMG/Private/Components/PanelWidget.cpp
--- a/Engine/Source/Runtime/UMG/Private/Components/PanelWidget.cpp
+++ b/Engine/Source/Runtime/UMG/Private/Components/PanelWidget.cpp
@@ -45,6 +45,12 @@ int32 UPanelWidget::GetChildIndex(UWidget* Content) const
 bool UPanelWidget::RemoveChildAt(int32 Index)
 {
 	UPanelSlot* Slot = Slots[Index];
+	if ( Slot->Content )
+	{
+		Slot->Content->Slot = NULL;
+	}
+
+	Slot->Parent = NULL;
 	Slots.RemoveAt(Index);
 
 	OnSlotRemoved(Slot);
@@ -64,7 +70,10 @@ UPanelSlot* UPanelWidget::AddChild(UWidget* Content)
 	Slot->Content = Content;
 	Slot->Parent = this;
 
-	Content->Slot = Slot;
+	if ( Content )
+	{
+		Content->Slot = Slot;
+	}
 
 	Slots.Add(Slot);
 
@@ -78,7 +87,10 @@ void UPanelWidget::ReplaceChildAt(int32 Index, UWidget* Content)
 	UPanelSlot* Slot = Slots[Index];
 	Slot->Content = Content;
 
-	Content->Slot = Slot;
+	if ( Content )
+	{
+		Content->Slot = Slot;
+	}
 
 	Slot->SyncronizeProperties();
 }
@@ -90,7 +102,10 @@ void UPanelWidget::InsertChildAt(int32 Index, UWidget* Content)
 	Slot->Content = Content;
 	Slot->Parent = this;
 
-	Content->Slot = Slot;
+	if ( Content )
+	{
+		Content->Slot = Slot;
+	}
 
 	Slots.Insert(Slot, Index);
 
