UMG - No more layout flash.  Fixing crashes.  Dragged objects now maintain their relative offset from the cursor when the drag started.

diff --git a/Engine/Source/Editor/UMGEditor/Private/SUMGDesigner.cpp b/Engine/Source/Editor/UMGEditor/Private/SUMGDesigner.cpp
--- a/Engine/Source/Editor/UMGEditor/Private/SUMGDesigner.cpp
+++ b/Engine/Source/Editor/UMGEditor/Private/SUMGDesigner.cpp
@@ -197,7 +197,10 @@ void SUMGDesigner::CreateExtensionWidgetsForSelection()
 	ExtensionWidgets.Reset();
 
 	TArray<FWidgetReference> Selected;
-	Selected.Add(SelectedWidget);
+	if ( SelectedWidget.IsValid() )
+	{
+		Selected.Add(SelectedWidget);
+	}
 
 	// Build extension widgets for new selection
 	for ( TSharedRef<FDesignerExtension>& Ext : DesignerExtensions )
@@ -331,7 +334,10 @@ FReply SUMGDesigner::OnMouseButtonDown(const FGeometry& MyGeometry, const FPoint
 		SelectedTemplates.Add(SelectedWidget);
 		BlueprintEditor.Pin()->SelectWidgets(SelectedTemplates);
 
-		bMouseDown = true;
+		if ( MouseEvent.GetEffectingButton() == EKeys::LeftMouseButton )
+		{
+			bMouseDown = true;
+		}
 	}
 
 	// Capture mouse for the drag handle and general mouse actions
@@ -397,9 +403,13 @@ FReply SUMGDesigner::OnMouseMove(const FGeometry& MyGeometry, const FPointerEven
 	{
 		if ( SelectedWidget.IsValid() && !bMovingExistingWidget )
 		{
-			bMovingExistingWidget = true;
-			//Drag selected widgets
-			return FReply::Handled().DetectDrag(AsShared(), EKeys::LeftMouseButton);
+			const bool bIsRootWidget = SelectedWidget.GetTemplate()->GetParent() == NULL;
+			if ( !bIsRootWidget )
+			{
+				bMovingExistingWidget = true;
+				//Drag selected widgets
+				return FReply::Handled().DetectDrag(AsShared(), EKeys::LeftMouseButton);
+			}
 		}
 	}
 	
@@ -819,6 +829,7 @@ void SUMGDesigner::Tick(const FGeometry& AllottedGeometry, const double InCurren
 		if ( PreviewWidget )
 		{
 			TSharedRef<SWidget> CurrentWidget = PreviewWidget->MakeFullScreenWidget();
+			CurrentWidget->SlatePrepass();
 
 			if ( CurrentWidget != PreviewSlateWidget.Pin() )
 			{
@@ -836,7 +847,7 @@ void SUMGDesigner::Tick(const FGeometry& AllottedGeometry, const double InCurren
 				.VAlign(VAlign_Center)
 				[
 					SNew(STextBlock)
-					.Text(LOCTEXT("NoWrappedWidget", "No Widget Preview"))
+					.Text(LOCTEXT("NoWidgetPreview", "No Widget Preview"))
 				]
 			];
 		}
@@ -1098,7 +1109,7 @@ UWidget* SUMGDesigner::ProcessDropAndAddWidget(const FGeometry& MyGeometry, cons
 					FWidgetBlueprintEditorUtils::ImportPropertiesFromText(Slot, SelectedDragDropOp->ExportedSlotProperties);
 
 					//TODO UMG Migrate existing slot info
-					Slot->SetDesiredPosition(LocalPosition);
+					Slot->SetDesiredPosition(LocalPosition - SelectedWidgetContextMenuLocation);
 					//Slot->SetDesiredSize(FVector2D(150, 30));
 					//@TODO UMG When we add a child blindly we need to default the slot size to the preferred size of the widget if the container supports such things.
 					//@TODO UMG We may need a desired size canvas, where the slots have no size, they only give you position, alternatively, maybe slots that don't clip, so center is still easy.
