UMG - New layout for UMG.

diff --git a/Engine/Source/Editor/UMGEditor/Private/UMGBlueprintEditorExtensionHook.cpp b/Engine/Source/Editor/UMGEditor/Private/UMGBlueprintEditorExtensionHook.cpp
--- a/Engine/Source/Editor/UMGEditor/Private/UMGBlueprintEditorExtensionHook.cpp
+++ b/Engine/Source/Editor/UMGEditor/Private/UMGBlueprintEditorExtensionHook.cpp
@@ -186,7 +186,7 @@ public:
 		BlueprintDefaultsTabFactories.RegisterFactory(MakeShareable(new FWidgetTemplatesSummoner(WidgetBlueprintEditor)));
 		BlueprintDefaultsTabFactories.RegisterFactory(MakeShareable(new FSequencerSummoner(WidgetBlueprintEditor)));
 
-		TabLayout = FTabManager::NewLayout( "Standalone_UMGEditor_Layout_v2" )
+		TabLayout = FTabManager::NewLayout( "Standalone_UMGEditor_Layout_v3" )
 		->AddArea
 		(
 			FTabManager::NewPrimaryArea()
@@ -213,12 +213,6 @@ public:
 						->SetSizeCoefficient( 0.5f )
 						->AddTab( UMGWidgetTemplatesTabID, ETabState::OpenedTab )
 					)
-				)
-				->Split
-				(
-					FTabManager::NewSplitter()
-					->SetSizeCoefficient( 0.15f )
-					->SetOrientation(Orient_Vertical)
 					->Split
 					(
 						FTabManager::NewStack()
@@ -229,7 +223,7 @@ public:
 				->Split
 				(
 					FTabManager::NewSplitter()
-					->SetSizeCoefficient( 0.45f )
+					->SetSizeCoefficient( 0.85f )
 					->SetOrientation(Orient_Horizontal)
 					->Split
 					(
