UMG - Fixing a bug with not being able to drag a widget multiple times in a row without reselecting it.

diff --git a/Engine/Source/Editor/UMGEditor/Private/SUMGDesigner.cpp b/Engine/Source/Editor/UMGEditor/Private/SUMGDesigner.cpp
--- a/Engine/Source/Editor/UMGEditor/Private/SUMGDesigner.cpp
+++ b/Engine/Source/Editor/UMGEditor/Private/SUMGDesigner.cpp
@@ -1142,6 +1142,9 @@ UWidget* SUMGDesigner::ProcessDropAndAddWidget(const FGeometry& MyGeometry, cons
 
 FReply SUMGDesigner::OnDrop(const FGeometry& MyGeometry, const FDragDropEvent& DragDropEvent)
 {
+	bMouseDown = false;
+	bMovingExistingWidget = false;
+
 	UWidgetBlueprint* BP = GetBlueprint();
 	
 	if (DropPreviewWidget)
