UMG - Fixing some undo bugs with moving / adding widgets.

diff --git a/Engine/Source/Editor/UMGEditor/Private/SUMGDesigner.cpp b/Engine/Source/Editor/UMGEditor/Private/SUMGDesigner.cpp
--- a/Engine/Source/Editor/UMGEditor/Private/SUMGDesigner.cpp
+++ b/Engine/Source/Editor/UMGEditor/Private/SUMGDesigner.cpp
@@ -1011,7 +1011,7 @@ UWidget* SUMGDesigner::ProcessDropAndAddWidget(const FGeometry& MyGeometry, cons
 
 					DropPreviewParent = Parent;
 
-					if ( !bIsPreview )
+					if ( bIsPreview )
 					{
 						Transaction.Cancel();
 					}
@@ -1024,7 +1024,7 @@ UWidget* SUMGDesigner::ProcessDropAndAddWidget(const FGeometry& MyGeometry, cons
 					//          Maybe we can traverse the hierarchy and add it to the first parent that will accept it?
 				}
 
-				if ( !bIsPreview )
+				if ( bIsPreview )
 				{
 					Transaction.Cancel();
 				}
@@ -1061,7 +1061,7 @@ UWidget* SUMGDesigner::ProcessDropAndAddWidget(const FGeometry& MyGeometry, cons
 
 				DropPreviewParent = NULL;
 
-				if ( !bIsPreview )
+				if ( bIsPreview )
 				{
 					Transaction.Cancel();
 				}
@@ -1116,7 +1116,7 @@ UWidget* SUMGDesigner::ProcessDropAndAddWidget(const FGeometry& MyGeometry, cons
 
 					DropPreviewParent = NewParent;
 
-					if ( !bIsPreview )
+					if ( bIsPreview )
 					{
 						Transaction.Cancel();
 					}
@@ -1129,7 +1129,7 @@ UWidget* SUMGDesigner::ProcessDropAndAddWidget(const FGeometry& MyGeometry, cons
 					//          Maybe we can traverse the hierarchy and add it to the first parent that will accept it?
 				}
 
-				if ( !bIsPreview )
+				if ( bIsPreview )
 				{
 					Transaction.Cancel();
 				}
