UMG - Cleaning up left ofter hacks from before widget tree was a tree and removal was much more complicated.

diff --git a/Engine/Source/Editor/UMGEditor/Private/SUMGDesigner.cpp b/Engine/Source/Editor/UMGEditor/Private/SUMGDesigner.cpp
--- a/Engine/Source/Editor/UMGEditor/Private/SUMGDesigner.cpp
+++ b/Engine/Source/Editor/UMGEditor/Private/SUMGDesigner.cpp
@@ -912,8 +912,7 @@ void SUMGDesigner::OnDragLeave(const FDragDropEvent& DragDropEvent)
 		}
 
 		UWidgetBlueprint* BP = GetBlueprint();
-		const bool bIsRecursive = false;
-		BP->WidgetTree->RemoveWidget(DropPreviewWidget, bIsRecursive);
+		BP->WidgetTree->RemoveWidget(DropPreviewWidget);
 		DropPreviewWidget = NULL;
 	}
 }
@@ -929,8 +928,7 @@ FReply SUMGDesigner::OnDragOver(const FGeometry& MyGeometry, const FDragDropEven
 			DropPreviewParent->RemoveChild(DropPreviewWidget);
 		}
 		
-		const bool bIsRecursive = false;
-		BP->WidgetTree->RemoveWidget(DropPreviewWidget, bIsRecursive);
+		BP->WidgetTree->RemoveWidget(DropPreviewWidget);
 		DropPreviewWidget = NULL;
 	}
 	
@@ -964,8 +962,7 @@ UWidget* SUMGDesigner::ProcessDropAndAddWidget(const FGeometry& MyGeometry, cons
 			DropPreviewParent->RemoveChild(DropPreviewWidget);
 		}
 
-		const bool bIsRecursive = false;
-		BP->WidgetTree->RemoveWidget(DropPreviewWidget, bIsRecursive);
+		BP->WidgetTree->RemoveWidget(DropPreviewWidget);
 		DropPreviewWidget = NULL;
 	}
 
@@ -1156,8 +1153,7 @@ FReply SUMGDesigner::OnDrop(const FGeometry& MyGeometry, const FDragDropEvent& D
 			DropPreviewParent->RemoveChild(DropPreviewWidget);
 		}
 		
-		const bool bIsRecursive = false;
-		BP->WidgetTree->RemoveWidget(DropPreviewWidget, bIsRecursive);
+		BP->WidgetTree->RemoveWidget(DropPreviewWidget);
 		DropPreviewWidget = NULL;
 	}
 	
diff --git a/Engine/Source/Runtime/UMG/Private/WidgetTree.cpp b/Engine/Source/Runtime/UMG/Private/WidgetTree.cpp
--- a/Engine/Source/Runtime/UMG/Private/WidgetTree.cpp
+++ b/Engine/Source/Runtime/UMG/Private/WidgetTree.cpp
@@ -45,7 +45,7 @@ UPanelWidget* UWidgetTree::FindWidgetParent(UWidget* Widget, int32& OutChildInde
 	return Parent;
 }
 
-bool UWidgetTree::RemoveWidget(UWidget* InRemovedWidget, bool bIsRecursive)
+bool UWidgetTree::RemoveWidget(UWidget* InRemovedWidget)
 {
 	InRemovedWidget->Modify();
 
@@ -62,15 +62,10 @@ bool UWidgetTree::RemoveWidget(UWidget* InRemovedWidget, bool bIsRecursive)
 		}
 	}
 
-	if ( bIsRecursive )
+	if ( InRemovedWidget == RootWidget )
 	{
-		if ( UPanelWidget* InNonLeafRemovedWidget = Cast<UPanelWidget>(InRemovedWidget) )
-		{
-			while ( InNonLeafRemovedWidget->GetChildrenCount() > 0 )
-			{
-				RemoveWidget(InNonLeafRemovedWidget->GetChildAt(0));
-			}
-		}
+		RootWidget = NULL;
+		bRemoved = true;
 	}
 
 	return bRemoved;
diff --git a/Engine/Source/Runtime/UMG/Public/Blueprint/WidgetTree.h b/Engine/Source/Runtime/UMG/Public/Blueprint/WidgetTree.h
--- a/Engine/Source/Runtime/UMG/Public/Blueprint/WidgetTree.h
+++ b/Engine/Source/Runtime/UMG/Public/Blueprint/WidgetTree.h
@@ -19,7 +19,7 @@ public:
 
 	UWidget* FindWidget(const FString& Name) const;
 
-	bool RemoveWidget(UWidget* Widget, bool bIsRecursive = true);
+	bool RemoveWidget(UWidget* Widget);
 
 	class UPanelWidget* FindWidgetParent(UWidget* Widget, int32& OutChildIndex);
 
