UMG - Adding some more code doc.  Moving WidgetTemplate into a publically accessible area, though still no public way to register new ones yet.  Dragging and dropping new items onto the designer now takes into account the objects desired size. (clamped to a reasonable minimum to prevent 0,0 sized widgets).

diff --git a/Engine/Source/Editor/UMGEditor/Private/SUMGDesigner.cpp b/Engine/Source/Editor/UMGEditor/Private/SUMGDesigner.cpp
--- a/Engine/Source/Editor/UMGEditor/Private/SUMGDesigner.cpp
+++ b/Engine/Source/Editor/UMGEditor/Private/SUMGDesigner.cpp
@@ -25,8 +25,6 @@ public:
 	FWidgetReference Widget;
 
 	static TSharedRef<FSelectedWidgetDragDropOp> New(FWidgetReference InWidget);
-
-	//virtual TSharedPtr<SWidget> GetDefaultDecorator() const override;
 };
 
 TSharedRef<FSelectedWidgetDragDropOp> FSelectedWidgetDragDropOp::New(FWidgetReference InWidget)
@@ -975,90 +973,75 @@ UWidget* SUMGDesigner::ProcessDropAndAddWidget(const FGeometry& MyGeometry, cons
 		Target = bIsPreview ? WidgetUnderCursor.GetPreview() : WidgetUnderCursor.GetTemplate();
 	}
 
-	TSharedPtr<FWidgetTemplateDragDropOp> DragDropOp = DragDropEvent.GetOperationAs<FWidgetTemplateDragDropOp>();
-	if ( DragDropOp.IsValid() )
-	{		
-		if ( Target )
+	TSharedPtr<FWidgetTemplateDragDropOp> TemplateDragDropOp = DragDropEvent.GetOperationAs<FWidgetTemplateDragDropOp>();
+	if ( TemplateDragDropOp.IsValid() )
+	{
+		// If there's no root widget go ahead and add the widget into the root slot.
+		if ( BP->WidgetTree->RootWidget == NULL )
 		{
-			if ( Target->IsA(UPanelWidget::StaticClass()) )
+			FScopedTransaction Transaction(LOCTEXT("Designer_AddWidget", "Add Widget"));
+
+			if ( !bIsPreview )
 			{
-				UPanelWidget* Parent = Cast<UPanelWidget>(Target);
+				BP->WidgetTree->SetFlags(RF_Transactional);
+				BP->WidgetTree->Modify();
+			}
 
-				FScopedTransaction Transaction(LOCTEXT("Designer_AddWidget", "Add Widget"));
+			// TODO UMG This method isn't great, maybe the user widget should just be a canvas.
 
-				if ( !bIsPreview )
-				{
-					Parent->SetFlags(RF_Transactional);
-					Parent->Modify();
+			// Add it to the root if there are no other widgets to add it to.
+			UWidget* Widget = TemplateDragDropOp->Template->Create(BP->WidgetTree);
+			Widget->IsDesignTime(true);
 
-					BP->WidgetTree->SetFlags(RF_Transactional);
-					BP->WidgetTree->Modify();
-				}
-				
-				UWidget* Widget = DragDropOp->Template->Create(BP->WidgetTree);
-				Widget->IsDesignTime(true);
-			
-				FVector2D LocalPosition = ArrangedWidget.Geometry.AbsoluteToLocal(DragDropEvent.GetScreenSpacePosition());
-				if ( UPanelSlot* Slot = Parent->AddChild(Widget) )
-				{
-					Slot->SetDesiredPosition(LocalPosition);
-					Slot->SetDesiredSize(FVector2D(150, 30));
-					//@TODO UMG When we add a child blindly we need to default the slot size to the preferred size of the widget if the container supports such things.
-					//@TODO UMG We may need a desired size canvas, where the slots have no size, they only give you position, alternatively, maybe slots that don't clip, so center is still easy.
-
-					DropPreviewParent = Parent;
+			BP->WidgetTree->RootWidget = Widget;
 
-					if ( bIsPreview )
-					{
-						Transaction.Cancel();
-					}
+			SelectedWidget = FWidgetReference::FromTemplate(BlueprintEditor.Pin(), Widget);
 
-					return Widget;
-				}
-				else
-				{
-					// TODO UMG ERROR Slot can not be created because maybe the max children has been reached.
-					//          Maybe we can traverse the hierarchy and add it to the first parent that will accept it?
-				}
+			DropPreviewParent = NULL;
 
-				if ( bIsPreview )
-				{
-					Transaction.Cancel();
-				}
+			if ( bIsPreview )
+			{
+				Transaction.Cancel();
 			}
-			//else if ( BP->WidgetTree->RootWidget != NULL && BP->WidgetTree->RootWidget.Child == 1 )
-			//{
-			//	UWidget* Widget = DragDropOp->Template->Create(BP->WidgetTree);
-			//	Widget->IsDesignTime(true);
-			//	
-			//	DropPreviewParent = NULL;
-			//	
-			//	return Widget;
-			//}
+
+			return Widget;
 		}
-		else
+		// If there's already a root widget we need to try and place our widget into a parent widget that we've picked against
+		else if ( Target && Target->IsA(UPanelWidget::StaticClass()) )
 		{
-			if ( BP->WidgetTree->RootWidget == NULL )
+			UPanelWidget* Parent = Cast<UPanelWidget>(Target);
+
+			FScopedTransaction Transaction(LOCTEXT("Designer_AddWidget", "Add Widget"));
+
+			// If this isn't a preview operation we need to modify a few things to properly undo the operation.
+			if ( !bIsPreview )
 			{
-				FScopedTransaction Transaction(LOCTEXT("Designer_AddWidget", "Add Widget"));
+				Parent->SetFlags(RF_Transactional);
+				Parent->Modify();
 
-				if ( !bIsPreview )
-				{
-					BP->WidgetTree->SetFlags(RF_Transactional);
-					BP->WidgetTree->Modify();
-				}
+				BP->WidgetTree->SetFlags(RF_Transactional);
+				BP->WidgetTree->Modify();
+			}
 
-				// TODO UMG This method isn't great, maybe the user widget should just be a canvas.
+			// Construct the widget and mark it for design time rendering.
+			UWidget* Widget = TemplateDragDropOp->Template->Create(BP->WidgetTree);
+			Widget->IsDesignTime(true);
 
-				// Add it to the root if there are no other widgets to add it to.
-				UWidget* Widget = DragDropOp->Template->Create(BP->WidgetTree);
-				Widget->IsDesignTime(true);
+			// Determine local position inside the parent widget and add the widget to the slot.
+			FVector2D LocalPosition = ArrangedWidget.Geometry.AbsoluteToLocal(DragDropEvent.GetScreenSpacePosition());
+			if ( UPanelSlot* Slot = Parent->AddChild(Widget) )
+			{
+				TSharedPtr<SWidget> SlateWidget = Widget->GetWidget();
+				SlateWidget->SlatePrepass();
+				const FVector2D& WidgetDesiredSize = SlateWidget->GetDesiredSize();
 
-				BP->WidgetTree->RootWidget = Widget;
+				static const FVector2D MinimumDefaultSize(20, 20);
+				FVector2D LocalSize = FVector2D(FMath::Max(WidgetDesiredSize.X, MinimumDefaultSize.X), FMath::Max(WidgetDesiredSize.Y, MinimumDefaultSize.Y));
 
-				SelectedWidget = FWidgetReference::FromTemplate(BlueprintEditor.Pin(), Widget);
+				Slot->SetDesiredPosition(LocalPosition);
+				Slot->SetDesiredSize(LocalSize);
 
-				DropPreviewParent = NULL;
+				DropPreviewParent = Parent;
 
 				if ( bIsPreview )
 				{
@@ -1067,71 +1050,79 @@ UWidget* SUMGDesigner::ProcessDropAndAddWidget(const FGeometry& MyGeometry, cons
 
 				return Widget;
 			}
+			else
+			{
+				// TODO UMG ERROR Slot can not be created because maybe the max children has been reached.
+				//          Maybe we can traverse the hierarchy and add it to the first parent that will accept it?
+			}
+
+			if ( bIsPreview )
+			{
+				Transaction.Cancel();
+			}
 		}
 	}
 
 	// Attempt to deal with moving widgets from a drag operation.
 	if ( SelectedDragDropOp.IsValid() )
 	{
-		if ( Target )
+		if ( Target && Target->IsA(UPanelWidget::StaticClass()) )
 		{
-			if ( Target->IsA(UPanelWidget::StaticClass()) )
-			{
-				UPanelWidget* NewParent = Cast<UPanelWidget>(Target);
+			UPanelWidget* NewParent = Cast<UPanelWidget>(Target);
 
-				FScopedTransaction Transaction(LOCTEXT("Designer_MoveWidget", "Move Widget"));
+			FScopedTransaction Transaction(LOCTEXT("Designer_MoveWidget", "Move Widget"));
 
-				if ( !bIsPreview )
-				{
-					NewParent->SetFlags(RF_Transactional);
-					NewParent->Modify();
+			// If this isn't a preview operation we need to modify a few things to properly undo the operation.
+			if ( !bIsPreview )
+			{
+				NewParent->SetFlags(RF_Transactional);
+				NewParent->Modify();
 
-					BP->WidgetTree->SetFlags(RF_Transactional);
-					BP->WidgetTree->Modify();
-				}
+				BP->WidgetTree->SetFlags(RF_Transactional);
+				BP->WidgetTree->Modify();
+			}
 
-				UWidget* Widget = bIsPreview ? SelectedDragDropOp->Widget.GetPreview() : SelectedDragDropOp->Widget.GetTemplate();
+			UWidget* Widget = bIsPreview ? SelectedDragDropOp->Widget.GetPreview() : SelectedDragDropOp->Widget.GetTemplate();
 
-				if ( Widget->GetParent() )
+			if ( Widget->GetParent() )
+			{
+				if ( !bIsPreview )
 				{
-					if ( !bIsPreview )
-					{
-						Widget->GetParent()->Modify();
-					}
-
-					Widget->GetParent()->RemoveChild(Widget);
+					Widget->GetParent()->Modify();
 				}
 
-				FVector2D LocalPosition = ArrangedWidget.Geometry.AbsoluteToLocal(DragDropEvent.GetScreenSpacePosition());
-				if ( UPanelSlot* Slot = NewParent->AddChild(Widget) )
-				{
-					FWidgetBlueprintEditorUtils::ImportPropertiesFromText(Slot, SelectedDragDropOp->ExportedSlotProperties);
-
-					//TODO UMG Migrate existing slot info
-					Slot->SetDesiredPosition(LocalPosition - SelectedWidgetContextMenuLocation);
-					//Slot->SetDesiredSize(FVector2D(150, 30));
-					//@TODO UMG When we add a child blindly we need to default the slot size to the preferred size of the widget if the container supports such things.
-					//@TODO UMG We may need a desired size canvas, where the slots have no size, they only give you position, alternatively, maybe slots that don't clip, so center is still easy.
+				Widget->GetParent()->RemoveChild(Widget);
+			}
 
-					DropPreviewParent = NewParent;
+			FVector2D LocalPosition = ArrangedWidget.Geometry.AbsoluteToLocal(DragDropEvent.GetScreenSpacePosition());
+			if ( UPanelSlot* Slot = NewParent->AddChild(Widget) )
+			{
+				FWidgetBlueprintEditorUtils::ImportPropertiesFromText(Slot, SelectedDragDropOp->ExportedSlotProperties);
 
-					if ( bIsPreview )
-					{
-						Transaction.Cancel();
-					}
+				//TODO UMG Migrate existing slot info
+				Slot->SetDesiredPosition(LocalPosition - SelectedWidgetContextMenuLocation);
+				//Slot->SetDesiredSize(FVector2D(150, 30));
+				//@TODO UMG When we add a child blindly we need to default the slot size to the preferred size of the widget if the container supports such things.
+				//@TODO UMG We may need a desired size canvas, where the slots have no size, they only give you position, alternatively, maybe slots that don't clip, so center is still easy.
 
-					return Widget;
-				}
-				else
-				{
-					// TODO UMG ERROR Slot can not be created because maybe the max children has been reached.
-					//          Maybe we can traverse the hierarchy and add it to the first parent that will accept it?
-				}
+				DropPreviewParent = NewParent;
 
 				if ( bIsPreview )
 				{
 					Transaction.Cancel();
 				}
+
+				return Widget;
+			}
+			else
+			{
+				// TODO UMG ERROR Slot can not be created because maybe the max children has been reached.
+				//          Maybe we can traverse the hierarchy and add it to the first parent that will accept it?
+			}
+
+			if ( bIsPreview )
+			{
+				Transaction.Cancel();
 			}
 		}
 	}
diff --git a/Engine/Source/Editor/UMGEditor/Private/Templates/WidgetTemplate.h b/Engine/Source/Editor/UMGEditor/Private/Templates/WidgetTemplate.h
deleted file mode 100644
diff --git a/Engine/Source/Editor/UMGEditor/Private/WidgetTemplateDragDropOp.cpp b/Engine/Source/Editor/UMGEditor/Private/WidgetTemplateDragDropOp.cpp
--- a/Engine/Source/Editor/UMGEditor/Private/WidgetTemplateDragDropOp.cpp
+++ b/Engine/Source/Editor/UMGEditor/Private/WidgetTemplateDragDropOp.cpp
@@ -17,23 +17,4 @@ TSharedRef<FWidgetTemplateDragDropOp> FWidgetTemplateDragDropOp::New(const TShar
 	return Operation;
 }
 
-//TSharedPtr<SWidget> FWidgetTemplateDragDropOp::GetDefaultDecorator() const
-//{
-//	return 
-//		SNew(SBorder)
-//		.BorderImage(FEditorStyle::GetBrush("ContentBrowser.AssetDragDropTooltipBackground"))
-//		.Content()
-//		[
-//			SNew(SHorizontalBox)
-//
-//			+ SHorizontalBox::Slot()
-//			.AutoWidth()
-//			.HAlign(HAlign_Left)
-//			[
-//				SNew(STextBlock)
-//				.Text(Template->Name)
-//			]
-//		];
-//}
-
 #undef LOCTEXT_NAMESPACE
diff --git a/Engine/Source/Editor/UMGEditor/Private/WidgetTemplateDragDropOp.h b/Engine/Source/Editor/UMGEditor/Private/WidgetTemplateDragDropOp.h
--- a/Engine/Source/Editor/UMGEditor/Private/WidgetTemplateDragDropOp.h
+++ b/Engine/Source/Editor/UMGEditor/Private/WidgetTemplateDragDropOp.h
@@ -5,14 +5,18 @@
 #include "DecoratedDragDropOp.h"
 #include "WidgetTemplate.h"
 
+/**
+ * This drag drop operation allows widget templates from the palate to be dragged and dropped into the designer
+ * or the widget hierarchy in order to spawn new widgets.
+ */
 class FWidgetTemplateDragDropOp : public FDecoratedDragDropOp
 {
 public:
 	DRAG_DROP_OPERATOR_TYPE(FWidgetTemplateDragDropOp, FDecoratedDragDropOp)
 
+	/** The template to create an instance */
 	TSharedPtr<FWidgetTemplate> Template;
 
+	/** Constructs the drag drop operation */
 	static TSharedRef<FWidgetTemplateDragDropOp> New(const TSharedPtr<FWidgetTemplate>& InTemplate);
-
-	//virtual TSharedPtr<SWidget> GetDefaultDecorator() const override;
 };
diff --git a/Engine/Source/Editor/UMGEditor/Public/UMGEditorModule.h b/Engine/Source/Editor/UMGEditor/Public/UMGEditorModule.h
--- a/Engine/Source/Editor/UMGEditor/Public/UMGEditorModule.h
+++ b/Engine/Source/Editor/UMGEditor/Public/UMGEditorModule.h
@@ -11,6 +11,7 @@ extern const FName UMGEditorAppIdentifier;
 
 class FUMGEditor;
 
+/** The public interface of the UMG editor module. */
 class IUMGEditorModule : public IModuleInterface, public IHasMenuExtensibility, public IHasToolBarExtensibility
 {
 public:
diff --git a/Engine/Source/Editor/UMGEditor/Public/WidgetTemplate.h b/Engine/Source/Editor/UMGEditor/Public/WidgetTemplate.h
new file mode 100644
--- /dev/null
+++ b/Engine/Source/Editor/UMGEditor/Public/WidgetTemplate.h
@@ -0,0 +1,30 @@
+// Copyright 1998-2014 Epic Games, Inc. All Rights Reserved.
+
+#pragma once
+
+/**
+ * The widget template represents a widget or a set of widgets to create and spawn into the widget tree.
+ * More complicated defaults could be created by deriving from this class and registering new templates with the module.
+ */
+class FWidgetTemplate : public TSharedFromThis<FWidgetTemplate>
+{
+public:
+	/** Constructor */
+	FWidgetTemplate();
+
+	/** The category this template fits into. */
+	virtual FText GetCategory() = 0;
+
+	/** Constructs the widget template. */
+	virtual UWidget* Create(class UWidgetTree* Tree) = 0;
+
+public:
+	/** The name of the widget template. */
+	FText Name;
+
+	/** A tooltip to display for the template. */
+	FText ToolTip;
+
+	/** An icon to display for the template. */
+	FSlateIcon Icon;
+};
