UMG - Fixing multiple crashes in the designer from dead pointers, illegal pasting.

diff --git a/Engine/Source/Editor/UMGEditor/Private/Extensions/DesignerExtension.cpp b/Engine/Source/Editor/UMGEditor/Private/Extensions/DesignerExtension.cpp
--- a/Engine/Source/Editor/UMGEditor/Private/Extensions/DesignerExtension.cpp
+++ b/Engine/Source/Editor/UMGEditor/Private/Extensions/DesignerExtension.cpp
@@ -47,12 +47,12 @@ FWidgetReference FWidgetReference::FromPreview(TSharedPtr<FWidgetBlueprintEditor
 
 bool FWidgetReference::IsValid() const
 {
-	return TemplateWidget != NULL && GetPreview();
+	return TemplateWidget.Get() != NULL && GetPreview();
 }
 
 UWidget* FWidgetReference::GetTemplate() const
 {
-	return TemplateWidget;
+	return TemplateWidget.Get();
 }
 
 UWidget* FWidgetReference::GetPreview() const
@@ -61,9 +61,9 @@ UWidget* FWidgetReference::GetPreview() const
 	{
 		UUserWidget* PreviewRoot = WidgetEditor.Pin()->GetPreview();
 
-		if ( PreviewRoot && TemplateWidget )
+		if ( PreviewRoot && TemplateWidget.Get() )
 		{
-			UWidget* PreviewWidget = PreviewRoot->GetHandleFromName(TemplateWidget->GetName());
+			UWidget* PreviewWidget = PreviewRoot->GetHandleFromName(TemplateWidget.Get()->GetName());
 			return PreviewWidget;
 		}
 	}
diff --git a/Engine/Source/Editor/UMGEditor/Private/WidgetBlueprintEditorUtils.cpp b/Engine/Source/Editor/UMGEditor/Private/WidgetBlueprintEditorUtils.cpp
--- a/Engine/Source/Editor/UMGEditor/Private/WidgetBlueprintEditorUtils.cpp
+++ b/Engine/Source/Editor/UMGEditor/Private/WidgetBlueprintEditorUtils.cpp
@@ -275,9 +275,13 @@ void FWidgetBlueprintEditorUtils::PasteWidgets(UWidgetBlueprint* BP, FWidgetRefe
 
 	UPanelWidget* ParentWidget = CastChecked<UPanelWidget>(ParentWidgetRef.GetTemplate());
 
-	if ( ParentWidget )
+	// TODO UMG Find paste parent, may not be the selected widget...  Maybe it should be the parent of the copied widget until,
+	// we do a paste here, from a right click menu.
+
+	if ( !ParentWidget )
 	{
-		ParentWidget->Modify();
+		// TODO UMG allow pasting into the root slot?
+		return;
 	}
 
 	// Grab the text to paste from the clipboard.
@@ -288,6 +292,17 @@ void FWidgetBlueprintEditorUtils::PasteWidgets(UWidgetBlueprint* BP, FWidgetRefe
 	TSet<UWidget*> PastedWidgets;
 	FWidgetBlueprintEditorUtils::ImportWidgetsFromText(BP, TextToImport, /*out*/ PastedWidgets);
 
+	if ( !ParentWidget->CanHaveMultipleChildren() )
+	{
+		if ( ParentWidget->GetChildrenCount() > 0 || PastedWidgets.Num() > 1 )
+		{
+			// We can't paste the widgets into this container, there isn't enough room.
+			return;
+		}
+	}
+
+	ParentWidget->Modify();
+
 	for ( UWidget* NewWidget : PastedWidgets )
 	{
 		// Widgets with a null parent mean that they were the root most widget of their selection set when
@@ -296,7 +311,10 @@ void FWidgetBlueprintEditorUtils::PasteWidgets(UWidgetBlueprint* BP, FWidgetRefe
 		if ( NewWidget->GetParent() == NULL )
 		{
 			UPanelSlot* Slot = ParentWidget->AddChild(NewWidget);
-			Slot->SetDesiredPosition(PasteLocation);
+			if ( Slot )
+			{
+				Slot->SetDesiredPosition(PasteLocation);
+			}
 			//TODO UMG - The paste location needs to be relative from the most upper left hand corner of other widgets in their container.
 		}
 
diff --git a/Engine/Source/Editor/UMGEditor/Public/WidgetReference.h b/Engine/Source/Editor/UMGEditor/Public/WidgetReference.h
--- a/Engine/Source/Editor/UMGEditor/Public/WidgetReference.h
+++ b/Engine/Source/Editor/UMGEditor/Public/WidgetReference.h
@@ -38,7 +38,7 @@ public:
 	/** Checks if widget reference is the same as nother widget reference, based on the template pointers. */
 	bool operator==( const FWidgetReference& Other ) const
 	{
-		return TemplateWidget == Other.TemplateWidget;
+		return TemplateWidget.Get() == Other.TemplateWidget.Get();
 	}
 
 private:
@@ -49,7 +49,7 @@ private:
 	TWeakPtr<class FWidgetBlueprintEditor> WidgetEditor;
 
 	//TODO UMG We probably want to make this a weak object reference.
-	UWidget* TemplateWidget;
+	TWeakObjectPtr<UWidget> TemplateWidget;
 };
 
 inline uint32 GetTypeHash(const struct FWidgetReference& WidgetRef)
diff --git a/Engine/Source/Runtime/UMG/Public/Components/PanelWidget.h b/Engine/Source/Runtime/UMG/Public/Components/PanelWidget.h
--- a/Engine/Source/Runtime/UMG/Public/Components/PanelWidget.h
+++ b/Engine/Source/Runtime/UMG/Public/Components/PanelWidget.h
@@ -20,22 +20,33 @@ public:
 
 	//TODO UMG Add ways to make adding slots callable by blueprints.
 
+	/**  */
 	int32 GetChildrenCount() const;
 
+	/**  */
 	UWidget* GetChildAt(int32 Index) const;
 
+	/**  */
 	int32 GetChildIndex(UWidget* Content) const;
 
+	/**  */
 	bool RemoveChildAt(int32 Index);
 
+	/**  */
 	UPanelSlot* AddChild(UWidget* Content);
 
+	/**  */
 	void ReplaceChildAt(int32 Index, UWidget* Content);
 
+	/**  */
 	void InsertChildAt(int32 Index, UWidget* Content);
 
+	/**  */
 	bool RemoveChild(UWidget* Content);
 
+	/**  */
+	bool CanHaveMultipleChildren() const { return bCanHaveMultipleChildren; }
+
 #if WITH_EDITOR
 	virtual void ConnectEditorData() override
 	{
