UMG - Now says experimental in the overlay.

diff --git a/Engine/Source/Editor/UMGEditor/Private/SUMGDesigner.cpp b/Engine/Source/Editor/UMGEditor/Private/SUMGDesigner.cpp
--- a/Engine/Source/Editor/UMGEditor/Private/SUMGDesigner.cpp
+++ b/Engine/Source/Editor/UMGEditor/Private/SUMGDesigner.cpp
@@ -132,6 +132,7 @@ void SUMGDesigner::Construct(const FArguments& InArgs, TSharedPtr<FWidgetBluepri
 		[
 			SNew(SOverlay)
 
+			// The bottom layer of the overlay where the actual preview widget appears.
 			+ SOverlay::Slot()
 			.HAlign(HAlign_Fill)
 			.VAlign(VAlign_Fill)
@@ -140,6 +141,7 @@ void SUMGDesigner::Construct(const FArguments& InArgs, TSharedPtr<FWidgetBluepri
 				.Visibility(EVisibility::HitTestInvisible)
 			]
 
+			// A layer in the overlay where we put all the user intractable widgets, like the reorder widgets.
 			+ SOverlay::Slot()
 			.HAlign(HAlign_Fill)
 			.VAlign(VAlign_Fill)
@@ -147,6 +149,18 @@ void SUMGDesigner::Construct(const FArguments& InArgs, TSharedPtr<FWidgetBluepri
 				SAssignNew(ExtensionWidgetCanvas, SCanvas)
 				.Visibility(EVisibility::SelfHitTestInvisible)
 			]
+
+			// Bottom-right corner text indicating the experimental nature of umg
+			+ SOverlay::Slot()
+			.Padding(10)
+			.VAlign(VAlign_Bottom)
+			.HAlign(HAlign_Right)
+			[
+				SNew(STextBlock)
+				.Visibility( EVisibility::HitTestInvisible )
+				.TextStyle( FEditorStyle::Get(), "Graph.CornerText" )
+				.Text(LOCTEXT("ExperimentalWarning", "Experimental"))
+			]
 		]
 	);
 
