UMG - Adding overlay text to the graph.

diff --git a/Engine/Source/Editor/UMGEditor/Private/WidgetBlueprintEditor.cpp b/Engine/Source/Editor/UMGEditor/Private/WidgetBlueprintEditor.cpp
--- a/Engine/Source/Editor/UMGEditor/Private/WidgetBlueprintEditor.cpp
+++ b/Engine/Source/Editor/UMGEditor/Private/WidgetBlueprintEditor.cpp
@@ -415,4 +415,16 @@ UMovieScene* FWidgetBlueprintEditor::GetDefaultMovieScene()
 	return DefaultMovieScene;
 }
 
+FGraphAppearanceInfo FWidgetBlueprintEditor::GetGraphAppearance() const
+{
+	FGraphAppearanceInfo AppearanceInfo = FBlueprintEditor::GetGraphAppearance();
+
+	if ( GetBlueprintObj()->IsA(UWidgetBlueprint::StaticClass()) )
+	{
+		AppearanceInfo.CornerText = LOCTEXT("AppearanceCornerText", "WIDGET BLUEPRINT").ToString();
+	}
+
+	return AppearanceInfo;
+}
+
 #undef LOCTEXT_NAMESPACE
diff --git a/Engine/Source/Editor/UMGEditor/Private/WidgetBlueprintEditor.h b/Engine/Source/Editor/UMGEditor/Private/WidgetBlueprintEditor.h
--- a/Engine/Source/Editor/UMGEditor/Private/WidgetBlueprintEditor.h
+++ b/Engine/Source/Editor/UMGEditor/Private/WidgetBlueprintEditor.h
@@ -60,6 +60,9 @@ public:
 	/** Paste Metadata */
 	FVector2D PasteDropLocation;
 
+protected:
+	virtual FGraphAppearanceInfo GetGraphAppearance() const override;
+
 private:
 	bool CanDeleteSelectedWidgets();
 	void DeleteSelectedWidgets();
