Move WidgetBlueprint to the UMG Editor module as it is not needed at runtime

diff --git a/Engine/Source/Editor/UMGEditor/Classes/WidgetGraphNode_Base.h b/Engine/Source/Editor/UMGEditor/Classes/WidgetGraphNode_Base.h
--- a/Engine/Source/Editor/UMGEditor/Classes/WidgetGraphNode_Base.h
+++ b/Engine/Source/Editor/UMGEditor/Classes/WidgetGraphNode_Base.h
@@ -48,7 +48,7 @@ class UWidgetGraphNode_Base : public UK2Node
 	virtual void CustomizePinData(UEdGraphPin* Pin, FName SourcePropertyName, int32 ArrayIndex) const;
 
 	/** Get the animation blueprint to which this node belongs */
-	UWidgetBlueprint* GetWidgetBlueprint() const { return CastChecked<UWidgetBlueprint>(GetBlueprint()); }
+	class UWidgetBlueprint* GetWidgetBlueprint() const;
 
 	// UWidgetGraphNode_Base
 
diff --git a/Engine/Source/Editor/UMGEditor/Private/AssetTypeActions_WidgetBlueprint.cpp b/Engine/Source/Editor/UMGEditor/Private/AssetTypeActions_WidgetBlueprint.cpp
--- a/Engine/Source/Editor/UMGEditor/Private/AssetTypeActions_WidgetBlueprint.cpp
+++ b/Engine/Source/Editor/UMGEditor/Private/AssetTypeActions_WidgetBlueprint.cpp
@@ -10,7 +10,7 @@
 #include "SBlueprintDiff.h"
 #include "ISourceControlModule.h"
 #include "MessageLog.h"
-
+#include "WidgetBlueprint.h"
 #include "AssetTypeActions_WidgetBlueprint.h"
 #include "WidgetBlueprintEditor.h"
 
@@ -41,6 +41,11 @@ void FAssetTypeActions_WidgetBlueprint::OpenAssetEditor( const TArray<UObject*>&
 	}
 }
 
+UClass* FAssetTypeActions_WidgetBlueprint::GetSupportedClass() const
+{
+	return UWidgetBlueprint::StaticClass();
+}
+
 bool FAssetTypeActions_WidgetBlueprint::ShouldUseDataOnlyEditor(const UBlueprint* Blueprint) const
 {
 	return FBlueprintEditorUtils::IsDataOnlyBlueprint(Blueprint)
diff --git a/Engine/Source/Editor/UMGEditor/Private/AssetTypeActions_WidgetBlueprint.h b/Engine/Source/Editor/UMGEditor/Private/AssetTypeActions_WidgetBlueprint.h
--- a/Engine/Source/Editor/UMGEditor/Private/AssetTypeActions_WidgetBlueprint.h
+++ b/Engine/Source/Editor/UMGEditor/Private/AssetTypeActions_WidgetBlueprint.h
@@ -9,7 +9,7 @@ public:
 	// IAssetTypeActions Implementation
 	virtual FText GetName() const override { return NSLOCTEXT("AssetTypeActions", "AssetTypeActions_WidgetBlueprint", "Widget Blueprint"); }
 	virtual FColor GetTypeColor() const override { return FColor(255,255,255); }
-	virtual UClass* GetSupportedClass() const override { return UWidgetBlueprint::StaticClass(); }
+	virtual UClass* GetSupportedClass() const override;
 	virtual void OpenAssetEditor(const TArray<UObject*>& InObjects, TSharedPtr<class IToolkitHost> EditWithinLevelEditor = TSharedPtr<IToolkitHost>()) override;
 	virtual uint32 GetCategories() override { return EAssetTypeCategories::UI; }
 	// End IAssetTypeActions Implementation
diff --git a/Engine/Source/Editor/UMGEditor/Private/SUMGDesigner.cpp b/Engine/Source/Editor/UMGEditor/Private/SUMGDesigner.cpp
--- a/Engine/Source/Editor/UMGEditor/Private/SUMGDesigner.cpp
+++ b/Engine/Source/Editor/UMGEditor/Private/SUMGDesigner.cpp
@@ -6,7 +6,7 @@
 #include "BlueprintEditor.h"
 #include "SKismetInspector.h"
 #include "BlueprintEditorUtils.h"
-
+#include "WidgetBlueprint.h"
 #include "WidgetBlueprintEditor.h"
 
 #include "WidgetTemplateDragDropOp.h"
diff --git a/Engine/Source/Editor/UMGEditor/Private/WidgetBlueprint.cpp b/Engine/Source/Editor/UMGEditor/Private/WidgetBlueprint.cpp
new file mode 100644
--- /dev/null
+++ b/Engine/Source/Editor/UMGEditor/Private/WidgetBlueprint.cpp
@@ -0,0 +1,93 @@
+// Copyright 1998-2014 Epic Games, Inc. All Rights Reserved.
+
+#include "UMGEditorPrivatePCH.h"
+#include "Runtime/MovieSceneCore/Classes/MovieScene.h"
+
+#define LOCTEXT_NAMESPACE "UMG"
+
+FDelegateRuntimeBinding FDelegateEditorBinding::ToRuntimeBinding(UWidgetBlueprint* Blueprint) const
+{
+	FDelegateRuntimeBinding Binding;
+	Binding.ObjectName = ObjectName;
+	Binding.PropertyName = PropertyName;
+	Binding.FunctionName = Blueprint->GetFieldNameFromClassByGuid<UFunction>(Blueprint->SkeletonGeneratedClass, MemberGuid);
+
+	return Binding;
+}
+
+
+/////////////////////////////////////////////////////
+// UWidgetBlueprint
+
+UWidgetBlueprint::UWidgetBlueprint(const FPostConstructInitializeProperties& PCIP)
+	: Super(PCIP)
+{
+	WidgetTree = ConstructObject<UWidgetTree>(UWidgetTree::StaticClass(), this);
+	WidgetTree->SetFlags(RF_Transactional);
+}
+
+void UWidgetBlueprint::PostLoad()
+{
+	Super::PostLoad();
+
+	TArray<UWidget*> Widgets;
+	WidgetTree->GetAllWidgets(Widgets);
+	for ( UWidget* Widget : Widgets )
+	{
+		Widget->ConnectEditorData();
+	}
+
+}
+
+UClass* UWidgetBlueprint::RegenerateClass(UClass* ClassToRegenerate, UObject* PreviousCDO, TArray<UObject*>& ObjLoaded)
+{
+	return Super::RegenerateClass(ClassToRegenerate, PreviousCDO, ObjLoaded);
+}
+
+
+UClass* UWidgetBlueprint::GetBlueprintClass() const
+{
+	return UWidgetBlueprintGeneratedClass::StaticClass();
+}
+
+
+bool UWidgetBlueprint::ValidateGeneratedClass(const UClass* InClass)
+{
+	bool Result = Super::ValidateGeneratedClass(InClass);
+
+	const UWidgetBlueprintGeneratedClass* GeneratedClass = Cast<const UWidgetBlueprintGeneratedClass>(InClass);
+	if ( !ensure(GeneratedClass) )
+	{
+		return false;
+	}
+	const UWidgetBlueprint* Blueprint = Cast<UWidgetBlueprint>(GetBlueprintFromClass(GeneratedClass));
+	if ( !ensure(Blueprint) )
+	{
+		return false;
+	}
+
+	//if ( !ensure(Blueprint->WidgetTree && ( Blueprint->WidgetTree->GetOuter() == GeneratedClass )) )
+	//{
+	//	return false;
+	//}
+
+	//if ( !ensure(GeneratedClass->WidgetTree && ( GeneratedClass->WidgetTree->GetOuter() == GeneratedClass )) )
+	//{
+	//	return false;
+	//}
+
+	return Result;
+}
+
+void UWidgetBlueprint::PostLoadSubobjects(FObjectInstancingGraph* OuterInstanceGraph)
+{
+	Super::PostLoadSubobjects(OuterInstanceGraph);
+
+	UWidgetBlueprintGeneratedClass* BPGClass = Cast<UWidgetBlueprintGeneratedClass>(*GeneratedClass);
+	if ( BPGClass )
+	{
+		//BPGClass->WidgetTree = WidgetTree;
+	}
+}
+
+#undef LOCTEXT_NAMESPACE 
\ No newline at end of file
diff --git a/Engine/Source/Editor/UMGEditor/Private/WidgetBlueprint.h b/Engine/Source/Editor/UMGEditor/Private/WidgetBlueprint.h
new file mode 100644
--- /dev/null
+++ b/Engine/Source/Editor/UMGEditor/Private/WidgetBlueprint.h
@@ -0,0 +1,74 @@
+// Copyright 1998-2014 Epic Games, Inc. All Rights Reserved.
+
+#pragma once
+
+#include "Engine/Blueprint.h"
+#include "UserWidget.h"
+#include "WidgetBlueprintGeneratedClass.h"
+
+#include "WidgetBlueprint.generated.h"
+
+class UMovieScene;
+
+USTRUCT()
+struct FDelegateEditorBinding
+{
+	GENERATED_USTRUCT_BODY()
+
+	UPROPERTY()
+	FString ObjectName;
+
+	UPROPERTY()
+	FName PropertyName;
+
+	UPROPERTY()
+	FName FunctionName;
+
+	UPROPERTY()
+	FGuid MemberGuid;
+
+	bool operator==( const FDelegateEditorBinding& Other ) const
+	{
+		// NOTE: We intentionally only compare object name and property name, the function is irrelevant since
+		// you're only allowed to bind a property on an object to a single function.
+		return ObjectName == Other.ObjectName && PropertyName == Other.PropertyName;
+	}
+
+	FDelegateRuntimeBinding ToRuntimeBinding(class UWidgetBlueprint* Blueprint) const;
+};
+
+/**
+ * The widget blueprint enables extending UUserWidget the user extensible UWidget.
+ */
+UCLASS(BlueprintType)
+class UWidgetBlueprint : public UBlueprint
+{
+	GENERATED_UCLASS_BODY()
+
+public:
+	/** A tree of the widget templates to be created */
+	UPROPERTY()
+	class UWidgetTree* WidgetTree;
+
+	UPROPERTY()
+	TArray< FDelegateEditorBinding > Bindings;
+
+	UPROPERTY()
+	TArray< UMovieScene* > AnimationData;
+	 
+	virtual void PostLoad() override;
+	virtual void PostLoadSubobjects(FObjectInstancingGraph* OuterInstanceGraph) override;
+
+	virtual UClass* RegenerateClass(UClass* ClassToRegenerate, UObject* PreviousCDO, TArray<UObject*>& ObjLoaded) override;
+	
+	// UBlueprint interface
+	virtual UClass* GetBlueprintClass() const override;
+
+	virtual bool SupportedByDefaultBlueprintFactory() const override
+	{
+		return false;
+	}
+	// End of UBlueprint interface
+
+	static bool ValidateGeneratedClass(const UClass* InClass);
+};
\ No newline at end of file
diff --git a/Engine/Source/Editor/UMGEditor/Private/WidgetGraphNode_Base.cpp b/Engine/Source/Editor/UMGEditor/Private/WidgetGraphNode_Base.cpp
--- a/Engine/Source/Editor/UMGEditor/Private/WidgetGraphNode_Base.cpp
+++ b/Engine/Source/Editor/UMGEditor/Private/WidgetGraphNode_Base.cpp
@@ -228,6 +228,11 @@ void UWidgetGraphNode_Base::CustomizePinData(UEdGraphPin* Pin, FName SourcePrope
 
 }
 
+UWidgetBlueprint* UWidgetGraphNode_Base::GetWidgetBlueprint() const
+{
+	 return CastChecked<UWidgetBlueprint>(GetBlueprint());
+}
+
 FLinearColor UWidgetGraphNode_Base::GetNodeTitleColor() const
 {
 	return FLinearColor::Black;
diff --git a/Engine/Source/Runtime/UMG/Private/WidgetBlueprint.cpp b/Engine/Source/Runtime/UMG/Private/WidgetBlueprint.cpp
deleted file mode 100644
diff --git a/Engine/Source/Runtime/UMG/Public/Blueprint/WidgetBlueprint.h b/Engine/Source/Runtime/UMG/Public/Blueprint/WidgetBlueprint.h
deleted file mode 100644
diff --git a/Engine/Source/Runtime/UMG/Public/UMG.h b/Engine/Source/Runtime/UMG/Public/UMG.h
--- a/Engine/Source/Runtime/UMG/Public/UMG.h
+++ b/Engine/Source/Runtime/UMG/Public/UMG.h
@@ -85,7 +85,6 @@
 #include "WidgetNode_Base.h"
 #include "WidgetNode_Text.h"
 
-#include "WidgetBlueprint.h"
 #include "WidgetBlueprintGeneratedClass.h"
 
 #include "WidgetActor.h"
