UMG - Moving the WidgetBlueprint into the public folder for the UMGEditor.  Adding dll export macro to the public headers.

diff --git a/Engine/Source/Editor/UMGEditor/Private/WidgetBlueprint.h b/Engine/Source/Editor/UMGEditor/Private/WidgetBlueprint.h
deleted file mode 100644
diff --git a/Engine/Source/Editor/UMGEditor/Public/DesignerExtension.h b/Engine/Source/Editor/UMGEditor/Public/DesignerExtension.h
--- a/Engine/Source/Editor/UMGEditor/Public/DesignerExtension.h
+++ b/Engine/Source/Editor/UMGEditor/Public/DesignerExtension.h
@@ -12,7 +12,7 @@ class UWidgetBlueprint;
  * The Designer extension allows developers to provide additional widgets and custom painting to the designer surface for
  * specific widgets.  Allowing for a more customized and specific editors for the different widgets.
  */
-class FDesignerExtension : public TSharedFromThis<FDesignerExtension>
+class UMGEDITOR_API FDesignerExtension : public TSharedFromThis<FDesignerExtension>
 {
 public:
 	/** Constructor */
diff --git a/Engine/Source/Editor/UMGEditor/Public/IUMGDesigner.h b/Engine/Source/Editor/UMGEditor/Public/IUMGDesigner.h
--- a/Engine/Source/Editor/UMGEditor/Public/IUMGDesigner.h
+++ b/Engine/Source/Editor/UMGEditor/Public/IUMGDesigner.h
@@ -6,7 +6,7 @@
  * The public interface implemented by the UMG Designer to allow extensions to call methods
  * on the designer.
  */
-class IUMGDesigner
+class UMGEDITOR_API IUMGDesigner
 {
 public:
 
diff --git a/Engine/Source/Editor/UMGEditor/Public/WidgetBlueprint.h b/Engine/Source/Editor/UMGEditor/Public/WidgetBlueprint.h
new file mode 100644
--- /dev/null
+++ b/Engine/Source/Editor/UMGEditor/Public/WidgetBlueprint.h
@@ -0,0 +1,76 @@
+// Copyright 1998-2014 Epic Games, Inc. All Rights Reserved.
+
+#pragma once
+
+#include "Engine/Blueprint.h"
+#include "UserWidget.h"
+#include "WidgetBlueprintGeneratedClass.h"
+
+#include "WidgetBlueprint.generated.h"
+
+class UMovieScene;
+
+USTRUCT()
+struct UMGEDITOR_API FDelegateEditorBinding
+{
+	GENERATED_USTRUCT_BODY()
+
+	UPROPERTY()
+	FString ObjectName;
+
+	UPROPERTY()
+	FName PropertyName;
+
+	UPROPERTY()
+	FName FunctionName;
+
+	UPROPERTY()
+	FGuid MemberGuid;
+
+	bool operator==( const FDelegateEditorBinding& Other ) const
+	{
+		// NOTE: We intentionally only compare object name and property name, the function is irrelevant since
+		// you're only allowed to bind a property on an object to a single function.
+		return ObjectName == Other.ObjectName && PropertyName == Other.PropertyName;
+	}
+
+	FDelegateRuntimeBinding ToRuntimeBinding(class UWidgetBlueprint* Blueprint) const;
+};
+
+/**
+ * The widget blueprint enables extending UUserWidget the user extensible UWidget.
+ */
+UCLASS(BlueprintType)
+class UMGEDITOR_API UWidgetBlueprint : public UBlueprint
+{
+	GENERATED_UCLASS_BODY()
+
+public:
+	/** A tree of the widget templates to be created */
+	UPROPERTY()
+	class UWidgetTree* WidgetTree;
+
+	UPROPERTY()
+	TArray< FDelegateEditorBinding > Bindings;
+
+	UPROPERTY()
+	TArray< UMovieScene* > AnimationData;
+	 
+	virtual void PostLoad() override;
+	virtual void PostLoadSubobjects(FObjectInstancingGraph* OuterInstanceGraph) override;
+
+	virtual UClass* RegenerateClass(UClass* ClassToRegenerate, UObject* PreviousCDO, TArray<UObject*>& ObjLoaded) override;
+	
+	// UBlueprint interface
+	virtual UClass* GetBlueprintClass() const override;
+
+	virtual bool AllowsDynamicBinding() const override;
+
+	virtual bool SupportedByDefaultBlueprintFactory() const override
+	{
+		return false;
+	}
+	// End of UBlueprint interface
+
+	static bool ValidateGeneratedClass(const UClass* InClass);
+};
\ No newline at end of file
diff --git a/Engine/Source/Editor/UMGEditor/Public/WidgetReference.h b/Engine/Source/Editor/UMGEditor/Public/WidgetReference.h
--- a/Engine/Source/Editor/UMGEditor/Public/WidgetReference.h
+++ b/Engine/Source/Editor/UMGEditor/Public/WidgetReference.h
@@ -10,7 +10,7 @@ class UWidgetBlueprint;
  * The Widget reference is a useful way to hold onto the selection in a way that allows for up to date access to the current preview object.
  * Because the designer could end up rebuilding the preview, it's best to hold onto an FWidgetReference.
  */
-struct FWidgetReference
+struct UMGEDITOR_API FWidgetReference
 {
 public:
 	/**
diff --git a/Engine/Source/Editor/UMGEditor/Public/WidgetTemplate.h b/Engine/Source/Editor/UMGEditor/Public/WidgetTemplate.h
--- a/Engine/Source/Editor/UMGEditor/Public/WidgetTemplate.h
+++ b/Engine/Source/Editor/UMGEditor/Public/WidgetTemplate.h
@@ -6,7 +6,7 @@
  * The widget template represents a widget or a set of widgets to create and spawn into the widget tree.
  * More complicated defaults could be created by deriving from this class and registering new templates with the module.
  */
-class FWidgetTemplate : public TSharedFromThis<FWidgetTemplate>
+class UMGEDITOR_API FWidgetTemplate : public TSharedFromThis<FWidgetTemplate>
 {
 public:
 	/** Constructor */
