Build Fix - Adding UWidgetBlueprint to the editor PCH to fix non-unity builds.

diff --git a/Engine/Source/Editor/UMGEditor/Private/UMGEditorPrivatePCH.h b/Engine/Source/Editor/UMGEditor/Private/UMGEditorPrivatePCH.h
--- a/Engine/Source/Editor/UMGEditor/Private/UMGEditorPrivatePCH.h
+++ b/Engine/Source/Editor/UMGEditor/Private/UMGEditorPrivatePCH.h
@@ -15,6 +15,7 @@
 #include "UMG.h"
 
 #include "UMGEditorClasses.h"
+#include "UWidgetBlueprint.h"
 
 #include "AssetTypeActions_Base.h"
 
