umg - Non-unity build fix.

diff --git a/Engine/Source/Editor/UMGEditor/Private/UMGEditorPrivatePCH.h b/Engine/Source/Editor/UMGEditor/Private/UMGEditorPrivatePCH.h
--- a/Engine/Source/Editor/UMGEditor/Private/UMGEditorPrivatePCH.h
+++ b/Engine/Source/Editor/UMGEditor/Private/UMGEditorPrivatePCH.h
@@ -15,7 +15,7 @@
 #include "UMG.h"
 
 #include "UMGEditorClasses.h"
-#include "UWidgetBlueprint.h"
+#include "WidgetBlueprint.h"
 
 #include "AssetTypeActions_Base.h"
 
