UMG - Designer now has more robust support for copy paste.  Can now paste as the root widget of the tree.  Built out a proof of concept solution for support pasting multiple root widgets, once multiselection is enabled, by rooting multiple widgets automatically to a Canvas if the parent can't support more than one child.  Also better support for removing widgets from the selection that are deleted, in some fashion, undo or normal delete.

diff --git a/Engine/Source/Editor/UMGEditor/Private/WidgetBlueprintEditor.cpp b/Engine/Source/Editor/UMGEditor/Private/WidgetBlueprintEditor.cpp
--- a/Engine/Source/Editor/UMGEditor/Private/WidgetBlueprintEditor.cpp
+++ b/Engine/Source/Editor/UMGEditor/Private/WidgetBlueprintEditor.cpp
@@ -28,13 +28,43 @@ FWidgetBlueprintEditor::~FWidgetBlueprintEditor()
 
 void FWidgetBlueprintEditor::SelectWidgets(const TSet<FWidgetReference>& Widgets)
 {
-	SelectedWidgets = Widgets;
+	SelectedWidgets.Empty();
+
+	for ( const FWidgetReference& Widget : Widgets )
+	{
+		if ( Widget.IsValid() )
+		{
+			SelectedWidgets.Add(Widget);
+		}
+	}
 
 	RefreshDetails();
 
 	OnSelectedWidgetsChanged.Broadcast();
 }
 
+void FWidgetBlueprintEditor::CleanSelection()
+{
+	TSet<FWidgetReference> TempSelection;
+
+	TArray<UWidget*> WidgetsInTree;
+	GetWidgetBlueprintObj()->WidgetTree->GetAllWidgets(WidgetsInTree);
+	TSet<UWidget*> TreeWidgetSet(WidgetsInTree);
+
+	for ( FWidgetReference& WidgetRef : SelectedWidgets )
+	{
+		if ( TreeWidgetSet.Contains(WidgetRef.GetTemplate()) )
+		{
+			SelectedWidgets.Add(WidgetRef);
+		}
+	}
+
+	if ( TempSelection.Num() != SelectedWidgets.Num() )
+	{
+		SelectWidgets(TempSelection);
+	}
+}
+
 const TSet<FWidgetReference>& FWidgetBlueprintEditor::GetSelectedWidgets() const
 {
 	return SelectedWidgets;
@@ -66,6 +96,8 @@ void FWidgetBlueprintEditor::OnBlueprintChanged(UBlueprint* InBlueprint)
 {
 	if ( InBlueprint )
 	{
+		CleanSelection();
+
 		UpdatePreview(InBlueprint, true);
 
 		RefreshDetails();
@@ -121,6 +153,10 @@ void FWidgetBlueprintEditor::DeleteSelectedWidgets()
 {
 	TSet<FWidgetReference> Widgets = GetSelectedWidgets();
 	FWidgetBlueprintEditorUtils::DeleteWidgets(GetWidgetBlueprintObj(), Widgets);
+
+	// Clear the selection now that the widget has been deleted.
+	TSet<FWidgetReference> Empty;
+	SelectWidgets(Empty);
 }
 
 bool FWidgetBlueprintEditor::CanCopySelectedWidgets()
@@ -144,6 +180,13 @@ bool FWidgetBlueprintEditor::CanPasteWidgets()
 		const bool bIsPanel = Cast<UPanelWidget>(Target.GetTemplate()) != NULL;
 		return bIsPanel;
 	}
+	else if ( Widgets.Num() == 0 )
+	{
+		if ( GetWidgetBlueprintObj()->WidgetTree->RootWidget == NULL )
+		{
+			return true;
+		}
+	}
 
 	return false;
 }
@@ -151,9 +194,11 @@ bool FWidgetBlueprintEditor::CanPasteWidgets()
 void FWidgetBlueprintEditor::PasteWidgets()
 {
 	TSet<FWidgetReference> Widgets = GetSelectedWidgets();
-	FWidgetReference Target = *Widgets.CreateIterator();
+	FWidgetReference Target = Widgets.Num() > 0 ? *Widgets.CreateIterator() : FWidgetReference();
 
 	FWidgetBlueprintEditorUtils::PasteWidgets(GetWidgetBlueprintObj(), Target, PasteDropLocation);
+
+	//TODO UMG - Select the newly selected pasted widgets.
 }
 
 void FWidgetBlueprintEditor::Tick(float DeltaTime)
diff --git a/Engine/Source/Editor/UMGEditor/Private/WidgetBlueprintEditor.h b/Engine/Source/Editor/UMGEditor/Private/WidgetBlueprintEditor.h
--- a/Engine/Source/Editor/UMGEditor/Private/WidgetBlueprintEditor.h
+++ b/Engine/Source/Editor/UMGEditor/Private/WidgetBlueprintEditor.h
@@ -44,6 +44,9 @@ public:
 	/** Sets the currently selected set of widgets */
 	void SelectWidgets(const TSet<FWidgetReference>& Widgets);
 
+	/** Removes removed widgets from the selection set. */
+	void CleanSelection();
+
 	/** @return The selected set of widgets */
 	const TSet<FWidgetReference>& GetSelectedWidgets() const;
 
diff --git a/Engine/Source/Editor/UMGEditor/Private/WidgetBlueprintEditorUtils.cpp b/Engine/Source/Editor/UMGEditor/Private/WidgetBlueprintEditorUtils.cpp
--- a/Engine/Source/Editor/UMGEditor/Private/WidgetBlueprintEditorUtils.cpp
+++ b/Engine/Source/Editor/UMGEditor/Private/WidgetBlueprintEditorUtils.cpp
@@ -273,15 +273,23 @@ void FWidgetBlueprintEditorUtils::PasteWidgets(UWidgetBlueprint* BP, FWidgetRefe
 {
 	const FScopedTransaction Transaction(FGenericCommands::Get().Paste->GetDescription());
 
-	UPanelWidget* ParentWidget = CastChecked<UPanelWidget>(ParentWidgetRef.GetTemplate());
-
+	UPanelWidget* ParentWidget = NULL;
+	
+	if ( ParentWidgetRef.IsValid() )
+	{
+		ParentWidget = CastChecked<UPanelWidget>(ParentWidgetRef.GetTemplate());
+	}
+	
 	// TODO UMG Find paste parent, may not be the selected widget...  Maybe it should be the parent of the copied widget until,
 	// we do a paste here, from a right click menu.
 
 	if ( !ParentWidget )
 	{
-		// TODO UMG allow pasting into the root slot?
-		return;
+		// If we already have a root widget, then we can't replace the root.
+		if ( BP->WidgetTree->RootWidget )
+		{
+			return;
+		}
 	}
 
 	// Grab the text to paste from the clipboard.
@@ -292,23 +300,51 @@ void FWidgetBlueprintEditorUtils::PasteWidgets(UWidgetBlueprint* BP, FWidgetRefe
 	TSet<UWidget*> PastedWidgets;
 	FWidgetBlueprintEditorUtils::ImportWidgetsFromText(BP, TextToImport, /*out*/ PastedWidgets);
 
-	if ( !ParentWidget->CanHaveMultipleChildren() )
+	// Ignore an empty set of widget paste data.
+	if ( PastedWidgets.Num() == 0 )
 	{
-		if ( ParentWidget->GetChildrenCount() > 0 || PastedWidgets.Num() > 1 )
-		{
-			// We can't paste the widgets into this container, there isn't enough room.
-			return;
-		}
+		return;
 	}
 
-	ParentWidget->Modify();
-
+	TArray<UWidget*> RootPasteWidgets;
 	for ( UWidget* NewWidget : PastedWidgets )
 	{
 		// Widgets with a null parent mean that they were the root most widget of their selection set when
 		// they were copied and thus we need to paste only the root most widgets.  All their children will be added
 		// automatically.
 		if ( NewWidget->GetParent() == NULL )
+		{
+			RootPasteWidgets.Add(NewWidget);
+		}
+
+		NewWidget->SetFlags(RF_Transactional);
+	}
+
+	// If there isn't a root widget and we're copying multiple root widgets, then we need to add a container root
+	// to hold the pasted data since multiple root widgets isn't permitted.
+	if ( !ParentWidget && RootPasteWidgets.Num() > 1 )
+	{
+		ParentWidget = BP->WidgetTree->ConstructWidget<UCanvasPanel>(UCanvasPanel::StaticClass());
+		BP->WidgetTree->Modify();
+		BP->WidgetTree->RootWidget = ParentWidget;
+	}
+
+	if ( ParentWidget )
+	{
+		if ( !ParentWidget->CanHaveMultipleChildren() )
+		{
+			if ( ParentWidget->GetChildrenCount() > 0 || RootPasteWidgets.Num() > 1 )
+			{
+				UCanvasPanel* PasteContainer = BP->WidgetTree->ConstructWidget<UCanvasPanel>(UCanvasPanel::StaticClass());
+				//TODO UMG The new container could be tiny, unless filling the space.
+				UPanelSlot* Slot = ParentWidget->AddChild(PasteContainer);
+				ParentWidget = PasteContainer;
+			}
+		}
+
+		ParentWidget->Modify();
+
+		for ( UWidget* NewWidget : RootPasteWidgets )
 		{
 			UPanelSlot* Slot = ParentWidget->AddChild(NewWidget);
 			if ( Slot )
@@ -318,11 +354,23 @@ void FWidgetBlueprintEditorUtils::PasteWidgets(UWidgetBlueprint* BP, FWidgetRefe
 			//TODO UMG - The paste location needs to be relative from the most upper left hand corner of other widgets in their container.
 		}
 
-		NewWidget->SetFlags( RF_Transactional );
+		FBlueprintEditorUtils::MarkBlueprintAsStructurallyModified(BP);
 	}
+	else
+	{
+		check(RootPasteWidgets.Num() == 1)
+		// If we've arrived here, we must be creating the root widget from paste data, and there can only be
+		// one item in the paste data by now.
+		BP->WidgetTree->Modify();
 
-	//TODO UMG Reorder the live slot without rebuilding the structure
-	FBlueprintEditorUtils::MarkBlueprintAsStructurallyModified(BP);
+		for ( UWidget* NewWidget : RootPasteWidgets )
+		{
+			BP->WidgetTree->RootWidget = NewWidget;
+			break;
+		}
+		
+		FBlueprintEditorUtils::MarkBlueprintAsStructurallyModified(BP);
+	}
 }
 
 void FWidgetBlueprintEditorUtils::ImportWidgetsFromText(UWidgetBlueprint* BP, const FString& TextToImport, /*out*/ TSet<UWidget*>& ImportedWidgetSet)
diff --git a/Engine/Source/Runtime/UMG/Public/Blueprint/WidgetTree.h b/Engine/Source/Runtime/UMG/Public/Blueprint/WidgetTree.h
--- a/Engine/Source/Runtime/UMG/Public/Blueprint/WidgetTree.h
+++ b/Engine/Source/Runtime/UMG/Public/Blueprint/WidgetTree.h
@@ -10,14 +10,6 @@ class UMG_API UWidgetTree : public UObject
 {
 	GENERATED_UCLASS_BODY()
 
-public:
-	/** The root widget of the tree */
-	UPROPERTY()
-	UWidget* RootWidget;
-
-	UPROPERTY()
-	TArray< UWidget* > AllWidgets;
-
 public:
 
 	UWidget* FindWidget(const FString& Name) const;
@@ -53,4 +45,14 @@ public:
 
 		Super::PreSave();
 	}
+
+public:
+	/** The root widget of the tree */
+	UPROPERTY()
+	UWidget* RootWidget;
+
+protected:
+
+	UPROPERTY()
+	TArray< UWidget* > AllWidgets;
 };
