UMG - Selection and property grid are a lot more stable now.  Reduced a number of crashes resulting from stale selection data.

diff --git a/Engine/Source/Editor/UMGEditor/Private/WidgetBlueprintEditor.cpp b/Engine/Source/Editor/UMGEditor/Private/WidgetBlueprintEditor.cpp
--- a/Engine/Source/Editor/UMGEditor/Private/WidgetBlueprintEditor.cpp
+++ b/Engine/Source/Editor/UMGEditor/Private/WidgetBlueprintEditor.cpp
@@ -28,17 +28,22 @@ FWidgetBlueprintEditor::~FWidgetBlueprintEditor()
 
 void FWidgetBlueprintEditor::SelectWidgets(const TSet<FWidgetReference>& Widgets)
 {
-	SelectedWidgets.Empty();
-
+	TSet<FWidgetReference> TempSelection;
 	for ( const FWidgetReference& Widget : Widgets )
 	{
 		if ( Widget.IsValid() )
 		{
-			SelectedWidgets.Add(Widget);
+			TempSelection.Add(Widget);
 		}
 	}
 
-	RefreshDetails();
+	RefreshDetails(TempSelection);
+
+	// Finally change the selected widgets after we've updated the details panel 
+	// to ensure values that are pending are committed on focus loss, and migrated properly
+	// to the old selected widgets.
+	SelectedWidgets.Empty();
+	SelectedWidgets.Append(TempSelection);
 
 	OnSelectedWidgetsChanged.Broadcast();
 }
@@ -53,9 +58,12 @@ void FWidgetBlueprintEditor::CleanSelection()
 
 	for ( FWidgetReference& WidgetRef : SelectedWidgets )
 	{
-		if ( TreeWidgetSet.Contains(WidgetRef.GetTemplate()) )
+		if ( WidgetRef.IsValid() )
 		{
-			SelectedWidgets.Add(WidgetRef);
+			if ( TreeWidgetSet.Contains(WidgetRef.GetTemplate()) )
+			{
+				TempSelection.Add(WidgetRef);
+			}
 		}
 	}
 
@@ -71,12 +79,17 @@ const TSet<FWidgetReference>& FWidgetBlueprintEditor::GetSelectedWidgets() const
 }
 
 void FWidgetBlueprintEditor::RefreshDetails()
+{
+	RefreshDetails(SelectedWidgets);
+}
+
+void FWidgetBlueprintEditor::RefreshDetails(TSet<FWidgetReference>& Widgets)
 {
 	// Convert the selection set to an array of UObject* pointers
 	FString InspectorTitle = "Widget";// Widget->GetDisplayString();
 
 	TArray<UObject*> InspectorObjects;
-	for ( FWidgetReference& WidgetRef : SelectedWidgets )
+	for ( FWidgetReference& WidgetRef : Widgets )
 	{
 		UWidget* PreviewWidget = WidgetRef.GetPreview();
 		if ( PreviewWidget )
diff --git a/Engine/Source/Editor/UMGEditor/Private/WidgetBlueprintEditor.h b/Engine/Source/Editor/UMGEditor/Private/WidgetBlueprintEditor.h
--- a/Engine/Source/Editor/UMGEditor/Private/WidgetBlueprintEditor.h
+++ b/Engine/Source/Editor/UMGEditor/Private/WidgetBlueprintEditor.h
@@ -80,6 +80,9 @@ private:
 	/** Updates the inspector to be viewing the currently selected set of widgets */
 	void RefreshDetails();
 
+	/** Updates the inspector to be viewing a set of widgets */
+	void RefreshDetails(TSet<FWidgetReference>& Widgets);
+
 	/** Called whenever the blueprint is structurally changed. */
 	virtual void OnBlueprintChanged(UBlueprint* InBlueprint) override;
 
@@ -116,7 +119,7 @@ private:
 
 	/** The preview GUI object */
 	mutable TWeakObjectPtr<UUserWidget> PreviewWidgetActorPtr;
-	
+
 	/** Notification for when the preview widget has been updated  */
 	FOnWidgetPreviewUpdated OnWidgetPreviewUpdated;
 };
diff --git a/Engine/Source/Runtime/UMG/Private/UserWidget.cpp b/Engine/Source/Runtime/UMG/Private/UserWidget.cpp
--- a/Engine/Source/Runtime/UMG/Private/UserWidget.cpp
+++ b/Engine/Source/Runtime/UMG/Private/UserWidget.cpp
@@ -232,12 +232,11 @@ TSharedPtr<SWidget> UUserWidget::GetWidgetFromName(const FString& Name) const
 
 UWidget* UUserWidget::GetHandleFromName(const FString& Name) const
 {
-	for ( auto& Entry : WidgetToComponent )
+	for ( UWidget* Widget : Components )
 	{
-		if ( Entry.Value->GetName().Equals(Name, ESearchCase::IgnoreCase) )
+		if ( Widget->GetName().Equals(Name, ESearchCase::IgnoreCase) )
 		{
-			TWeakObjectPtr<UWidget> VisualWidget = Entry.Value;
-			return VisualWidget.Get();
+			return Widget;
 		}
 	}
 
