UMG - the '-umg' flag will no longer work.  UMG must be enabled using the UMG setting in the experimental area of the editor settings.

diff --git a/Engine/Source/Editor/Sequencer/Private/Sequencer.cpp b/Engine/Source/Editor/Sequencer/Private/Sequencer.cpp
--- a/Engine/Source/Editor/Sequencer/Private/Sequencer.cpp
+++ b/Engine/Source/Editor/Sequencer/Private/Sequencer.cpp
@@ -39,7 +39,12 @@ DEFINE_LOG_CATEGORY(LogSequencer);
 
 bool FSequencer::IsSequencerEnabled()
 {
-	return FParse::Param( FCommandLine::Get(), TEXT( "Sequencer" ) ) || FParse::Param( FCommandLine::Get(), TEXT( "umg" ) );
+	if ( GetDefault<UEditorExperimentalSettings>()->bUnrealMotionGraphics )
+	{
+		return true;
+	}
+
+	return FParse::Param(FCommandLine::Get(), TEXT("Sequencer"));
 }
 
 
diff --git a/Engine/Source/Editor/UMGEditor/Private/UMGEditorModule.cpp b/Engine/Source/Editor/UMGEditor/Private/UMGEditorModule.cpp
--- a/Engine/Source/Editor/UMGEditor/Private/UMGEditorModule.cpp
+++ b/Engine/Source/Editor/UMGEditor/Private/UMGEditorModule.cpp
@@ -32,23 +32,28 @@ public:
 		MenuExtensibilityManager = MakeShareable(new FExtensibilityManager());
 		ToolBarExtensibilityManager = MakeShareable(new FExtensibilityManager());
 
-		// Register asset types
-		IAssetTools& AssetTools = FModuleManager::LoadModuleChecked<FAssetToolsModule>("AssetTools").Get();
-		RegisterAssetTypeAction(AssetTools, MakeShareable(new FAssetTypeActions_WidgetBlueprint()));
-
-		// Register blueprint compiler
+		// Register widget blueprint compiler we do this no matter what.
 		IKismetCompilerInterface& KismetCompilerModule = FModuleManager::LoadModuleChecked<IKismetCompilerInterface>("KismetCompiler");
 		FBlueprintCompileDelegate sd;
 		sd.BindRaw(this, &FUMGEditorModule::CompileWidgetBlueprint);
 		KismetCompilerModule.GetCompilers().Add(sd);
-		
-		FUMGBlueprintEditorExtensionHook::InstallHooks();
 
+		bUMGEnabled = GetDefault<UEditorExperimentalSettings>()->bUnrealMotionGraphics;
+
+		// In the event we were loaded by an uncooked game, we need to check if we actually should be enabled.
+		// Because the UMG runtime will attempt to start our module no matter what, if WITH_EDITOR is enabled.
+		if ( bUMGEnabled )
+		{
+			// Register asset types
+			IAssetTools& AssetTools = FModuleManager::LoadModuleChecked<FAssetToolsModule>("AssetTools").Get();
+			RegisterAssetTypeAction(AssetTools, MakeShareable(new FAssetTypeActions_WidgetBlueprint()));
 
-		// Register with the sequencer module that we provide auto-key handlers.
-		ISequencerModule& SequencerModule = FModuleManager::Get().LoadModuleChecked<ISequencerModule>("Sequencer");
-		SequencerModule.RegisterTrackEditor(FOnCreateTrackEditor::CreateStatic(&FMarginTrackEditor::CreateTrackEditor));
+			FUMGBlueprintEditorExtensionHook::InstallHooks();
 
+			// Register with the sequencer module that we provide auto-key handlers.
+			ISequencerModule& SequencerModule = FModuleManager::Get().LoadModuleChecked<ISequencerModule>("Sequencer");
+			SequencerModule.RegisterTrackEditor(FOnCreateTrackEditor::CreateStatic(&FMarginTrackEditor::CreateTrackEditor));
+		}
 	}
 
 	/** Called before the module is unloaded, right before the module object is destroyed. */
@@ -57,21 +62,24 @@ public:
 		MenuExtensibilityManager.Reset();
 		ToolBarExtensibilityManager.Reset();
 
-		if ( UObjectInitialized() )
+		if ( bUMGEnabled )
 		{
-			FUMGBlueprintEditorExtensionHook::InstallHooks();
-		}
+			if ( UObjectInitialized() )
+			{
+				FUMGBlueprintEditorExtensionHook::RemoveHooks();
+			}
 
-		// Unregister all the asset types that we registered
-		if ( FModuleManager::Get().IsModuleLoaded("AssetTools") )
-		{
-			IAssetTools& AssetTools = FModuleManager::GetModuleChecked<FAssetToolsModule>("AssetTools").Get();
-			for ( int32 Index = 0; Index < CreatedAssetTypeActions.Num(); ++Index )
+			// Unregister all the asset types that we registered
+			if ( FModuleManager::Get().IsModuleLoaded("AssetTools") )
 			{
-				AssetTools.UnregisterAssetTypeActions(CreatedAssetTypeActions[Index].ToSharedRef());
+				IAssetTools& AssetTools = FModuleManager::GetModuleChecked<FAssetToolsModule>("AssetTools").Get();
+				for ( int32 Index = 0; Index < CreatedAssetTypeActions.Num(); ++Index )
+				{
+					AssetTools.UnregisterAssetTypeActions(CreatedAssetTypeActions[Index].ToSharedRef());
+				}
 			}
+			CreatedAssetTypeActions.Empty();
 		}
-		CreatedAssetTypeActions.Empty();
 	}
 
 	FReply CompileWidgetBlueprint(UBlueprint* Blueprint, const FKismetCompilerOptions& CompileOptions, FCompilerResultsLog& Results, TArray<UObject*>* ObjLoaded)
@@ -89,8 +97,8 @@ public:
 	}
 
 	/** Gets the extensibility managers for outside entities to extend gui page editor's menus and toolbars */
-	virtual TSharedPtr<FExtensibilityManager> GetMenuExtensibilityManager() {return MenuExtensibilityManager;}
-	virtual TSharedPtr<FExtensibilityManager> GetToolBarExtensibilityManager() {return ToolBarExtensibilityManager;}
+	virtual TSharedPtr<FExtensibilityManager> GetMenuExtensibilityManager() { return MenuExtensibilityManager; }
+	virtual TSharedPtr<FExtensibilityManager> GetToolBarExtensibilityManager() { return ToolBarExtensibilityManager; }
 
 private:
 	void RegisterAssetTypeAction(IAssetTools& AssetTools, TSharedRef<IAssetTypeActions> Action)
@@ -103,6 +111,9 @@ private:
 	TSharedPtr<FExtensibilityManager> MenuExtensibilityManager;
 	TSharedPtr<FExtensibilityManager> ToolBarExtensibilityManager;
 
+	/** We cache the value for UMG being enabled on load so that we can shutdown properly. */
+	bool bUMGEnabled;
+
 	/** All created asset type actions.  Cached here so that we can unregister it during shutdown. */
 	TArray< TSharedPtr<IAssetTypeActions> > CreatedAssetTypeActions;
 };
diff --git a/Engine/Source/Editor/UnrealEd/Classes/Settings/EditorExperimentalSettings.h b/Engine/Source/Editor/UnrealEd/Classes/Settings/EditorExperimentalSettings.h
--- a/Engine/Source/Editor/UnrealEd/Classes/Settings/EditorExperimentalSettings.h
+++ b/Engine/Source/Editor/UnrealEd/Classes/Settings/EditorExperimentalSettings.h
@@ -21,6 +21,9 @@ class UNREALED_API UEditorExperimentalSettings
 
 public:
 
+	UPROPERTY(EditAnywhere, config, Category = Tools, meta = ( DisplayName = "Unreal Motion Graphics (UMG)" ))
+	bool bUnrealMotionGraphics;
+
 	/** Allows to open editor for Behavior Tree assets */
 	UPROPERTY(EditAnywhere, config, Category=Tools, meta=(DisplayName="Behavior Tree Editor"))
 	bool bBehaviorTreeEditor;
diff --git a/Engine/Source/Editor/UnrealEd/Private/Editor.cpp b/Engine/Source/Editor/UnrealEd/Private/Editor.cpp
--- a/Engine/Source/Editor/UnrealEd/Private/Editor.cpp
+++ b/Engine/Source/Editor/UnrealEd/Private/Editor.cpp
@@ -558,11 +558,6 @@ void UEditorEngine::Init(IEngineLoop* InEngineLoop)
 			FModuleManager::Get().LoadModule(TEXT("IntroTutorials"));
 		}
 
-		if ( FParse::Param(FCommandLine::Get(), TEXT("umg")) )
-		{
-			FModuleManager::Get().LoadModule(TEXT("UMGEditor"));
-		}
-
 		if( FParse::Param( FCommandLine::Get(),TEXT( "PListEditor" ) ) )
 		{
 			FModuleManager::Get().LoadModule(TEXT("PListEditor"));
diff --git a/Engine/Source/Runtime/UMG/Private/UMGModule.cpp b/Engine/Source/Runtime/UMG/Private/UMGModule.cpp
--- a/Engine/Source/Runtime/UMG/Private/UMGModule.cpp
+++ b/Engine/Source/Runtime/UMG/Private/UMGModule.cpp
@@ -21,10 +21,7 @@ public:
 
 		// This is done so that the compiler is available in non-cooked builds when the widget blueprint is 
 		// compiled again in the running game.
-		if ( FParse::Param(FCommandLine::Get(), TEXT("umg")) )
-		{
-			FModuleManager::Get().LoadModule(TEXT("UMGEditor"));
-		}
+		FModuleManager::Get().LoadModule(TEXT("UMGEditor"));
 #endif
 	}
 
