UMG - Changing keyboard focus to the designer on mouse down.

diff --git a/Engine/Source/Editor/UMGEditor/Private/SUMGDesigner.cpp b/Engine/Source/Editor/UMGEditor/Private/SUMGDesigner.cpp
--- a/Engine/Source/Editor/UMGEditor/Private/SUMGDesigner.cpp
+++ b/Engine/Source/Editor/UMGEditor/Private/SUMGDesigner.cpp
@@ -353,7 +353,7 @@ FReply SUMGDesigner::OnMouseButtonDown(const FGeometry& MyGeometry, const FPoint
 	}
 
 	// Capture mouse for the drag handle and general mouse actions
-	return FReply::Handled().PreventThrottling().CaptureMouse(AsShared());
+	return FReply::Handled().PreventThrottling().SetKeyboardFocus(AsShared(), EKeyboardFocusCause::Mouse).CaptureMouse(AsShared());
 }
 
 FReply SUMGDesigner::OnMouseButtonUp(const FGeometry& MyGeometry, const FPointerEvent& MouseEvent)
