UMG - Complete break from the hack to patch UMG into the original Blueprint editor.  This fixes the crash that occurs with running UMG and Persona.  UMG now works like Persona, completely new modes for each tab, and custom tab UI.

diff --git a/Engine/Source/Editor/Kismet/Private/BlueprintEditor.cpp b/Engine/Source/Editor/Kismet/Private/BlueprintEditor.cpp
--- a/Engine/Source/Editor/Kismet/Private/BlueprintEditor.cpp
+++ b/Engine/Source/Editor/Kismet/Private/BlueprintEditor.cpp
@@ -1096,9 +1096,6 @@ void FBlueprintEditor::InitBlueprintEditor(const EToolkitMode::Type Mode, const
 {
 	check(InBlueprints.Num() == 1 || bShouldOpenInDefaultsMode);
 
-	// Newly-created Blueprints will open in Components mode rather than Standard mode
-	bool bShouldOpenInComponentsMode = !bShouldOpenInDefaultsMode && InBlueprints.Num() == 1 && InBlueprints[0]->bIsNewlyCreated;
-
 	TArray< UObject* > Objects;
 	for( auto BlueprintIter = InBlueprints.CreateConstIterator(); BlueprintIter; ++BlueprintIter )
 	{
@@ -1134,26 +1131,50 @@ void FBlueprintEditor::InitBlueprintEditor(const EToolkitMode::Type Mode, const
 	
 	RegenerateMenusAndToolbars();
 
+	RegisterApplicationModes(InBlueprints, bShouldOpenInDefaultsMode);
+
+	// Cache the project type ( Blueprint or Code Based )
+	if( FPaths::IsProjectFilePathSet() )
+	{
+		FProjectStatus ProjectStatus;
+		if( IProjectManager::Get().QueryStatusForProject( FPaths::GetProjectFilePath(), ProjectStatus ) && ProjectStatus.bCodeBasedProject )
+		{
+			bCodeBasedProject = true;
+		}
+	}
+
+	// Post-layout initialization
+	PostLayoutBlueprintEditorInitialization();
+
+	// Find and set any instances of this blueprint type if any exists and we are not already editing one
+	FBlueprintEditorUtils::FindAndSetDebuggableBlueprintInstances();	
+}
+
+void FBlueprintEditor::RegisterApplicationModes(const TArray<UBlueprint*>& InBlueprints, bool bShouldOpenInDefaultsMode)
+{
+	// Newly-created Blueprints will open in Components mode rather than Standard mode
+	bool bShouldOpenInComponentsMode = !bShouldOpenInDefaultsMode && InBlueprints.Num() == 1 && InBlueprints[0]->bIsNewlyCreated;
+
 	// Create the modes and activate one (which will populate with a real layout)
-	if (UBlueprint* SingleBP = GetBlueprintObj())
+	if ( UBlueprint* SingleBP = GetBlueprintObj() )
 	{
-		if (!bShouldOpenInDefaultsMode && FBlueprintEditorUtils::IsInterfaceBlueprint(SingleBP))
+		if ( !bShouldOpenInDefaultsMode && FBlueprintEditorUtils::IsInterfaceBlueprint(SingleBP) )
 		{
 			// Interfaces are only valid in the Interface mode
 			AddApplicationMode(
-				FBlueprintEditorApplicationModes::BlueprintInterfaceMode, 
-				MakeShareable(new FBlueprintInterfaceApplicationMode( SharedThis(this) )));
+				FBlueprintEditorApplicationModes::BlueprintInterfaceMode,
+				MakeShareable(new FBlueprintInterfaceApplicationMode(SharedThis(this))));
 			SetCurrentMode(FBlueprintEditorApplicationModes::BlueprintInterfaceMode);
 		}
-		else if (SingleBP->BlueprintType == BPTYPE_MacroLibrary)
+		else if ( SingleBP->BlueprintType == BPTYPE_MacroLibrary )
 		{
 			// Macro libraries are only valid in the Macro mode
 			AddApplicationMode(
-				FBlueprintEditorApplicationModes::BlueprintMacroMode, 
-				MakeShareable(new FBlueprintMacroApplicationMode( SharedThis(this) )));
+				FBlueprintEditorApplicationModes::BlueprintMacroMode,
+				MakeShareable(new FBlueprintMacroApplicationMode(SharedThis(this))));
 			SetCurrentMode(FBlueprintEditorApplicationModes::BlueprintMacroMode);
 		}
-		else if (SingleBP->BlueprintType == BPTYPE_FunctionLibrary)
+		else if ( SingleBP->BlueprintType == BPTYPE_FunctionLibrary )
 		{
 			AddApplicationMode(
 				FBlueprintEditorApplicationModes::StandardBlueprintEditorMode,
@@ -1163,20 +1184,20 @@ void FBlueprintEditor::InitBlueprintEditor(const EToolkitMode::Type Mode, const
 		else
 		{
 			AddApplicationMode(
-				FBlueprintEditorApplicationModes::StandardBlueprintEditorMode, 
+				FBlueprintEditorApplicationModes::StandardBlueprintEditorMode,
 				MakeShareable(new FBlueprintEditorApplicationMode(SharedThis(this), FBlueprintEditorApplicationModes::StandardBlueprintEditorMode)));
 			AddApplicationMode(
-				FBlueprintEditorApplicationModes::BlueprintDefaultsMode, 
+				FBlueprintEditorApplicationModes::BlueprintDefaultsMode,
 				MakeShareable(new FBlueprintDefaultsApplicationMode(SharedThis(this))));
 			AddApplicationMode(
-				FBlueprintEditorApplicationModes::BlueprintComponentsMode, 
+				FBlueprintEditorApplicationModes::BlueprintComponentsMode,
 				MakeShareable(new FBlueprintComponentsApplicationMode(SharedThis(this))));
 
-			if (bShouldOpenInDefaultsMode)
+			if ( bShouldOpenInDefaultsMode )
 			{
 				SetCurrentMode(FBlueprintEditorApplicationModes::BlueprintDefaultsMode);
 			}
-			else if (bShouldOpenInComponentsMode && CanAccessComponentsMode())
+			else if ( bShouldOpenInComponentsMode && CanAccessComponentsMode() )
 			{
 				SetCurrentMode(FBlueprintEditorApplicationModes::BlueprintComponentsMode);
 			}
@@ -1190,26 +1211,10 @@ void FBlueprintEditor::InitBlueprintEditor(const EToolkitMode::Type Mode, const
 	{
 		// We either have no blueprints or many, open in the defaults mode for multi-editing
 		AddApplicationMode(
-			FBlueprintEditorApplicationModes::BlueprintDefaultsMode, 
+			FBlueprintEditorApplicationModes::BlueprintDefaultsMode,
 			MakeShareable(new FBlueprintDefaultsApplicationMode(SharedThis(this))));
 		SetCurrentMode(FBlueprintEditorApplicationModes::BlueprintDefaultsMode);
 	}
-
-	// Cache the project type ( Blueprint or Code Based )
-	if( FPaths::IsProjectFilePathSet() )
-	{
-		FProjectStatus ProjectStatus;
-		if( IProjectManager::Get().QueryStatusForProject( FPaths::GetProjectFilePath(), ProjectStatus ) && ProjectStatus.bCodeBasedProject )
-		{
-			bCodeBasedProject = true;
-		}
-	}
-
-	// Post-layout initialization
-	PostLayoutBlueprintEditorInitialization();
-
-	// Find and set any instances of this blueprint type if any exists and we are not already editing one
-	FBlueprintEditorUtils::FindAndSetDebuggableBlueprintInstances();	
 }
 
 void FBlueprintEditor::PostRegenerateMenusAndToolbars()
diff --git a/Engine/Source/Editor/Kismet/Public/BlueprintEditor.h b/Engine/Source/Editor/Kismet/Public/BlueprintEditor.h
--- a/Engine/Source/Editor/Kismet/Public/BlueprintEditor.h
+++ b/Engine/Source/Editor/Kismet/Public/BlueprintEditor.h
@@ -473,6 +473,8 @@ public:
 	bool IsGraphInCurrentBlueprint(UEdGraph* InGraph) const;
 
 protected:
+	/** Called during initialization of the blueprint editor to register any application modes. */
+	virtual void RegisterApplicationModes(const TArray<UBlueprint*>& InBlueprints, bool bShouldOpenInDefaultsMode);
 	
 	// Zooming to fit the entire graph
 	void ZoomToWindow_Clicked();
diff --git a/Engine/Source/Editor/UMGEditor/Private/BlueprintModes/WidgetBlueprintApplicationMode.cpp b/Engine/Source/Editor/UMGEditor/Private/BlueprintModes/WidgetBlueprintApplicationMode.cpp
new file mode 100644
--- /dev/null
+++ b/Engine/Source/Editor/UMGEditor/Private/BlueprintModes/WidgetBlueprintApplicationMode.cpp
@@ -0,0 +1,37 @@
+// Copyright 1998-2014 Epic Games, Inc. All Rights Reserved.
+
+#include "UMGEditorPrivatePCH.h"
+
+#include "Editor/PropertyEditor/Public/PropertyEditorModule.h"
+#include "Editor/Kismet/Public/BlueprintEditorTabs.h"
+#include "Editor/Kismet/Public/SBlueprintEditorToolbar.h"
+#include "Editor/Kismet/Public/BlueprintEditorModes.h"
+
+#include "WidgetBlueprintEditor.h"
+#include "WidgetBlueprintApplicationMode.h"
+
+/////////////////////////////////////////////////////
+// FWidgetBlueprintApplicationMode
+
+FWidgetBlueprintApplicationMode::FWidgetBlueprintApplicationMode(TSharedPtr<FWidgetBlueprintEditor> InWidgetEditor, FName InModeName)
+	: FBlueprintEditorApplicationMode(InWidgetEditor, InModeName, false, false)
+	, MyWidgetBlueprintEditor(InWidgetEditor)
+{
+}
+
+UWidgetBlueprint* FWidgetBlueprintApplicationMode::GetBlueprint() const
+{
+	if ( FWidgetBlueprintEditor* Editor = MyWidgetBlueprintEditor.Pin().Get() )
+	{
+		return Editor->GetWidgetBlueprintObj();
+	}
+	else
+	{
+		return NULL;
+	}
+}
+
+TSharedPtr<FWidgetBlueprintEditor> FWidgetBlueprintApplicationMode::GetBlueprintEditor() const
+{
+	return MyWidgetBlueprintEditor.Pin();
+}
diff --git a/Engine/Source/Editor/UMGEditor/Private/BlueprintModes/WidgetBlueprintApplicationMode.h b/Engine/Source/Editor/UMGEditor/Private/BlueprintModes/WidgetBlueprintApplicationMode.h
new file mode 100644
--- /dev/null
+++ b/Engine/Source/Editor/UMGEditor/Private/BlueprintModes/WidgetBlueprintApplicationMode.h
@@ -0,0 +1,25 @@
+// Copyright 1998-2014 Epic Games, Inc. All Rights Reserved.
+
+#pragma once
+
+/////////////////////////////////////////////////////
+// FWidgetBlueprintApplicationMode
+
+class FWidgetBlueprintApplicationMode : public FBlueprintEditorApplicationMode
+{
+public:
+	FWidgetBlueprintApplicationMode(TSharedPtr<class FWidgetBlueprintEditor> InWidgetEditor, FName InModeName);
+
+	// FApplicationMode interface
+	// End of FApplicationMode interface
+
+protected:
+	UWidgetBlueprint* GetBlueprint() const;
+	
+	TSharedPtr<class FWidgetBlueprintEditor> GetBlueprintEditor() const;
+
+	TWeakPtr<class FWidgetBlueprintEditor> MyWidgetBlueprintEditor;
+
+	// Set of spawnable tabs in the mode
+	FWorkflowAllowedTabSet TabFactories;
+};
diff --git a/Engine/Source/Editor/UMGEditor/Private/BlueprintModes/WidgetBlueprintApplicationModes.cpp b/Engine/Source/Editor/UMGEditor/Private/BlueprintModes/WidgetBlueprintApplicationModes.cpp
new file mode 100644
--- /dev/null
+++ b/Engine/Source/Editor/UMGEditor/Private/BlueprintModes/WidgetBlueprintApplicationModes.cpp
@@ -0,0 +1,24 @@
+// Copyright 1998-2014 Epic Games, Inc. All Rights Reserved.
+
+#include "UMGEditorPrivatePCH.h"
+
+// Mode constants
+const FName FWidgetBlueprintApplicationModes::DesignerMode("DesignModeName");
+const FName FWidgetBlueprintApplicationModes::GraphMode("GraphModeName");
+
+FText FWidgetBlueprintApplicationModes::GetLocalizedMode(const FName InMode)
+{
+	static TMap< FName, FText > LocModes;
+
+	if ( LocModes.Num() == 0 )
+	{
+		LocModes.Add(DesignerMode, NSLOCTEXT("WidgetBlueprintModes", "DesignerMode", "Designer"));
+		LocModes.Add(GraphMode, NSLOCTEXT("WidgetBlueprintModes", "GraphMode", "Graph"));
+	}
+
+	check(InMode != NAME_None);
+	const FText* OutDesc = LocModes.Find(InMode);
+	check(OutDesc);
+
+	return *OutDesc;
+}
diff --git a/Engine/Source/Editor/UMGEditor/Private/BlueprintModes/WidgetBlueprintApplicationModes.h b/Engine/Source/Editor/UMGEditor/Private/BlueprintModes/WidgetBlueprintApplicationModes.h
new file mode 100644
--- /dev/null
+++ b/Engine/Source/Editor/UMGEditor/Private/BlueprintModes/WidgetBlueprintApplicationModes.h
@@ -0,0 +1,19 @@
+// Copyright 1998-2014 Epic Games, Inc. All Rights Reserved.
+
+#pragma once
+
+/////////////////////////////////////////////////////
+// FWidgetBlueprintApplicationModes
+
+// This is the list of IDs for widget blueprint editor modes
+struct FWidgetBlueprintApplicationModes
+{
+	// Mode constants
+	static const FName DesignerMode;
+	static const FName GraphMode;
+
+	static FText GetLocalizedMode(const FName InMode);
+
+private:
+	FWidgetBlueprintApplicationModes() {}
+};
\ No newline at end of file
diff --git a/Engine/Source/Editor/UMGEditor/Private/BlueprintModes/WidgetDesignerApplicationMode.cpp b/Engine/Source/Editor/UMGEditor/Private/BlueprintModes/WidgetDesignerApplicationMode.cpp
new file mode 100644
--- /dev/null
+++ b/Engine/Source/Editor/UMGEditor/Private/BlueprintModes/WidgetDesignerApplicationMode.cpp
@@ -0,0 +1,142 @@
+// Copyright 1998-2014 Epic Games, Inc. All Rights Reserved.
+
+#include "UMGEditorPrivatePCH.h"
+
+#include "Editor/PropertyEditor/Public/PropertyEditorModule.h"
+#include "Editor/Kismet/Public/BlueprintEditorTabs.h"
+#include "Editor/Kismet/Public/SBlueprintEditorToolbar.h"
+#include "Editor/Kismet/Public/BlueprintEditorModes.h"
+
+#include "WidgetBlueprintEditor.h"
+#include "WidgetDesignerApplicationMode.h"
+#include "WidgetBlueprintEditorToolbar.h"
+
+#include "IDetailsView.h"
+
+#include "BlueprintEditorTabs.h"
+#include "PaletteTabSummoner.h"
+#include "HierarchyTabSummoner.h"
+#include "DesignerTabSummoner.h"
+#include "SequencerTabSummoner.h"
+
+/////////////////////////////////////////////////////
+// FWidgetDesignerApplicationMode
+
+FWidgetDesignerApplicationMode::FWidgetDesignerApplicationMode(TSharedPtr<FWidgetBlueprintEditor> InWidgetEditor)
+	: FWidgetBlueprintApplicationMode(InWidgetEditor, FWidgetBlueprintApplicationModes::DesignerMode)
+{
+	TabLayout = FTabManager::NewLayout( "WidgetBlueprintEditor_Designer_Layout_v1" )
+	->AddArea
+	(
+		FTabManager::NewPrimaryArea()
+		->SetOrientation(Orient_Vertical)
+		->Split
+		(
+			FTabManager::NewStack()
+			->SetSizeCoefficient( 0.2f )
+			->SetHideTabWell(true)
+			->AddTab( InWidgetEditor->GetToolbarTabId(), ETabState::OpenedTab )
+		)
+		->Split
+		(
+			FTabManager::NewSplitter()
+			->SetOrientation(Orient_Horizontal)
+			->Split
+			(
+				FTabManager::NewSplitter()
+				->SetSizeCoefficient( 0.15f )
+				->SetOrientation(Orient_Vertical)
+				->Split
+				(
+					FTabManager::NewStack()
+					->SetSizeCoefficient( 0.5f )
+					->AddTab( FPaletteTabSummoner::TabID, ETabState::OpenedTab )
+				)
+				->Split
+				(
+					FTabManager::NewStack()
+					->SetSizeCoefficient( 0.5f )
+					->AddTab( FHierarchyTabSummoner::TabID, ETabState::OpenedTab )
+				)
+			)
+			->Split
+			(
+				FTabManager::NewSplitter()
+				->SetSizeCoefficient( 0.85f )
+				->SetOrientation(Orient_Horizontal)
+				->Split
+				(
+					FTabManager::NewStack()
+					->SetHideTabWell(true)
+					->AddTab( FDesignerTabSummoner::TabID, ETabState::OpenedTab )
+				)
+				->Split
+				(
+					FTabManager::NewStack()
+					->SetSizeCoefficient( 0.5f )
+					->AddTab( FBlueprintEditorTabs::DetailsID, ETabState::OpenedTab )
+				)
+			)
+		)
+		->Split
+		(
+			FTabManager::NewStack()
+			->SetSizeCoefficient(0.2f)
+			->SetHideTabWell(true)
+			->AddTab( FSequencerTabSummoner::TabID, ETabState::OpenedTab )
+		)
+	);
+
+
+	// Add Tab Spawners
+	TabFactories.RegisterFactory(MakeShareable(new FSelectionDetailsSummoner(InWidgetEditor)));
+	TabFactories.RegisterFactory(MakeShareable(new FDesignerTabSummoner(InWidgetEditor)));
+	TabFactories.RegisterFactory(MakeShareable(new FHierarchyTabSummoner(InWidgetEditor)));
+	TabFactories.RegisterFactory(MakeShareable(new FPaletteTabSummoner(InWidgetEditor)));
+	TabFactories.RegisterFactory(MakeShareable(new FSequencerTabSummoner(InWidgetEditor)));
+
+	// setup toolbar - clear existing toolbar extender from the BP mode
+	//@TODO: Keep this in sync with BlueprintEditorModes.cpp
+	ToolbarExtender = MakeShareable(new FExtender);
+	InWidgetEditor->GetWidgetToolbarBuilder()->AddWidgetBlueprintEditorModesToolbar(ToolbarExtender);
+	InWidgetEditor->GetToolbarBuilder()->AddCompileToolbar(ToolbarExtender);
+	//InPersona->GetToolbarBuilder()->AddScriptingToolbar(ToolbarExtender);
+	//InPersona->GetToolbarBuilder()->AddBlueprintGlobalOptionsToolbar(ToolbarExtender);
+	//InPersona->GetToolbarBuilder()->AddDebuggingToolbar(ToolbarExtender);
+}
+
+void FWidgetDesignerApplicationMode::RegisterTabFactories(TSharedPtr<FTabManager> InTabManager)
+{
+	TSharedPtr<FWidgetBlueprintEditor> BP = GetBlueprintEditor();
+	
+	BP->RegisterToolbarTab(InTabManager.ToSharedRef());
+	BP->PushTabFactories(TabFactories);
+}
+
+void FWidgetDesignerApplicationMode::PreDeactivateMode()
+{
+	//FWidgetBlueprintApplicationMode::PreDeactivateMode();
+
+	GetBlueprintEditor()->GetInspector()->EnableComponentDetailsCustomization(false);
+
+	static FName PropertyEditor("PropertyEditor");
+	FPropertyEditorModule& PropertyModule = FModuleManager::GetModuleChecked<FPropertyEditorModule>(PropertyEditor);
+	PropertyModule.UnregisterCustomPropertyTypeLayout(TEXT("CanvasPanelSlot"));
+}
+
+void FWidgetDesignerApplicationMode::PostActivateMode()
+{
+	//FWidgetBlueprintApplicationMode::PostActivateMode();
+
+	GetBlueprintEditor()->GetInspector()->EnableComponentDetailsCustomization(false);
+
+	TSharedRef<class SKismetInspector> Inspector = MyBlueprintEditor.Pin()->GetInspector();
+	FOnGetDetailCustomizationInstance LayoutDelegateDetails = FOnGetDetailCustomizationInstance::CreateStatic(&FBlueprintWidgetCustomization::MakeInstance, MyBlueprintEditor.Pin()->GetBlueprintObj());
+	Inspector->GetPropertyView()->RegisterInstancedCustomPropertyLayout(UWidget::StaticClass(), LayoutDelegateDetails);
+
+	FOnGetPropertyTypeCustomizationInstance CanvasSlotCustomization = FOnGetPropertyTypeCustomizationInstance::CreateStatic(&FCanvasSlotCustomization::MakeInstance, MyBlueprintEditor.Pin()->GetBlueprintObj());
+
+	static FName PropertyEditor("PropertyEditor");
+	FPropertyEditorModule& PropertyModule = FModuleManager::GetModuleChecked<FPropertyEditorModule>(PropertyEditor);
+	PropertyModule.RegisterCustomPropertyTypeLayout(TEXT("PanelSlot"), CanvasSlotCustomization);
+}
diff --git a/Engine/Source/Editor/UMGEditor/Private/BlueprintModes/WidgetDesignerApplicationMode.h b/Engine/Source/Editor/UMGEditor/Private/BlueprintModes/WidgetDesignerApplicationMode.h
new file mode 100644
--- /dev/null
+++ b/Engine/Source/Editor/UMGEditor/Private/BlueprintModes/WidgetDesignerApplicationMode.h
@@ -0,0 +1,18 @@
+// Copyright 1998-2014 Epic Games, Inc. All Rights Reserved.
+
+#pragma once
+
+/////////////////////////////////////////////////////
+// FWidgetDesignerApplicationMode
+
+class FWidgetDesignerApplicationMode : public FWidgetBlueprintApplicationMode
+{
+public:
+	FWidgetDesignerApplicationMode(TSharedPtr<FWidgetBlueprintEditor> InWidgetEditor);
+
+	// FApplicationMode interface
+	virtual void RegisterTabFactories(TSharedPtr<FTabManager> InTabManager) override;
+	virtual void PreDeactivateMode() override;
+	virtual void PostActivateMode() override;
+	// End of FApplicationMode interface
+};
diff --git a/Engine/Source/Editor/UMGEditor/Private/BlueprintModes/WidgetGraphApplicationMode.cpp b/Engine/Source/Editor/UMGEditor/Private/BlueprintModes/WidgetGraphApplicationMode.cpp
new file mode 100644
--- /dev/null
+++ b/Engine/Source/Editor/UMGEditor/Private/BlueprintModes/WidgetGraphApplicationMode.cpp
@@ -0,0 +1,103 @@
+// Copyright 1998-2014 Epic Games, Inc. All Rights Reserved.
+
+#include "UMGEditorPrivatePCH.h"
+
+#include "Editor/PropertyEditor/Public/PropertyEditorModule.h"
+#include "Editor/Kismet/Public/BlueprintEditorTabs.h"
+#include "Editor/Kismet/Public/SBlueprintEditorToolbar.h"
+#include "Editor/Kismet/Public/BlueprintEditorModes.h"
+
+#include "WidgetBlueprintEditor.h"
+#include "WidgetGraphApplicationMode.h"
+#include "WidgetBlueprintEditorToolbar.h"
+
+/////////////////////////////////////////////////////
+// FWidgetGraphApplicationMode
+
+FWidgetGraphApplicationMode::FWidgetGraphApplicationMode(TSharedPtr<FWidgetBlueprintEditor> InWidgetEditor)
+	: FWidgetBlueprintApplicationMode(InWidgetEditor, FWidgetBlueprintApplicationModes::GraphMode)
+{
+	TabLayout = FTabManager::NewLayout( "WidgetBlueprintEditor_Graph_Layout_v1" )
+		->AddArea
+		(
+			FTabManager::NewPrimaryArea() ->SetOrientation(Orient_Vertical)
+			->Split
+			(
+				FTabManager::NewStack()
+				->SetSizeCoefficient( 0.186721f )
+				->SetHideTabWell(true)
+				->AddTab( InWidgetEditor->GetToolbarTabId(), ETabState::OpenedTab )
+			)
+			->Split
+			(
+				FTabManager::NewSplitter() ->SetOrientation(Orient_Horizontal)
+				->Split
+				(
+					FTabManager::NewSplitter() ->SetOrientation( Orient_Vertical )
+					->SetSizeCoefficient(0.15f)
+					->Split
+					(
+						FTabManager::NewStack() ->SetSizeCoefficient(0.5f)
+						->AddTab( FBlueprintEditorTabs::MyBlueprintID, ETabState::OpenedTab )
+					)
+					->Split
+					(
+						FTabManager::NewStack() ->SetSizeCoefficient(0.5f)
+						->AddTab( FBlueprintEditorTabs::DetailsID, ETabState::OpenedTab )
+					)
+				)
+				->Split
+				(
+					FTabManager::NewSplitter() ->SetOrientation( Orient_Vertical )
+					->SetSizeCoefficient(0.70f)
+					->Split
+					(
+						FTabManager::NewStack()
+						->SetSizeCoefficient( 0.80f )
+						->AddTab( "Document", ETabState::ClosedTab )
+					)
+					->Split
+					(
+						FTabManager::NewStack()
+						->SetSizeCoefficient( 0.20f )
+						->AddTab( FBlueprintEditorTabs::CompilerResultsID, ETabState::ClosedTab )
+						->AddTab( FBlueprintEditorTabs::FindResultsID, ETabState::ClosedTab )
+					)
+				)
+				->Split
+				(
+					FTabManager::NewSplitter() ->SetOrientation( Orient_Vertical )
+					->SetSizeCoefficient(0.15f)
+					->Split
+					(
+						FTabManager::NewStack()
+						->AddTab( FBlueprintEditorTabs::PaletteID, ETabState::OpenedTab )
+					)
+				)
+			)
+		);
+	
+	// setup toolbar
+	//@TODO: Keep this in sync with AnimBlueprintMode.cpp
+	ToolbarExtender = MakeShareable(new FExtender);
+	InWidgetEditor->GetWidgetToolbarBuilder()->AddWidgetBlueprintEditorModesToolbar(ToolbarExtender);
+	InWidgetEditor->GetToolbarBuilder()->AddCompileToolbar(ToolbarExtender);
+	InWidgetEditor->GetToolbarBuilder()->AddScriptingToolbar(ToolbarExtender);
+	InWidgetEditor->GetToolbarBuilder()->AddBlueprintGlobalOptionsToolbar(ToolbarExtender);
+	InWidgetEditor->GetToolbarBuilder()->AddDebuggingToolbar(ToolbarExtender);
+}
+
+void FWidgetGraphApplicationMode::RegisterTabFactories(TSharedPtr<FTabManager> InTabManager)
+{
+	TSharedPtr<FBlueprintEditor> BP = GetBlueprintEditor();
+	
+	BP->RegisterToolbarTab(InTabManager.ToSharedRef());
+	BP->PushTabFactories(CoreTabFactories);
+	BP->PushTabFactories(BlueprintEditorTabFactories);
+	BP->PushTabFactories(TabFactories);
+}
+
+void FWidgetGraphApplicationMode::PostActivateMode()
+{
+	FWidgetBlueprintApplicationMode::PostActivateMode();
+}
diff --git a/Engine/Source/Editor/UMGEditor/Private/BlueprintModes/WidgetGraphApplicationMode.h b/Engine/Source/Editor/UMGEditor/Private/BlueprintModes/WidgetGraphApplicationMode.h
new file mode 100644
--- /dev/null
+++ b/Engine/Source/Editor/UMGEditor/Private/BlueprintModes/WidgetGraphApplicationMode.h
@@ -0,0 +1,17 @@
+// Copyright 1998-2014 Epic Games, Inc. All Rights Reserved.
+
+#pragma once
+
+/////////////////////////////////////////////////////
+// FWidgetGraphApplicationMode
+
+class FWidgetGraphApplicationMode : public FWidgetBlueprintApplicationMode
+{
+public:
+	FWidgetGraphApplicationMode(TSharedPtr<FWidgetBlueprintEditor> InWidgetEditor);
+
+	// FApplicationMode interface
+	virtual void RegisterTabFactories(TSharedPtr<FTabManager> InTabManager) override;
+	virtual void PostActivateMode() override;
+	// End of FApplicationMode interface
+};
diff --git a/Engine/Source/Editor/UMGEditor/Private/SUMGEditorWidgetTemplates.cpp b/Engine/Source/Editor/UMGEditor/Private/SUMGEditorWidgetTemplates.cpp
--- a/Engine/Source/Editor/UMGEditor/Private/SUMGEditorWidgetTemplates.cpp
+++ b/Engine/Source/Editor/UMGEditor/Private/SUMGEditorWidgetTemplates.cpp
@@ -42,15 +42,17 @@ public:
 
 			+ SHorizontalBox::Slot()
 			.AutoWidth()
-			.Padding(2.0f)
+			.VAlign(VAlign_Center)
 			[
 				SNew(SImage)
-				.Image(Template->Icon.GetIcon())
+				.ColorAndOpacity(FLinearColor(1, 1, 1, 0.5))
+				.Image(Template->GetIcon())
 			]
 
 			+ SHorizontalBox::Slot()
 			.AutoWidth()
-			.Padding(2.0f)
+			.Padding(2, 0, 0, 0)
+			.VAlign(VAlign_Center)
 			[
 				SNew(STextBlock)
 				.Text(Template->Name)
diff --git a/Engine/Source/Editor/UMGEditor/Private/TabFactory/DesignerTabSummoner.cpp b/Engine/Source/Editor/UMGEditor/Private/TabFactory/DesignerTabSummoner.cpp
new file mode 100644
--- /dev/null
+++ b/Engine/Source/Editor/UMGEditor/Private/TabFactory/DesignerTabSummoner.cpp
@@ -0,0 +1,37 @@
+// Copyright 1998-2014 Epic Games, Inc. All Rights Reserved.
+
+#include "UMGEditorPrivatePCH.h"
+
+#include "DesignerTabSummoner.h"
+
+#define LOCTEXT_NAMESPACE "UMG"
+
+const FName FDesignerTabSummoner::TabID(TEXT("SlatePreview"));
+
+FDesignerTabSummoner::FDesignerTabSummoner(TSharedPtr<class FWidgetBlueprintEditor> InBlueprintEditor)
+		: FWorkflowTabFactory(TabID, InBlueprintEditor)
+		, BlueprintEditor(InBlueprintEditor)
+{
+	TabLabel = LOCTEXT("DesignerTabLabel", "Designer");
+	TabIcon = FEditorStyle::GetBrush("UMGEditor.Tabs.Designer");
+
+	bIsSingleton = true;
+
+	ViewMenuDescription = LOCTEXT("SlatePreview_ViewMenu_Desc", "Designer");
+	ViewMenuTooltip = LOCTEXT("SlatePreview_ViewMenu_ToolTip", "Show the Designer");
+}
+
+TSharedRef<SWidget> FDesignerTabSummoner::CreateTabBody(const FWorkflowTabSpawnInfo& Info) const
+{
+	return SNew(SHorizontalBox)
+		+ SHorizontalBox::Slot()
+		.FillWidth(1)
+		[
+			SNew(STutorialWrapper, TEXT("Designer"))
+			[
+				SNew(SUMGDesigner, BlueprintEditor.Pin())
+			]
+		];
+}
+
+#undef LOCTEXT_NAMESPACE 
diff --git a/Engine/Source/Editor/UMGEditor/Private/TabFactory/DesignerTabSummoner.h b/Engine/Source/Editor/UMGEditor/Private/TabFactory/DesignerTabSummoner.h
new file mode 100644
--- /dev/null
+++ b/Engine/Source/Editor/UMGEditor/Private/TabFactory/DesignerTabSummoner.h
@@ -0,0 +1,17 @@
+// Copyright 1998-2014 Epic Games, Inc. All Rights Reserved.
+
+#pragma once
+
+struct FDesignerTabSummoner : public FWorkflowTabFactory
+{
+public:
+	static const FName TabID;
+	
+public:
+	FDesignerTabSummoner(TSharedPtr<class FWidgetBlueprintEditor> InBlueprintEditor);
+	
+	virtual TSharedRef<SWidget> CreateTabBody(const FWorkflowTabSpawnInfo& Info) const override;
+	
+protected:
+	TWeakPtr<class FWidgetBlueprintEditor> BlueprintEditor;
+};
diff --git a/Engine/Source/Editor/UMGEditor/Private/TabFactory/HierarchyTabSummoner.cpp b/Engine/Source/Editor/UMGEditor/Private/TabFactory/HierarchyTabSummoner.cpp
new file mode 100644
--- /dev/null
+++ b/Engine/Source/Editor/UMGEditor/Private/TabFactory/HierarchyTabSummoner.cpp
@@ -0,0 +1,34 @@
+// Copyright 1998-2014 Epic Games, Inc. All Rights Reserved.
+
+#include "UMGEditorPrivatePCH.h"
+
+#include "HierarchyTabSummoner.h"
+
+#define LOCTEXT_NAMESPACE "UMG"
+
+const FName FHierarchyTabSummoner::TabID(TEXT("SlateHierarchy"));
+
+FHierarchyTabSummoner::FHierarchyTabSummoner(TSharedPtr<class FWidgetBlueprintEditor> InBlueprintEditor)
+		: FWorkflowTabFactory(TabID, InBlueprintEditor)
+		, BlueprintEditor(InBlueprintEditor)
+{
+	TabLabel = LOCTEXT("SlateHierarchyTabLabel", "Hierarchy");
+	//TabIcon = FEditorStyle::GetBrush("Kismet.Tabs.Components");
+
+	bIsSingleton = true;
+
+	ViewMenuDescription = LOCTEXT("SlateHierarchy_ViewMenu_Desc", "Hierarchy");
+	ViewMenuTooltip = LOCTEXT("SlateHierarchy_ViewMenu_ToolTip", "Show the Hierarchy");
+}
+
+TSharedRef<SWidget> FHierarchyTabSummoner::CreateTabBody(const FWorkflowTabSpawnInfo& Info) const
+{
+	TSharedPtr<FWidgetBlueprintEditor> BlueprintEditorPtr = StaticCastSharedPtr<FWidgetBlueprintEditor>(BlueprintEditor.Pin());
+
+	return SNew(STutorialWrapper, TEXT("Hierarchy"))
+		[
+			SNew(SUMGEditorTree, BlueprintEditorPtr, BlueprintEditorPtr->GetBlueprintObj()->SimpleConstructionScript)
+		];
+}
+
+#undef LOCTEXT_NAMESPACE 
diff --git a/Engine/Source/Editor/UMGEditor/Private/TabFactory/HierarchyTabSummoner.h b/Engine/Source/Editor/UMGEditor/Private/TabFactory/HierarchyTabSummoner.h
new file mode 100644
--- /dev/null
+++ b/Engine/Source/Editor/UMGEditor/Private/TabFactory/HierarchyTabSummoner.h
@@ -0,0 +1,17 @@
+// Copyright 1998-2014 Epic Games, Inc. All Rights Reserved.
+
+#pragma once
+
+struct FHierarchyTabSummoner : public FWorkflowTabFactory
+{
+public:
+	static const FName TabID;
+	
+public:
+	FHierarchyTabSummoner(TSharedPtr<class FWidgetBlueprintEditor> InBlueprintEditor);
+	
+	virtual TSharedRef<SWidget> CreateTabBody(const FWorkflowTabSpawnInfo& Info) const override;
+	
+protected:
+	TWeakPtr<class FWidgetBlueprintEditor> BlueprintEditor;
+};
diff --git a/Engine/Source/Editor/UMGEditor/Private/TabFactory/PaletteTabSummoner.cpp b/Engine/Source/Editor/UMGEditor/Private/TabFactory/PaletteTabSummoner.cpp
new file mode 100644
--- /dev/null
+++ b/Engine/Source/Editor/UMGEditor/Private/TabFactory/PaletteTabSummoner.cpp
@@ -0,0 +1,34 @@
+// Copyright 1998-2014 Epic Games, Inc. All Rights Reserved.
+
+#include "UMGEditorPrivatePCH.h"
+
+#include "PaletteTabSummoner.h"
+
+#define LOCTEXT_NAMESPACE "UMG"
+
+const FName FPaletteTabSummoner::TabID(TEXT("WidgetTemplates"));
+
+FPaletteTabSummoner::FPaletteTabSummoner(TSharedPtr<class FWidgetBlueprintEditor> InBlueprintEditor)
+		: FWorkflowTabFactory(TabID, InBlueprintEditor)
+		, BlueprintEditor(InBlueprintEditor)
+{
+	TabLabel = LOCTEXT("WidgetTemplatesTabLabel", "Palette");
+	//TabIcon = FEditorStyle::GetBrush("Kismet.Tabs.Components");
+
+	bIsSingleton = true;
+
+	ViewMenuDescription = LOCTEXT("WidgetTemplates_ViewMenu_Desc", "Palette");
+	ViewMenuTooltip = LOCTEXT("WidgetTemplates_ViewMenu_ToolTip", "Show the Palette");
+}
+
+TSharedRef<SWidget> FPaletteTabSummoner::CreateTabBody(const FWorkflowTabSpawnInfo& Info) const
+{
+	TSharedPtr<FWidgetBlueprintEditor> BlueprintEditorPtr = StaticCastSharedPtr<FWidgetBlueprintEditor>(BlueprintEditor.Pin());
+
+	return SNew(STutorialWrapper, TEXT("Palette"))
+		[
+			SNew(SUMGEditorWidgetTemplates, BlueprintEditorPtr, BlueprintEditorPtr->GetBlueprintObj()->SimpleConstructionScript)
+		];
+}
+
+#undef LOCTEXT_NAMESPACE 
diff --git a/Engine/Source/Editor/UMGEditor/Private/TabFactory/PaletteTabSummoner.h b/Engine/Source/Editor/UMGEditor/Private/TabFactory/PaletteTabSummoner.h
new file mode 100644
--- /dev/null
+++ b/Engine/Source/Editor/UMGEditor/Private/TabFactory/PaletteTabSummoner.h
@@ -0,0 +1,17 @@
+// Copyright 1998-2014 Epic Games, Inc. All Rights Reserved.
+
+#pragma once
+
+struct FPaletteTabSummoner : public FWorkflowTabFactory
+{
+public:
+	static const FName TabID;
+	
+public:
+	FPaletteTabSummoner(TSharedPtr<class FWidgetBlueprintEditor> InBlueprintEditor);
+	
+	virtual TSharedRef<SWidget> CreateTabBody(const FWorkflowTabSpawnInfo& Info) const override;
+	
+protected:
+	TWeakPtr<class FWidgetBlueprintEditor> BlueprintEditor;
+};
diff --git a/Engine/Source/Editor/UMGEditor/Private/TabFactory/SequencerTabSummoner.cpp b/Engine/Source/Editor/UMGEditor/Private/TabFactory/SequencerTabSummoner.cpp
new file mode 100644
--- /dev/null
+++ b/Engine/Source/Editor/UMGEditor/Private/TabFactory/SequencerTabSummoner.cpp
@@ -0,0 +1,34 @@
+// Copyright 1998-2014 Epic Games, Inc. All Rights Reserved.
+
+#include "UMGEditorPrivatePCH.h"
+
+#include "SequencerTabSummoner.h"
+
+#define LOCTEXT_NAMESPACE "UMG"
+
+const FName FSequencerTabSummoner::TabID(TEXT("Sequencer"));
+
+FSequencerTabSummoner::FSequencerTabSummoner(TSharedPtr<class FWidgetBlueprintEditor> InBlueprintEditor)
+		: FWorkflowTabFactory(TabID, InBlueprintEditor)
+		, BlueprintEditor(InBlueprintEditor)
+{
+	TabLabel = LOCTEXT("SequencerLabel", "Timeline");
+	//TabIcon = FEditorStyle::GetBrush("Kismet.Tabs.Components");
+
+	bIsSingleton = true;
+
+	ViewMenuDescription = LOCTEXT("Sequencer_ViewMenu_Desc", "Timeline");
+	ViewMenuTooltip = LOCTEXT("Sequencer_ViewMenu_ToolTip", "Show the Animation editor");
+}
+
+TSharedRef<SWidget> FSequencerTabSummoner::CreateTabBody(const FWorkflowTabSpawnInfo& Info) const
+{
+	TSharedPtr<FWidgetBlueprintEditor> BlueprintEditorPinned = BlueprintEditor.Pin();
+
+	return SNew(STutorialWrapper, TEXT("Sequencer"))
+		[
+			BlueprintEditorPinned->GetSequencer()->GetSequencerWidget()
+		];
+}
+
+#undef LOCTEXT_NAMESPACE 
diff --git a/Engine/Source/Editor/UMGEditor/Private/TabFactory/SequencerTabSummoner.h b/Engine/Source/Editor/UMGEditor/Private/TabFactory/SequencerTabSummoner.h
new file mode 100644
--- /dev/null
+++ b/Engine/Source/Editor/UMGEditor/Private/TabFactory/SequencerTabSummoner.h
@@ -0,0 +1,17 @@
+// Copyright 1998-2014 Epic Games, Inc. All Rights Reserved.
+
+#pragma once
+
+struct FSequencerTabSummoner : public FWorkflowTabFactory
+{
+public:
+	static const FName TabID;
+	
+public:
+	FSequencerTabSummoner(TSharedPtr<class FWidgetBlueprintEditor> InBlueprintEditor);
+	
+	virtual TSharedRef<SWidget> CreateTabBody(const FWorkflowTabSpawnInfo& Info) const override;
+	
+protected:
+	TWeakPtr<class FWidgetBlueprintEditor> BlueprintEditor;
+};
diff --git a/Engine/Source/Editor/UMGEditor/Private/Templates/WidgetTemplateClass.cpp b/Engine/Source/Editor/UMGEditor/Private/Templates/WidgetTemplateClass.cpp
--- a/Engine/Source/Editor/UMGEditor/Private/Templates/WidgetTemplateClass.cpp
+++ b/Engine/Source/Editor/UMGEditor/Private/Templates/WidgetTemplateClass.cpp
@@ -28,4 +28,10 @@ UWidget* FWidgetTemplateClass::Create(UWidgetTree* Tree)
 	return Tree->ConstructWidget<UWidget>(WidgetClass);
 }
 
+const FSlateBrush* FWidgetTemplateClass::GetIcon() const
+{
+	UWidget* DefaultWidget = WidgetClass->GetDefaultObject<UWidget>();
+	return DefaultWidget->GetEditorIcon();
+}
+
 #undef LOCTEXT_NAMESPACE
diff --git a/Engine/Source/Editor/UMGEditor/Private/Templates/WidgetTemplateClass.h b/Engine/Source/Editor/UMGEditor/Private/Templates/WidgetTemplateClass.h
--- a/Engine/Source/Editor/UMGEditor/Private/Templates/WidgetTemplateClass.h
+++ b/Engine/Source/Editor/UMGEditor/Private/Templates/WidgetTemplateClass.h
@@ -13,6 +13,8 @@ public:
 
 	virtual UWidget* Create(UWidgetTree* Tree) override;
 
+	virtual const FSlateBrush* GetIcon() const override;
+
 protected:
 	TSubclassOf<UWidget> WidgetClass;
 };
diff --git a/Engine/Source/Editor/UMGEditor/Private/UMGBlueprintEditorExtensionHook.cpp b/Engine/Source/Editor/UMGEditor/Private/UMGBlueprintEditorExtensionHook.cpp
deleted file mode 100644
diff --git a/Engine/Source/Editor/UMGEditor/Private/UMGBlueprintEditorExtensionHook.h b/Engine/Source/Editor/UMGEditor/Private/UMGBlueprintEditorExtensionHook.h
deleted file mode 100644
diff --git a/Engine/Source/Editor/UMGEditor/Private/UMGEditorModule.cpp b/Engine/Source/Editor/UMGEditor/Private/UMGEditorModule.cpp
--- a/Engine/Source/Editor/UMGEditor/Private/UMGEditorModule.cpp
+++ b/Engine/Source/Editor/UMGEditor/Private/UMGEditorModule.cpp
@@ -6,7 +6,6 @@
 #include "ModuleManager.h"
 
 #include "AssetToolsModule.h"
-#include "UMGBlueprintEditorExtensionHook.h"
 #include "AssetTypeActions_WidgetBlueprint.h"
 #include "KismetCompilerModule.h"
 #include "WidgetBlueprintCompiler.h"
@@ -48,8 +47,6 @@ public:
 			IAssetTools& AssetTools = FModuleManager::LoadModuleChecked<FAssetToolsModule>("AssetTools").Get();
 			RegisterAssetTypeAction(AssetTools, MakeShareable(new FAssetTypeActions_WidgetBlueprint()));
 
-			FUMGBlueprintEditorExtensionHook::InstallHooks();
-
 			// Register with the sequencer module that we provide auto-key handlers.
 			ISequencerModule& SequencerModule = FModuleManager::Get().LoadModuleChecked<ISequencerModule>("Sequencer");
 			SequencerModule.RegisterTrackEditor(FOnCreateTrackEditor::CreateStatic(&FMarginTrackEditor::CreateTrackEditor));
@@ -64,11 +61,6 @@ public:
 
 		if ( bUMGEnabled )
 		{
-			if ( UObjectInitialized() )
-			{
-				FUMGBlueprintEditorExtensionHook::RemoveHooks();
-			}
-
 			// Unregister all the asset types that we registered
 			if ( FModuleManager::Get().IsModuleLoaded("AssetTools") )
 			{
diff --git a/Engine/Source/Editor/UMGEditor/Private/UMGEditorPrivatePCH.h b/Engine/Source/Editor/UMGEditor/Private/UMGEditorPrivatePCH.h
--- a/Engine/Source/Editor/UMGEditor/Private/UMGEditorPrivatePCH.h
+++ b/Engine/Source/Editor/UMGEditor/Private/UMGEditorPrivatePCH.h
@@ -5,6 +5,9 @@
 #include "Slate.h"
 #include "SlateCore.h"
 #include "UnrealEd.h"
+
+#include "IDocumentation.h"
+
 #include "BlueprintGraphDefinitions.h"
 
 #include "BlueprintEditorUtils.h"
@@ -17,6 +20,11 @@
 #include "UMGEditorClasses.h"
 #include "WidgetBlueprint.h"
 
+#include "BlueprintEditorModes.h"
+#include "WidgetBlueprintApplicationModes.h"
+#include "WidgetBlueprintApplicationMode.h"
+
+
 #include "AssetTypeActions_Base.h"
 
 #include "DesignerExtension.h"
diff --git a/Engine/Source/Editor/UMGEditor/Private/WidgetBlueprintEditor.cpp b/Engine/Source/Editor/UMGEditor/Private/WidgetBlueprintEditor.cpp
--- a/Engine/Source/Editor/UMGEditor/Private/WidgetBlueprintEditor.cpp
+++ b/Engine/Source/Editor/UMGEditor/Private/WidgetBlueprintEditor.cpp
@@ -8,6 +8,13 @@
 #include "Editor/Sequencer/Public/ISequencerModule.h"
 #include "UMGSequencerObjectBindingManager.h"
 
+#include "WidgetBlueprintApplicationModes.h"
+//#include "WidgetDefafaultsApplicationMode.h"
+#include "WidgetDesignerApplicationMode.h"
+#include "WidgetGraphApplicationMode.h"
+
+#include "WidgetBlueprintEditorToolbar.h"
+
 #define LOCTEXT_NAMESPACE "UMG"
 
 FWidgetBlueprintEditor::FWidgetBlueprintEditor()
@@ -26,6 +33,78 @@ FWidgetBlueprintEditor::~FWidgetBlueprintEditor()
 	}
 }
 
+void FWidgetBlueprintEditor::InitWidgetBlueprintEditor(const EToolkitMode::Type Mode, const TSharedPtr< IToolkitHost >& InitToolkitHost, const TArray<UBlueprint*>& InBlueprints, bool bShouldOpenInDefaultsMode)
+{
+	TSharedPtr<FWidgetBlueprintEditor> ThisPtr(SharedThis(this));
+	WidgetToolbar = MakeShareable(new FWidgetBlueprintEditorToolbar(ThisPtr));
+
+	InitBlueprintEditor(Mode, InitToolkitHost, InBlueprints, bShouldOpenInDefaultsMode);
+
+	UWidgetBlueprint* Blueprint = GetWidgetBlueprintObj();
+
+	// If this blueprint is empty, add a canvas panel as the root widget.
+	if ( Blueprint->WidgetTree->RootWidget == NULL )
+	{
+		UWidget* RootWidget = Blueprint->WidgetTree->ConstructWidget<UCanvasPanel>(UCanvasPanel::StaticClass());
+		RootWidget->IsDesignTime(true);
+		Blueprint->WidgetTree->RootWidget = RootWidget;
+	}
+
+	UpdatePreview(GetWidgetBlueprintObj(), true);
+
+	WidgetCommandList = MakeShareable(new FUICommandList);
+
+	WidgetCommandList->MapAction(FGenericCommands::Get().Delete,
+		FExecuteAction::CreateSP(this, &FWidgetBlueprintEditor::DeleteSelectedWidgets),
+		FCanExecuteAction::CreateSP(this, &FWidgetBlueprintEditor::CanDeleteSelectedWidgets)
+		);
+
+	WidgetCommandList->MapAction(FGenericCommands::Get().Copy,
+		FExecuteAction::CreateSP(this, &FWidgetBlueprintEditor::CopySelectedWidgets),
+		FCanExecuteAction::CreateSP(this, &FWidgetBlueprintEditor::CanCopySelectedWidgets)
+		);
+
+	//WidgetCommandList->MapAction(FGenericCommands::Get().Cut,
+	//	FExecuteAction::CreateSP(this, &FWidgetBlueprintEditor::CutSelectedNodes),
+	//	FCanExecuteAction::CreateSP(this, &FWidgetBlueprintEditor::CanCutNodes)
+	//	);
+
+	WidgetCommandList->MapAction(FGenericCommands::Get().Paste,
+		FExecuteAction::CreateSP(this, &FWidgetBlueprintEditor::PasteWidgets),
+		FCanExecuteAction::CreateSP(this, &FWidgetBlueprintEditor::CanPasteWidgets)
+		);
+}
+
+void FWidgetBlueprintEditor::RegisterApplicationModes(const TArray<UBlueprint*>& InBlueprints, bool bShouldOpenInDefaultsMode)
+{
+	//FBlueprintEditor::RegisterApplicationModes(InBlueprints, bShouldOpenInDefaultsMode);
+
+	if ( InBlueprints.Num() == 1 )
+	{
+		TSharedPtr<FWidgetBlueprintEditor> ThisPtr(SharedThis(this));
+
+		// Create the modes and activate one (which will populate with a real layout)
+		TArray< TSharedRef<FApplicationMode> > TempModeList;
+		TempModeList.Add(MakeShareable(new FWidgetDesignerApplicationMode(ThisPtr)));
+		TempModeList.Add(MakeShareable(new FWidgetGraphApplicationMode(ThisPtr)));
+
+		for ( TSharedRef<FApplicationMode>& AppMode : TempModeList )
+		{
+			AddApplicationMode(AppMode->GetModeName(), AppMode);
+		}
+
+		SetCurrentMode(FWidgetBlueprintApplicationModes::DesignerMode);
+	}
+	else
+	{
+		//// We either have no blueprints or many, open in the defaults mode for multi-editing
+		//AddApplicationMode(
+		//	FBlueprintEditorApplicationModes::BlueprintDefaultsMode,
+		//	MakeShareable(new FBlueprintDefaultsApplicationMode(SharedThis(this))));
+		//SetCurrentMode(FBlueprintEditorApplicationModes::BlueprintDefaultsMode);
+	}
+}
+
 void FWidgetBlueprintEditor::SelectWidgets(const TSet<FWidgetReference>& Widgets)
 {
 	TSet<FWidgetReference> TempSelection;
@@ -117,45 +196,6 @@ void FWidgetBlueprintEditor::OnBlueprintChanged(UBlueprint* InBlueprint)
 	}
 }
 
-void FWidgetBlueprintEditor::InitWidgetBlueprintEditor(const EToolkitMode::Type Mode, const TSharedPtr< IToolkitHost >& InitToolkitHost, const TArray<UBlueprint*>& InBlueprints, bool bShouldOpenInDefaultsMode)
-{
-	InitBlueprintEditor(Mode, InitToolkitHost, InBlueprints, bShouldOpenInDefaultsMode);
-
-	UWidgetBlueprint* Blueprint = GetWidgetBlueprintObj();
-
-	// If this blueprint is empty, add a canvas panel as the root widget.
-	if ( Blueprint->WidgetTree->RootWidget == NULL )
-	{
-		 UWidget* RootWidget = Blueprint->WidgetTree->ConstructWidget<UCanvasPanel>(UCanvasPanel::StaticClass());
-		 RootWidget->IsDesignTime(true);
-		 Blueprint->WidgetTree->RootWidget = RootWidget;
-	}
-
-	UpdatePreview(GetWidgetBlueprintObj(), true);
-
-	WidgetCommandList = MakeShareable(new FUICommandList);
-
-	WidgetCommandList->MapAction(FGenericCommands::Get().Delete,
-		FExecuteAction::CreateSP(this, &FWidgetBlueprintEditor::DeleteSelectedWidgets),
-		FCanExecuteAction::CreateSP(this, &FWidgetBlueprintEditor::CanDeleteSelectedWidgets)
-		);
-
-	WidgetCommandList->MapAction(FGenericCommands::Get().Copy,
-		FExecuteAction::CreateSP(this, &FWidgetBlueprintEditor::CopySelectedWidgets),
-		FCanExecuteAction::CreateSP(this, &FWidgetBlueprintEditor::CanCopySelectedWidgets)
-		);
-
-	//WidgetCommandList->MapAction(FGenericCommands::Get().Cut,
-	//	FExecuteAction::CreateSP(this, &FWidgetBlueprintEditor::CutSelectedNodes),
-	//	FCanExecuteAction::CreateSP(this, &FWidgetBlueprintEditor::CanCutNodes)
-	//	);
-
-	WidgetCommandList->MapAction(FGenericCommands::Get().Paste,
-		FExecuteAction::CreateSP(this, &FWidgetBlueprintEditor::PasteWidgets),
-		FCanExecuteAction::CreateSP(this, &FWidgetBlueprintEditor::CanPasteWidgets)
-		);
-}
-
 bool FWidgetBlueprintEditor::CanDeleteSelectedWidgets()
 {
 	TSet<FWidgetReference> Widgets = GetSelectedWidgets();
diff --git a/Engine/Source/Editor/UMGEditor/Private/WidgetBlueprintEditor.h b/Engine/Source/Editor/UMGEditor/Private/WidgetBlueprintEditor.h
--- a/Engine/Source/Editor/UMGEditor/Private/WidgetBlueprintEditor.h
+++ b/Engine/Source/Editor/UMGEditor/Private/WidgetBlueprintEditor.h
@@ -53,6 +53,8 @@ public:
 	/** @return Notification for when the preview widget has been updated */
 	FOnWidgetPreviewUpdated& GetOnWidgetPreviewUpdated() { return OnWidgetPreviewUpdated; }
 
+	TSharedPtr<class FWidgetBlueprintEditorToolbar> GetWidgetToolbarBuilder() { return WidgetToolbar; }
+
 public:
 	/** Fires whenever the selected set of widgets changes */
 	FOnSelectedWidgetsChanged OnSelectedWidgetsChanged;
@@ -64,7 +66,10 @@ public:
 	FVector2D PasteDropLocation;
 
 protected:
+	// Begin FBlueprintEditor
+	virtual void RegisterApplicationModes(const TArray<UBlueprint*>& InBlueprints, bool bShouldOpenInDefaultsMode) override;
 	virtual FGraphAppearanceInfo GetGraphAppearance() const override;
+	// End FBlueprintEditor
 
 private:
 	bool CanDeleteSelectedWidgets();
@@ -120,6 +125,9 @@ private:
 	/** The preview GUI object */
 	mutable TWeakObjectPtr<UUserWidget> PreviewWidgetActorPtr;
 
-	/** Notification for when the preview widget has been updated  */
+	/** Notification for when the preview widget has been updated */
 	FOnWidgetPreviewUpdated OnWidgetPreviewUpdated;
+
+	/** The toolbar builder associated with this editor */
+	TSharedPtr<class FWidgetBlueprintEditorToolbar> WidgetToolbar;
 };
diff --git a/Engine/Source/Editor/UMGEditor/Private/WidgetBlueprintEditorToolbar.cpp b/Engine/Source/Editor/UMGEditor/Private/WidgetBlueprintEditorToolbar.cpp
new file mode 100644
--- /dev/null
+++ b/Engine/Source/Editor/UMGEditor/Private/WidgetBlueprintEditorToolbar.cpp
@@ -0,0 +1,116 @@
+// Copyright 1998-2014 Epic Games, Inc. All Rights Reserved.
+
+#include "UMGEditorPrivatePCH.h"
+
+#include "WidgetBlueprintEditorToolbar.h"
+
+#include "SModeWidget.h"
+
+#define LOCTEXT_NAMESPACE "UMG"
+
+//////////////////////////////////////////////////////////////////////////
+// SBlueprintModeSeparator
+
+class SBlueprintModeSeparator : public SBorder
+{
+public:
+	SLATE_BEGIN_ARGS(SBlueprintModeSeparator) {}
+	SLATE_END_ARGS()
+
+		void Construct(const FArguments& InArg)
+	{
+		SBorder::Construct(
+			SBorder::FArguments()
+			.BorderImage(FEditorStyle::GetBrush("BlueprintEditor.PipelineSeparator"))
+			.Padding(0.0f)
+			);
+	}
+
+	// SWidget interface
+	virtual FVector2D ComputeDesiredSize() const override
+	{
+		const float Height = 20.0f;
+		const float Thickness = 16.0f;
+		return FVector2D(Thickness, Height);
+	}
+	// End of SWidget interface
+};
+
+//////////////////////////////////////////////////////////////////////////
+// FWidgetBlueprintEditorToolbar
+
+FWidgetBlueprintEditorToolbar::FWidgetBlueprintEditorToolbar(TSharedPtr<FWidgetBlueprintEditor>& InWidgetEditor)
+	: WidgetEditor(InWidgetEditor)
+{
+}
+
+void FWidgetBlueprintEditorToolbar::AddWidgetBlueprintEditorModesToolbar(TSharedPtr<FExtender> Extender)
+{
+	TSharedPtr<FWidgetBlueprintEditor> BlueprintEditorPtr = WidgetEditor.Pin();
+
+	Extender->AddToolBarExtension(
+		"Asset",
+		EExtensionHook::After,
+		BlueprintEditorPtr->GetToolkitCommands(),
+		FToolBarExtensionDelegate::CreateSP(this, &FWidgetBlueprintEditorToolbar::FillWidgetBlueprintEditorModesToolbar));
+}
+
+void FWidgetBlueprintEditorToolbar::FillWidgetBlueprintEditorModesToolbar(FToolBarBuilder& ToolbarBuilder)
+{
+	TSharedPtr<FWidgetBlueprintEditor> BlueprintEditorPtr = WidgetEditor.Pin();
+	UBlueprint* BlueprintObj = BlueprintEditorPtr->GetBlueprintObj();
+
+	if( !BlueprintObj ||
+		(!FBlueprintEditorUtils::IsLevelScriptBlueprint(BlueprintObj) 
+		&& !FBlueprintEditorUtils::IsInterfaceBlueprint(BlueprintObj)
+		&& !BlueprintObj->bIsNewlyCreated)
+		)
+	{
+		TAttribute<FName> GetActiveMode(BlueprintEditorPtr.ToSharedRef(), &FBlueprintEditor::GetCurrentMode);
+		FOnModeChangeRequested SetActiveMode = FOnModeChangeRequested::CreateSP(BlueprintEditorPtr.ToSharedRef(), &FBlueprintEditor::SetCurrentMode);
+
+		// Left side padding
+		BlueprintEditorPtr->AddToolbarWidget(SNew(SSpacer).Size(FVector2D(4.0f, 1.0f)));
+
+		BlueprintEditorPtr->AddToolbarWidget(
+			SNew( STutorialWrapper, TEXT("DefaultsMode") )
+			[
+				SNew(SModeWidget, FWidgetBlueprintApplicationModes::GetLocalizedMode(FWidgetBlueprintApplicationModes::DesignerMode), FWidgetBlueprintApplicationModes::DesignerMode)
+				.OnGetActiveMode(GetActiveMode)
+				.OnSetActiveMode(SetActiveMode)
+				.ToolTip(IDocumentation::Get()->CreateToolTip(
+					LOCTEXT("BlueprintDefaultsModeButtonTooltip", "Switch to Blueprint Defaults Mode"),
+					NULL,
+					TEXT("Shared/Editors/BlueprintEditor"),
+					TEXT("DefaultsMode")))
+				.IconImage(FEditorStyle::GetBrush("FullBlueprintEditor.SwitchToBlueprintDefaultsMode"))
+				.SmallIconImage(FEditorStyle::GetBrush("FullBlueprintEditor.SwitchToBlueprintDefaultsMode.Small"))
+			]
+		);
+
+		BlueprintEditorPtr->AddToolbarWidget(SNew(SBlueprintModeSeparator));
+
+		BlueprintEditorPtr->AddToolbarWidget(
+			SNew( STutorialWrapper, TEXT("GraphMode") )
+			[
+				SNew(SModeWidget, FWidgetBlueprintApplicationModes::GetLocalizedMode(FWidgetBlueprintApplicationModes::GraphMode), FWidgetBlueprintApplicationModes::GraphMode)
+				.OnGetActiveMode(GetActiveMode)
+				.OnSetActiveMode(SetActiveMode)
+				.CanBeSelected(BlueprintEditorPtr.Get(), &FBlueprintEditor::IsEditingSingleBlueprint)
+				.ToolTip(IDocumentation::Get()->CreateToolTip(
+					LOCTEXT("GraphModeButtonTooltip", "Switch to Graph Editing Mode"),
+					NULL,
+					TEXT("Shared/Editors/BlueprintEditor"),
+					TEXT("GraphMode")))
+				.ToolTipText(LOCTEXT("GraphModeButtonTooltip", "Switch to Graph Editing Mode"))
+				.IconImage(FEditorStyle::GetBrush("FullBlueprintEditor.SwitchToScriptingMode"))
+				.SmallIconImage(FEditorStyle::GetBrush("FullBlueprintEditor.SwitchToScriptingMode.Small"))
+			]
+		);
+		
+		// Right side padding
+		BlueprintEditorPtr->AddToolbarWidget(SNew(SSpacer).Size(FVector2D(4.0f, 1.0f)));
+	}
+}
+
+#undef LOCTEXT_NAMESPACE
diff --git a/Engine/Source/Editor/UMGEditor/Private/WidgetBlueprintEditorToolbar.h b/Engine/Source/Editor/UMGEditor/Private/WidgetBlueprintEditorToolbar.h
new file mode 100644
--- /dev/null
+++ b/Engine/Source/Editor/UMGEditor/Private/WidgetBlueprintEditorToolbar.h
@@ -0,0 +1,25 @@
+// Copyright 1998-2014 Epic Games, Inc. All Rights Reserved.
+
+#pragma once
+
+/**
+ * Handles all of the toolbar related construction for the widget blueprint editor.
+ */
+class FWidgetBlueprintEditorToolbar : public TSharedFromThis<FWidgetBlueprintEditorToolbar>
+{
+
+public:
+	/** Constructor */
+	FWidgetBlueprintEditorToolbar(TSharedPtr<FWidgetBlueprintEditor>& InWidgetEditor);
+	
+	/**
+	 * Builds the modes toolbar for the widget blueprint editor.
+	 */
+	void AddWidgetBlueprintEditorModesToolbar(TSharedPtr<FExtender> Extender);
+
+public:
+	/**  */
+	void FillWidgetBlueprintEditorModesToolbar(FToolBarBuilder& ToolbarBuilder);
+
+	TWeakPtr<FWidgetBlueprintEditor> WidgetEditor;
+};
diff --git a/Engine/Source/Editor/UMGEditor/Public/WidgetTemplate.h b/Engine/Source/Editor/UMGEditor/Public/WidgetTemplate.h
--- a/Engine/Source/Editor/UMGEditor/Public/WidgetTemplate.h
+++ b/Engine/Source/Editor/UMGEditor/Public/WidgetTemplate.h
@@ -18,13 +18,17 @@ public:
 	/** Constructs the widget template. */
 	virtual UWidget* Create(class UWidgetTree* Tree) = 0;
 
+	/** Gets the icon to display in the template palate for this template. */
+	virtual const FSlateBrush* GetIcon() const
+	{
+		static FSlateNoResource NullBrush;
+		return &NullBrush;
+	}
+
 public:
 	/** The name of the widget template. */
 	FText Name;
 
 	/** A tooltip to display for the template. */
 	FText ToolTip;
-
-	/** An icon to display for the template. */
-	FSlateIcon Icon;
 };
diff --git a/Engine/Source/Editor/UMGEditor/UMGEditor.Build.cs b/Engine/Source/Editor/UMGEditor/UMGEditor.Build.cs
--- a/Engine/Source/Editor/UMGEditor/UMGEditor.Build.cs
+++ b/Engine/Source/Editor/UMGEditor/UMGEditor.Build.cs
@@ -10,8 +10,10 @@ public class UMGEditor : ModuleRules
 			new string[] {
 				"Editor/UMGEditor/Private", // For PCH includes (because they don't work with relative paths, yet)
 				"Editor/UMGEditor/Private/Templates",
-                "Editor/UMGEditor/Private/Extensions",
-                "Editor/UMGEditor/Private/Customizations",
+				"Editor/UMGEditor/Private/Extensions",
+				"Editor/UMGEditor/Private/Customizations",
+				"Editor/UMGEditor/Private/BlueprintModes",
+				"Editor/UMGEditor/Private/TabFactory",
 			});
 
 		PrivateIncludePathModuleNames.AddRange(
