UMG - Reorganizing and renaming a few classes in the editor.

diff --git a/Engine/Source/Editor/UMGEditor/Private/Designer/SDesignSurface.cpp b/Engine/Source/Editor/UMGEditor/Private/Designer/SDesignSurface.cpp
new file mode 100644
--- /dev/null
+++ b/Engine/Source/Editor/UMGEditor/Private/Designer/SDesignSurface.cpp
@@ -0,0 +1,371 @@
+// Copyright 1998-2014 Epic Games, Inc. All Rights Reserved.
+
+#include "UMGEditorPrivatePCH.h"
+
+#include "SDesignSurface.h"
+
+#define LOCTEXT_NAMESPACE "UMG"
+
+struct FZoomLevelEntry
+{
+public:
+	FZoomLevelEntry(float InZoomAmount, const FText& InDisplayText, EGraphRenderingLOD::Type InLOD)
+		: DisplayText(FText::Format(LOCTEXT("Zoom", "Zoom {0}"), InDisplayText))
+		, ZoomAmount(InZoomAmount)
+		, LOD(InLOD)
+	{
+	}
+
+public:
+	FText DisplayText;
+	float ZoomAmount;
+	EGraphRenderingLOD::Type LOD;
+};
+
+struct FFixedZoomLevelsContainer : public FZoomLevelsContainer
+{
+	FFixedZoomLevelsContainer()
+	{
+		// Initialize zoom levels if not done already
+		if ( ZoomLevels.Num() == 0 )
+		{
+			ZoomLevels.Reserve(20);
+			ZoomLevels.Add(FZoomLevelEntry(0.100f, LOCTEXT("ZoomLevel", "-12"), EGraphRenderingLOD::LowestDetail));
+			ZoomLevels.Add(FZoomLevelEntry(0.125f, LOCTEXT("ZoomLevel", "-11"), EGraphRenderingLOD::LowestDetail));
+			ZoomLevels.Add(FZoomLevelEntry(0.150f, LOCTEXT("ZoomLevel", "-10"), EGraphRenderingLOD::LowestDetail));
+			ZoomLevels.Add(FZoomLevelEntry(0.175f, LOCTEXT("ZoomLevel", "-9"), EGraphRenderingLOD::LowestDetail));
+			ZoomLevels.Add(FZoomLevelEntry(0.200f, LOCTEXT("ZoomLevel", "-8"), EGraphRenderingLOD::LowestDetail));
+			ZoomLevels.Add(FZoomLevelEntry(0.225f, LOCTEXT("ZoomLevel", "-7"), EGraphRenderingLOD::LowDetail));
+			ZoomLevels.Add(FZoomLevelEntry(0.250f, LOCTEXT("ZoomLevel", "-6"), EGraphRenderingLOD::LowDetail));
+			ZoomLevels.Add(FZoomLevelEntry(0.375f, LOCTEXT("ZoomLevel", "-5"), EGraphRenderingLOD::MediumDetail));
+			ZoomLevels.Add(FZoomLevelEntry(0.500f, LOCTEXT("ZoomLevel", "-4"), EGraphRenderingLOD::MediumDetail));
+			ZoomLevels.Add(FZoomLevelEntry(0.675f, LOCTEXT("ZoomLevel", "-3"), EGraphRenderingLOD::MediumDetail));
+			ZoomLevels.Add(FZoomLevelEntry(0.750f, LOCTEXT("ZoomLevel", "-2"), EGraphRenderingLOD::DefaultDetail));
+			ZoomLevels.Add(FZoomLevelEntry(0.875f, LOCTEXT("ZoomLevel", "-1"), EGraphRenderingLOD::DefaultDetail));
+			ZoomLevels.Add(FZoomLevelEntry(1.000f, LOCTEXT("ZoomLevel", "1:1"), EGraphRenderingLOD::DefaultDetail));
+			ZoomLevels.Add(FZoomLevelEntry(1.250f, LOCTEXT("ZoomLevel", "+1"), EGraphRenderingLOD::DefaultDetail));
+			ZoomLevels.Add(FZoomLevelEntry(1.375f, LOCTEXT("ZoomLevel", "+2"), EGraphRenderingLOD::DefaultDetail));
+			ZoomLevels.Add(FZoomLevelEntry(1.500f, LOCTEXT("ZoomLevel", "+3"), EGraphRenderingLOD::FullyZoomedIn));
+			ZoomLevels.Add(FZoomLevelEntry(1.675f, LOCTEXT("ZoomLevel", "+4"), EGraphRenderingLOD::FullyZoomedIn));
+			ZoomLevels.Add(FZoomLevelEntry(1.750f, LOCTEXT("ZoomLevel", "+5"), EGraphRenderingLOD::FullyZoomedIn));
+			ZoomLevels.Add(FZoomLevelEntry(1.875f, LOCTEXT("ZoomLevel", "+6"), EGraphRenderingLOD::FullyZoomedIn));
+			ZoomLevels.Add(FZoomLevelEntry(2.000f, LOCTEXT("ZoomLevel", "+7"), EGraphRenderingLOD::FullyZoomedIn));
+		}
+	}
+
+	float GetZoomAmount(int32 InZoomLevel) const override
+	{
+		checkSlow(ZoomLevels.IsValidIndex(InZoomLevel));
+		return ZoomLevels[InZoomLevel].ZoomAmount;
+	}
+
+	int32 GetNearestZoomLevel(float InZoomAmount) const override
+	{
+		for ( int32 ZoomLevelIndex=0; ZoomLevelIndex < GetNumZoomLevels(); ++ZoomLevelIndex )
+		{
+			if ( InZoomAmount <= GetZoomAmount(ZoomLevelIndex) )
+			{
+				return ZoomLevelIndex;
+			}
+		}
+
+		return GetDefaultZoomLevel();
+	}
+
+	FText GetZoomText(int32 InZoomLevel) const override
+	{
+		checkSlow(ZoomLevels.IsValidIndex(InZoomLevel));
+		return ZoomLevels[InZoomLevel].DisplayText;
+	}
+
+	int32 GetNumZoomLevels() const override
+	{
+		return ZoomLevels.Num();
+	}
+
+	int32 GetDefaultZoomLevel() const override
+	{
+		return 12;
+	}
+
+	EGraphRenderingLOD::Type GetLOD(int32 InZoomLevel) const override
+	{
+		checkSlow(ZoomLevels.IsValidIndex(InZoomLevel));
+		return ZoomLevels[InZoomLevel].LOD;
+	}
+
+	static TArray<FZoomLevelEntry> ZoomLevels;
+};
+
+TArray<FZoomLevelEntry> FFixedZoomLevelsContainer::ZoomLevels;
+
+/////////////////////////////////////////////////////
+// SDesignSurface
+
+void SDesignSurface::Construct(const FArguments& InArgs)
+{
+	if ( !ZoomLevels.IsValid() )
+	{
+		ZoomLevels = new FFixedZoomLevelsContainer();
+	}
+	ZoomLevel = ZoomLevels->GetDefaultZoomLevel();
+	PreviousZoomLevel = ZoomLevels->GetDefaultZoomLevel();
+	PostChangedZoom();
+	SnapGridSize = 16.0f;
+	bAllowContinousZoomInterpolation = false;
+
+	ViewOffset = FVector2D::ZeroVector;
+
+	ZoomLevelFade = FCurveSequence(0.0f, 1.0f);
+	ZoomLevelFade.Play();
+
+	ZoomLevelGraphFade = FCurveSequence(0.0f, 0.5f);
+	ZoomLevelGraphFade.Play();
+
+	ChildSlot
+	[
+		InArgs._Content.Widget
+	];
+}
+
+int32 SDesignSurface::OnPaint(const FGeometry& AllottedGeometry, const FSlateRect& MyClippingRect, FSlateWindowElementList& OutDrawElements, int32 LayerId, const FWidgetStyle& InWidgetStyle, bool bParentEnabled) const
+{
+	const FSlateBrush* BackgroundImage = FEditorStyle::GetBrush(TEXT("Graph.Panel.SolidBackground"));
+	PaintBackgroundAsLines(BackgroundImage, AllottedGeometry, MyClippingRect, OutDrawElements, LayerId);
+
+	SCompoundWidget::OnPaint(AllottedGeometry, MyClippingRect, OutDrawElements, LayerId, InWidgetStyle, bParentEnabled);
+
+	return LayerId;
+}
+
+FReply SDesignSurface::OnMouseWheel(const FGeometry& MyGeometry, const FPointerEvent& MouseEvent)
+{
+	// We want to zoom into this point; i.e. keep it the same fraction offset into the panel
+	const FVector2D WidgetSpaceCursorPos = MyGeometry.AbsoluteToLocal(MouseEvent.GetScreenSpacePosition());
+	const int32 ZoomLevelDelta = FMath::FloorToInt(MouseEvent.GetWheelDelta());
+	ChangeZoomLevel(ZoomLevelDelta, WidgetSpaceCursorPos, MouseEvent.IsControlDown());
+
+	return FReply::Handled();
+}
+
+inline float FancyMod(float Value, float Size)
+{
+	return ( ( Value >= 0 ) ? 0.0f : Size ) + FMath::Fmod(Value, Size);
+}
+
+float SDesignSurface::GetZoomAmount() const
+{
+	if ( bAllowContinousZoomInterpolation )
+	{
+		return FMath::Lerp(ZoomLevels->GetZoomAmount(PreviousZoomLevel), ZoomLevels->GetZoomAmount(ZoomLevel), ZoomLevelGraphFade.GetLerp());
+	}
+	else
+	{
+		return ZoomLevels->GetZoomAmount(ZoomLevel);
+	}
+}
+
+void SDesignSurface::ChangeZoomLevel(int32 ZoomLevelDelta, const FVector2D& WidgetSpaceZoomOrigin, bool bOverrideZoomLimiting)
+{
+	// We want to zoom into this point; i.e. keep it the same fraction offset into the panel
+	const FVector2D PointToMaintainGraphSpace = PanelCoordToGraphCoord(WidgetSpaceZoomOrigin);
+
+	const int32 DefaultZoomLevel = ZoomLevels->GetDefaultZoomLevel();
+	const int32 NumZoomLevels = ZoomLevels->GetNumZoomLevels();
+
+	const bool bAllowFullZoomRange =
+		// To zoom in past 1:1 the user must press control
+		( ZoomLevel == DefaultZoomLevel && ZoomLevelDelta > 0 && bOverrideZoomLimiting ) ||
+		// If they are already zoomed in past 1:1, user may zoom freely
+		( ZoomLevel > DefaultZoomLevel );
+
+	const float OldZoomLevel = ZoomLevel;
+
+	if ( bAllowFullZoomRange )
+	{
+		ZoomLevel = FMath::Clamp(ZoomLevel + ZoomLevelDelta, 0, NumZoomLevels - 1);
+	}
+	else
+	{
+		// Without control, we do not allow zooming in past 1:1.
+		ZoomLevel = FMath::Clamp(ZoomLevel + ZoomLevelDelta, 0, DefaultZoomLevel);
+	}
+
+	if ( OldZoomLevel != ZoomLevel )
+	{
+		PostChangedZoom();
+	}
+
+	// Note: This happens even when maxed out at a stop; so the user sees the animation and knows that they're at max zoom in/out
+	ZoomLevelFade.Play();
+
+	// Re-center the screen so that it feels like zooming around the cursor.
+	{
+		FSlateRect GraphBounds = ComputeSensibleGraphBounds();
+
+		// Make sure we are not zooming into/out into emptiness; otherwise the user will get lost..
+		const FVector2D ClampedPointToMaintainGraphSpace(
+			FMath::Clamp(PointToMaintainGraphSpace.X, GraphBounds.Left, GraphBounds.Right),
+			FMath::Clamp(PointToMaintainGraphSpace.Y, GraphBounds.Top, GraphBounds.Bottom)
+			);
+
+		this->ViewOffset = ClampedPointToMaintainGraphSpace - WidgetSpaceZoomOrigin / GetZoomAmount();
+	}
+}
+
+FSlateRect SDesignSurface::ComputeSensibleGraphBounds() const
+{
+	float Left = 0.0f;
+	float Top = 0.0f;
+	float Right = 0.0f;
+	float Bottom = 0.0f;
+
+	// Find the bounds of the node positions
+	//for ( int32 ChildIndex = 0; ChildIndex < Children.Num(); ++ChildIndex )
+	//{
+	//	const TSharedRef<SNode>& SomeChild = Children[ChildIndex];
+
+	//	FVector2D ChildPos = SomeChild->GetPosition();
+
+	//	Left = FMath::Min(Left, ChildPos.X);
+	//	Right = FMath::Max(Right, ChildPos.X);
+	//	Top = FMath::Min(Top, ChildPos.Y);
+	//	Bottom = FMath::Max(Bottom, ChildPos.Y);
+	//}
+
+	// Pad it out in every direction, to roughly account for nodes being of non-zero extent
+	const float Padding = 100.0f;
+
+	return FSlateRect(Left - Padding, Top - Padding, Right + Padding, Bottom + Padding);
+}
+
+void SDesignSurface::PostChangedZoom()
+{
+}
+
+FVector2D SDesignSurface::GetViewOffset() const
+{
+	return ViewOffset;
+}
+
+float SDesignSurface::GetSnapGridSize() const
+{
+	return SnapGridSize;
+}
+
+FVector2D SDesignSurface::GraphCoordToPanelCoord(const FVector2D& GraphSpaceCoordinate) const
+{
+	return ( GraphSpaceCoordinate - GetViewOffset() ) * GetZoomAmount();
+}
+
+FVector2D SDesignSurface::PanelCoordToGraphCoord(const FVector2D& PanelSpaceCoordinate) const
+{
+	return PanelSpaceCoordinate / GetZoomAmount() + GetViewOffset();
+}
+
+void SDesignSurface::PaintBackgroundAsLines(const FSlateBrush* BackgroundImage, const FGeometry& AllottedGeometry, const FSlateRect& MyClippingRect, FSlateWindowElementList& OutDrawElements, int32& DrawLayerId) const
+{
+	const bool bAntialias = false;
+
+	const int32 RulePeriod = (int32)FEditorStyle::GetFloat("Graph.Panel.GridRulePeriod");
+	check(RulePeriod > 0);
+
+	const FLinearColor RegularColor(FEditorStyle::GetColor("Graph.Panel.GridLineColor"));
+	const FLinearColor RuleColor(FEditorStyle::GetColor("Graph.Panel.GridRuleColor"));
+	const FLinearColor CenterColor(FEditorStyle::GetColor("Graph.Panel.GridCenterColor"));
+	const float GraphSmallestGridSize = 8.0f;
+	const float RawZoomFactor = GetZoomAmount();
+	const float NominalGridSize = GetSnapGridSize();
+
+	float ZoomFactor = RawZoomFactor;
+	float Inflation = 1.0f;
+	while ( ZoomFactor*Inflation*NominalGridSize <= GraphSmallestGridSize )
+	{
+		Inflation *= 2.0f;
+	}
+
+	const float GridCellSize = NominalGridSize * ZoomFactor * Inflation;
+
+	const float GraphSpaceGridX0 = FancyMod(ViewOffset.X, Inflation * NominalGridSize * RulePeriod);
+	const float GraphSpaceGridY0 = FancyMod(ViewOffset.Y, Inflation * NominalGridSize * RulePeriod);
+
+	float ImageOffsetX = GraphSpaceGridX0 * -ZoomFactor;
+	float ImageOffsetY = GraphSpaceGridY0 * -ZoomFactor;
+
+	const FVector2D ZeroSpace = GraphCoordToPanelCoord(FVector2D::ZeroVector);
+
+	// Fill the background
+	FSlateDrawElement::MakeBox(
+		OutDrawElements,
+		DrawLayerId,
+		AllottedGeometry.ToPaintGeometry(),
+		BackgroundImage,
+		MyClippingRect
+		);
+
+	TArray<FVector2D> LinePoints;
+	new (LinePoints)FVector2D(0.0f, 0.0f);
+	new (LinePoints)FVector2D(0.0f, 0.0f);
+
+	// Horizontal bars
+	for ( int32 GridIndex = 0; ImageOffsetY < AllottedGeometry.Size.Y; ImageOffsetY += GridCellSize, ++GridIndex )
+	{
+		if ( ImageOffsetY >= 0.0f )
+		{
+			const bool bIsRuleLine = ( GridIndex % RulePeriod ) == 0;
+			const int32 Layer = bIsRuleLine ? ( DrawLayerId + 1 ) : DrawLayerId;
+
+			const FLinearColor* Color = bIsRuleLine ? &RuleColor : &RegularColor;
+			if ( FMath::IsNearlyEqual(ZeroSpace.Y, ImageOffsetY, 1.0f) )
+			{
+				Color = &CenterColor;
+			}
+
+			LinePoints[0] = FVector2D(0.0f, ImageOffsetY);
+			LinePoints[1] = FVector2D(AllottedGeometry.Size.X, ImageOffsetY);
+
+			FSlateDrawElement::MakeLines(
+				OutDrawElements,
+				Layer,
+				AllottedGeometry.ToPaintGeometry(),
+				LinePoints,
+				MyClippingRect,
+				ESlateDrawEffect::None,
+				*Color,
+				bAntialias);
+		}
+	}
+
+	// Vertical bars
+	for ( int32 GridIndex = 0; ImageOffsetX < AllottedGeometry.Size.X; ImageOffsetX += GridCellSize, ++GridIndex )
+	{
+		if ( ImageOffsetX >= 0.0f )
+		{
+			const bool bIsRuleLine = ( GridIndex % RulePeriod ) == 0;
+			const int32 Layer = bIsRuleLine ? ( DrawLayerId + 1 ) : DrawLayerId;
+
+			const FLinearColor* Color = bIsRuleLine ? &RuleColor : &RegularColor;
+			if ( FMath::IsNearlyEqual(ZeroSpace.X, ImageOffsetX, 1.0f) )
+			{
+				Color = &CenterColor;
+			}
+
+			LinePoints[0] = FVector2D(ImageOffsetX, 0.0f);
+			LinePoints[1] = FVector2D(ImageOffsetX, AllottedGeometry.Size.Y);
+
+			FSlateDrawElement::MakeLines(
+				OutDrawElements,
+				Layer,
+				AllottedGeometry.ToPaintGeometry(),
+				LinePoints,
+				MyClippingRect,
+				ESlateDrawEffect::None,
+				*Color,
+				bAntialias);
+		}
+	}
+
+	DrawLayerId += 2;
+}
+
+#undef LOCTEXT_NAMESPACE
diff --git a/Engine/Source/Editor/UMGEditor/Private/Designer/SDesignSurface.h b/Engine/Source/Editor/UMGEditor/Private/Designer/SDesignSurface.h
new file mode 100644
--- /dev/null
+++ b/Engine/Source/Editor/UMGEditor/Private/Designer/SDesignSurface.h
@@ -0,0 +1,67 @@
+// Copyright 1998-2014 Epic Games, Inc. All Rights Reserved.
+
+#pragma once
+
+#include "SCompoundWidget.h"
+#include "SNodePanel.h"
+
+class SDesignSurface : public SCompoundWidget
+{
+public:
+
+	SLATE_BEGIN_ARGS( SDesignSurface ) {}
+		/** Slot for this designers content (optional) */
+		SLATE_DEFAULT_SLOT(FArguments, Content)
+	SLATE_END_ARGS()
+
+	void Construct(const FArguments& InArgs);
+
+	// SWidget interface
+	virtual FReply OnMouseWheel(const FGeometry& MyGeometry, const FPointerEvent& MouseEvent) override;
+	virtual int32 OnPaint(const FGeometry& AllottedGeometry, const FSlateRect& MyClippingRect, FSlateWindowElementList& OutDrawElements, int32 LayerId, const FWidgetStyle& InWidgetStyle, bool bParentEnabled) const override;
+	virtual bool SupportsKeyboardFocus() const override { return true; }
+	// End of Swidget interface
+
+protected:
+	void PaintBackgroundAsLines(const FSlateBrush* BackgroundImage, const FGeometry& AllottedGeometry, const FSlateRect& MyClippingRect, FSlateWindowElementList& OutDrawElements, int32& DrawLayerId) const;
+
+	/** Get the grid snap size */
+	float GetSnapGridSize() const;
+
+	float GetZoomAmount() const;
+
+	void ChangeZoomLevel(int32 ZoomLevelDelta, const FVector2D& WidgetSpaceZoomOrigin, bool bOverrideZoomLimiting);
+
+	void PostChangedZoom();
+
+	FVector2D GetViewOffset() const;
+
+	FSlateRect ComputeSensibleGraphBounds() const;
+
+	FVector2D GraphCoordToPanelCoord(const FVector2D& GraphSpaceCoordinate) const;
+	FVector2D PanelCoordToGraphCoord(const FVector2D& PanelSpaceCoordinate) const;
+
+protected:
+	/** The position within the graph at which the user is looking */
+	FVector2D ViewOffset;
+
+	float SnapGridSize;
+
+	/** Previous Zoom Level */
+	int32 PreviousZoomLevel;
+
+	/** How zoomed in/out we are. e.g. 0.25f results in quarter-sized nodes. */
+	int32 ZoomLevel;
+
+	/** Allow continuous zoom interpolation? */
+	bool bAllowContinousZoomInterpolation;
+
+	/** Fade on zoom for graph */
+	FCurveSequence ZoomLevelGraphFade;
+
+	/** Curve that handles fading the 'Zoom +X' text */
+	FCurveSequence ZoomLevelFade;
+
+	// The interface for mapping ZoomLevel values to actual node scaling values
+	TScopedPointer<FZoomLevelsContainer> ZoomLevels;
+};
diff --git a/Engine/Source/Editor/UMGEditor/Private/Designer/SDesignerView.cpp b/Engine/Source/Editor/UMGEditor/Private/Designer/SDesignerView.cpp
new file mode 100644
--- /dev/null
+++ b/Engine/Source/Editor/UMGEditor/Private/Designer/SDesignerView.cpp
@@ -0,0 +1,1200 @@
+﻿// Copyright 1998-2014 Epic Games, Inc. All Rights Reserved.
+
+#include "UMGEditorPrivatePCH.h"
+#include "SDesignerView.h"
+
+#include "BlueprintEditor.h"
+#include "SKismetInspector.h"
+#include "BlueprintEditorUtils.h"
+#include "WidgetBlueprint.h"
+#include "WidgetBlueprintEditor.h"
+
+#include "WidgetTemplateDragDropOp.h"
+
+#define LOCTEXT_NAMESPACE "UMG"
+
+const float HoveredAnimationTime = 0.150f;
+
+class FSelectedWidgetDragDropOp : public FDecoratedDragDropOp
+{
+public:
+	DRAG_DROP_OPERATOR_TYPE(FSelectedWidgetDragDropOp, FDecoratedDragDropOp)
+
+	TMap<FName, FString> ExportedSlotProperties;
+
+	FWidgetReference Widget;
+
+	static TSharedRef<FSelectedWidgetDragDropOp> New(FWidgetReference InWidget);
+};
+
+TSharedRef<FSelectedWidgetDragDropOp> FSelectedWidgetDragDropOp::New(FWidgetReference InWidget)
+{
+	TSharedRef<FSelectedWidgetDragDropOp> Operation = MakeShareable(new FSelectedWidgetDragDropOp());
+	Operation->Widget = InWidget;
+	Operation->DefaultHoverText = FText::FromString( InWidget.GetTemplate()->GetLabel() );
+	Operation->CurrentHoverText = FText::FromString( InWidget.GetTemplate()->GetLabel() );
+	Operation->Construct();
+
+	FWidgetBlueprintEditorUtils::ExportPropertiesToText(InWidget.GetTemplate()->Slot, Operation->ExportedSlotProperties);
+
+	return Operation;
+}
+
+//////////////////////////////////////////////////////////////////////////
+
+static bool LocateWidgetsUnderCursor_Helper(FArrangedWidget& Candidate, FVector2D InAbsoluteCursorLocation, FArrangedChildren& OutWidgetsUnderCursor, bool bIgnoreEnabledStatus)
+{
+	const bool bCandidateUnderCursor =
+		// Candidate is physically under the cursor
+		Candidate.Geometry.IsUnderLocation(InAbsoluteCursorLocation) &&
+		// Candidate actually considers itself hit by this test
+		( Candidate.Widget->OnHitTest(Candidate.Geometry, InAbsoluteCursorLocation) );
+
+	bool bHitAnyWidget = false;
+	if ( bCandidateUnderCursor )
+	{
+		// The candidate widget is under the mouse
+		OutWidgetsUnderCursor.AddWidget(Candidate);
+
+		// Check to see if we were asked to still allow children to be hit test visible
+		bool bHitChildWidget = false;
+		
+		if ( Candidate.Widget->GetVisibility().AreChildrenHitTestVisible() )//!= 0 || OutWidgetsUnderCursor. )
+		{
+			FArrangedChildren ArrangedChildren(OutWidgetsUnderCursor.GetFilter());
+			Candidate.Widget->ArrangeChildren(Candidate.Geometry, ArrangedChildren);
+
+			// A widget's children are implicitly Z-ordered from first to last
+			for ( int32 ChildIndex = ArrangedChildren.Num() - 1; !bHitChildWidget && ChildIndex >= 0; --ChildIndex )
+			{
+				FArrangedWidget& SomeChild = ArrangedChildren(ChildIndex);
+				bHitChildWidget = ( SomeChild.Widget->IsEnabled() || bIgnoreEnabledStatus ) && LocateWidgetsUnderCursor_Helper(SomeChild, InAbsoluteCursorLocation, OutWidgetsUnderCursor, bIgnoreEnabledStatus);
+			}
+		}
+
+		// If we hit a child widget or we hit our candidate widget then we'll append our widgets
+		const bool bHitCandidateWidget = OutWidgetsUnderCursor.Accepts(Candidate.Widget->GetVisibility()) &&
+			Candidate.Widget->GetVisibility().AreChildrenHitTestVisible();
+		
+		bHitAnyWidget = bHitChildWidget || bHitCandidateWidget;
+		if ( !bHitAnyWidget )
+		{
+			// No child widgets were hit, and even though the cursor was over our candidate widget, the candidate
+			// widget was not hit-testable, so we won't report it
+			check(OutWidgetsUnderCursor.Last() == Candidate);
+			OutWidgetsUnderCursor.Remove(OutWidgetsUnderCursor.Num() - 1);
+		}
+	}
+
+	return bHitAnyWidget;
+}
+
+/////////////////////////////////////////////////////
+// SDesignerView
+
+void SDesignerView::Construct(const FArguments& InArgs, TSharedPtr<FWidgetBlueprintEditor> InBlueprintEditor)
+{
+	ScopedTransaction = NULL;
+
+	PreviewWidget = NULL;
+	DropPreviewWidget = NULL;
+	DropPreviewParent = NULL;
+	BlueprintEditor = InBlueprintEditor;
+	CurrentHandle = DH_NONE;
+
+	HoverTime = 0;
+
+	bMouseDown = false;
+	bMovingExistingWidget = false;
+
+	DragDirections.Init((int32)DH_MAX);
+	DragDirections[DH_TOP_LEFT] = FVector2D(-1, -1);
+	DragDirections[DH_TOP_CENTER] = FVector2D(0, -1);
+	DragDirections[DH_TOP_RIGHT] = FVector2D(1, -1);
+
+	DragDirections[DH_MIDDLE_LEFT] = FVector2D(-1, 0);
+	DragDirections[DH_MIDDLE_RIGHT] = FVector2D(1, 0);
+
+	DragDirections[DH_BOTTOM_LEFT] = FVector2D(-1, 1);
+	DragDirections[DH_BOTTOM_CENTER] = FVector2D(0, 1);
+	DragDirections[DH_BOTTOM_RIGHT] = FVector2D(1, 1);
+
+	Register(MakeShareable(new FVerticalSlotExtension()));
+	Register(MakeShareable(new FHorizontalSlotExtension()));
+	Register(MakeShareable(new FCanvasSlotExtension()));
+	Register(MakeShareable(new FUniformGridSlotExtension()));
+
+	UWidgetBlueprint* Blueprint = GetBlueprint();
+	Blueprint->OnChanged().AddSP(this, &SDesignerView::OnBlueprintChanged);
+
+	SDesignSurface::Construct(SDesignSurface::FArguments()
+		.Content()
+		[
+			SNew(SOverlay)
+
+			// The bottom layer of the overlay where the actual preview widget appears.
+			+ SOverlay::Slot()
+			.HAlign(HAlign_Fill)
+			.VAlign(VAlign_Fill)
+			[
+				SAssignNew(PreviewSurface, SBorder)
+				.Visibility(EVisibility::HitTestInvisible)
+			]
+
+			// A layer in the overlay where we put all the user intractable widgets, like the reorder widgets.
+			+ SOverlay::Slot()
+			.HAlign(HAlign_Fill)
+			.VAlign(VAlign_Fill)
+			[
+				SAssignNew(ExtensionWidgetCanvas, SCanvas)
+				.Visibility(EVisibility::SelfHitTestInvisible)
+			]
+
+			// Bottom-right corner text indicating the experimental nature of umg
+			+ SOverlay::Slot()
+			.Padding(10)
+			.VAlign(VAlign_Bottom)
+			.HAlign(HAlign_Right)
+			[
+				SNew(STextBlock)
+				.Visibility( EVisibility::HitTestInvisible )
+				.TextStyle( FEditorStyle::Get(), "Graph.CornerText" )
+				.Text(LOCTEXT("ExperimentalWarning", "Experimental"))
+			]
+		]
+	);
+
+	BlueprintEditor.Pin()->OnSelectedWidgetsChanged.AddRaw(this, &SDesignerView::OnEditorSelectionChanged);
+}
+
+SDesignerView::~SDesignerView()
+{
+	UWidgetBlueprint* Blueprint = GetBlueprint();
+	if ( Blueprint )
+	{
+		Blueprint->OnChanged().RemoveAll(this);
+	}
+
+	if ( BlueprintEditor.IsValid() )
+	{
+		BlueprintEditor.Pin()->OnSelectedWidgetsChanged.RemoveAll(this);
+	}
+}
+
+void SDesignerView::OnEditorSelectionChanged()
+{
+	SelectedWidgets = BlueprintEditor.Pin()->GetSelectedWidgets();
+
+	if ( SelectedWidgets.Num() > 0 )
+	{
+		for ( FWidgetReference& Widget : SelectedWidgets )
+		{
+			SelectedWidget = Widget;
+			break;
+		}
+	}
+	else
+	{
+		SelectedWidget = FWidgetReference();
+	}
+
+	CreateExtensionWidgetsForSelection();
+}
+
+void SDesignerView::CreateExtensionWidgetsForSelection()
+{
+	// Remove all the current extension widgets
+	ExtensionWidgetCanvas->ClearChildren();
+
+	ExtensionWidgets.Reset();
+
+	TArray<FWidgetReference> Selected;
+	if ( SelectedWidget.IsValid() )
+	{
+		Selected.Add(SelectedWidget);
+	}
+
+	// Build extension widgets for new selection
+	for ( TSharedRef<FDesignerExtension>& Ext : DesignerExtensions )
+	{
+		Ext->BuildWidgetsForSelection(Selected, ExtensionWidgets);
+	}
+
+	//TODO UMG We need to have a better layout system for extensions.
+	//         currently everything is just laid out across the top.
+
+	TSharedRef<SHorizontalBox> ExtensionBox = SNew(SHorizontalBox);
+
+	// Add Widgets to designer surface
+	for ( TSharedRef<SWidget>& ExWidget : ExtensionWidgets )
+	{
+		ExtensionBox->AddSlot()
+			.AutoWidth()
+			[
+				ExWidget
+			];
+	}
+
+	ExtensionWidgetCanvas->AddSlot()
+		.Position(TAttribute<FVector2D>(this, &SDesignerView::GetCachedSelectionDesignerWidgetsLocation))
+		.Size(FVector2D(100, 25))
+		[
+			ExtensionBox
+		];
+}
+
+UWidgetBlueprint* SDesignerView::GetBlueprint() const
+{
+	if ( BlueprintEditor.IsValid() )
+	{
+		UBlueprint* BP = BlueprintEditor.Pin()->GetBlueprintObj();
+		return Cast<UWidgetBlueprint>(BP);
+	}
+
+	return NULL;
+}
+
+void SDesignerView::Register(TSharedRef<FDesignerExtension> Extension)
+{
+	Extension->Initialize(GetBlueprint());
+	DesignerExtensions.Add(Extension);
+}
+
+void SDesignerView::OnBlueprintChanged(UBlueprint* InBlueprint)
+{
+	if ( InBlueprint )
+	{
+		
+	}
+}
+
+void SDesignerView::OnObjectPropertyChanged(UObject* ObjectBeingModified, FPropertyChangedEvent& PropertyChangedEvent)
+{
+	if ( !ensure(ObjectBeingModified) )
+	{
+		return;
+	}
+
+	//UpdatePreview(InBlueprint);
+}
+
+FWidgetReference SDesignerView::GetWidgetAtCursor(const FGeometry& MyGeometry, const FPointerEvent& MouseEvent, FArrangedWidget& ArrangedWidget)
+{
+	//@TODO UMG Make it so you can request dropable widgets only, to find the first parentable.
+
+	FArrangedChildren Children(EVisibility::All);
+
+	PreviewSurface->SetVisibility(EVisibility::Visible);
+	FArrangedWidget WindowWidgetGeometry(PreviewSurface.ToSharedRef(), MyGeometry);
+	LocateWidgetsUnderCursor_Helper(WindowWidgetGeometry, MouseEvent.GetScreenSpacePosition(), Children, true);
+
+	PreviewSurface->SetVisibility(EVisibility::HitTestInvisible);
+
+	UUserWidget* WidgetActor = BlueprintEditor.Pin()->GetPreview();
+	if ( WidgetActor )
+	{
+		UWidget* Preview = NULL;
+
+		for ( int32 ChildIndex = Children.Num() - 1; ChildIndex >= 0; ChildIndex-- )
+		{
+			FArrangedWidget& Child = Children.GetInternalArray()[ChildIndex];
+			Preview = WidgetActor->GetWidgetHandle(Child.Widget);
+			
+			// Ignore the drop preview widget when doing widget picking
+			if (Preview == DropPreviewWidget)
+			{
+				Preview = NULL;
+				continue;
+			}
+			
+			if ( Preview )
+			{
+				ArrangedWidget = Child;
+				break;
+			}
+		}
+
+		if ( Preview )
+		{
+			return FWidgetReference::FromPreview(BlueprintEditor.Pin(), Preview);
+		}
+	}
+
+	return FWidgetReference();
+}
+
+FReply SDesignerView::OnMouseButtonDown(const FGeometry& MyGeometry, const FPointerEvent& MouseEvent)
+{
+	CurrentHandle = HitTestDragHandles(MyGeometry, MouseEvent);
+
+	if ( CurrentHandle != DH_NONE )
+	{
+		BeginTransaction(LOCTEXT("ResizeWidget", "Resize Widget"));
+		return FReply::Handled().PreventThrottling().CaptureMouse(AsShared());
+	}
+
+	//TODO UMG Undoable Selection
+	FArrangedWidget ArrangedWidget(SNullWidget::NullWidget, FGeometry());
+	SelectedWidget = GetWidgetAtCursor(MyGeometry, MouseEvent, ArrangedWidget);
+	SelectedWidgetContextMenuLocation = ArrangedWidget.Geometry.AbsoluteToLocal(MouseEvent.GetScreenSpacePosition());
+
+	if ( SelectedWidget.IsValid() )
+	{
+		//@TODO UMG primary FBlueprintEditor needs to be inherited and selection control needs to be centralized.
+		// Set the template as selected in the details panel
+		TSet<FWidgetReference> SelectedTemplates;
+		SelectedTemplates.Add(SelectedWidget);
+		BlueprintEditor.Pin()->SelectWidgets(SelectedTemplates);
+
+		if ( MouseEvent.GetEffectingButton() == EKeys::LeftMouseButton )
+		{
+			bMouseDown = true;
+		}
+	}
+
+	// Capture mouse for the drag handle and general mouse actions
+	return FReply::Handled().PreventThrottling().SetKeyboardFocus(AsShared(), EKeyboardFocusCause::Mouse).CaptureMouse(AsShared());
+}
+
+FReply SDesignerView::OnMouseButtonUp(const FGeometry& MyGeometry, const FPointerEvent& MouseEvent)
+{
+	if ( MouseEvent.GetEffectingButton() == EKeys::LeftMouseButton )
+	{
+		bMouseDown = false;
+		bMovingExistingWidget = false;
+
+		if ( CurrentHandle != DH_NONE )
+		{
+			EndTransaction();
+			CurrentHandle = DH_NONE;
+		}
+	}
+	else if ( MouseEvent.GetEffectingButton() == EKeys::RightMouseButton )
+	{
+		ShowContextMenu(MyGeometry, MouseEvent);
+	}
+
+	return FReply::Handled().ReleaseMouseCapture();
+}
+
+FReply SDesignerView::OnMouseMove(const FGeometry& MyGeometry, const FPointerEvent& MouseEvent)
+{
+	if ( MouseEvent.GetCursorDelta().IsZero() )
+	{
+		return FReply::Unhandled();
+	}
+
+	// Update the resizing based on the mouse movement
+	if ( CurrentHandle != DH_NONE )
+	{
+		if ( SelectedWidget.IsValid() )
+		{
+			//@TODO UMG - Implement some system to query slots to know if they support dragging so we can provide visual feedback by hiding handles that wouldn't work and such.
+			//SelectedTemplate->Slot->CanResize(DH_TOP_CENTER)
+
+			UWidget* Template = SelectedWidget.GetTemplate();
+			UWidget* Preview = SelectedWidget.GetPreview();
+
+			if ( Preview->Slot )
+			{
+				Preview->Slot->Resize(DragDirections[CurrentHandle], MouseEvent.GetCursorDelta());
+			}
+
+			if ( Template->Slot )
+			{
+				Template->Slot->Resize(DragDirections[CurrentHandle], MouseEvent.GetCursorDelta());
+			}
+
+			FBlueprintEditorUtils::MarkBlueprintAsModified(GetBlueprint());
+
+			return FReply::Handled().PreventThrottling();
+		}
+	}
+
+	if ( bMouseDown )
+	{
+		if ( SelectedWidget.IsValid() && !bMovingExistingWidget )
+		{
+			const bool bIsRootWidget = SelectedWidget.GetTemplate()->GetParent() == NULL;
+			if ( !bIsRootWidget )
+			{
+				bMovingExistingWidget = true;
+				//Drag selected widgets
+				return FReply::Handled().DetectDrag(AsShared(), EKeys::LeftMouseButton);
+			}
+		}
+	}
+	
+	if ( CurrentHandle == DH_NONE )
+	{
+		// Update the hovered widget under the mouse
+		FArrangedWidget ArrangedWidget(SNullWidget::NullWidget, FGeometry());
+		FWidgetReference NewHoveredWidget = GetWidgetAtCursor(MyGeometry, MouseEvent, ArrangedWidget);
+		if ( !( NewHoveredWidget == HoveredWidget ) )
+		{
+			HoveredWidget = NewHoveredWidget;
+			HoverTime = 0;
+		}
+	}
+
+	return FReply::Unhandled();
+}
+
+void SDesignerView::OnMouseEnter(const FGeometry& MyGeometry, const FPointerEvent& MouseEvent)
+{
+	HoveredWidget = FWidgetReference();
+	HoverTime = 0;
+}
+
+void SDesignerView::OnMouseLeave(const FPointerEvent& MouseEvent)
+{
+	HoveredWidget = FWidgetReference();
+	HoverTime = 0;
+}
+
+FReply SDesignerView::OnKeyDown(const FGeometry& MyGeometry, const FKeyboardEvent& InKeyboardEvent)
+{
+	BlueprintEditor.Pin()->PasteDropLocation = FVector2D(0, 0);
+
+	if ( BlueprintEditor.Pin()->WidgetCommandList->ProcessCommandBindings(InKeyboardEvent) )
+	{
+		return FReply::Handled();
+	}
+
+	return FReply::Unhandled();
+}
+
+void SDesignerView::ShowContextMenu(const FGeometry& MyGeometry, const FPointerEvent& MouseEvent)
+{
+	FMenuBuilder MenuBuilder(true, NULL);
+
+	FWidgetBlueprintEditorUtils::CreateWidgetContextMenu(MenuBuilder, BlueprintEditor.Pin().ToSharedRef(), SelectedWidgetContextMenuLocation);
+
+	TSharedPtr<SWidget> MenuContent = MenuBuilder.MakeWidget();
+
+	if ( MenuContent.IsValid() )
+	{
+		FVector2D SummonLocation = MouseEvent.GetScreenSpacePosition();
+		FSlateApplication::Get().PushMenu(AsShared(), MenuContent.ToSharedRef(), SummonLocation, FPopupTransitionEffect(FPopupTransitionEffect::ContextMenu));
+	}
+}
+
+bool SDesignerView::GetArrangedWidget(TSharedRef<SWidget> Widget, FArrangedWidget& ArrangedWidget) const
+{
+	// We can't screenshot the widget unless there's a valid window handle to draw it in.
+	TSharedPtr<SWindow> WidgetWindow = FSlateApplication::Get().FindWidgetWindow(Widget);
+	if ( !WidgetWindow.IsValid() )
+	{
+		return false;
+	}
+
+	TSharedRef<SWindow> CurrentWindowRef = WidgetWindow.ToSharedRef();
+
+	FWidgetPath WidgetPath;
+	if ( FSlateApplication::Get().GeneratePathToWidgetUnchecked(Widget, WidgetPath) )
+	{
+		ArrangedWidget = WidgetPath.FindArrangedWidget(Widget);
+		return true;
+	}
+
+	return false;
+}
+
+bool SDesignerView::GetArrangedWidgetRelativeToWindow(TSharedRef<SWidget> Widget, FArrangedWidget& ArrangedWidget) const
+{
+	// We can't screenshot the widget unless there's a valid window handle to draw it in.
+	TSharedPtr<SWindow> WidgetWindow = FSlateApplication::Get().FindWidgetWindow(Widget);
+	if ( !WidgetWindow.IsValid() )
+	{
+		return false;
+	}
+
+	TSharedRef<SWindow> CurrentWindowRef = WidgetWindow.ToSharedRef();
+
+	FWidgetPath WidgetPath;
+	if ( FSlateApplication::Get().GeneratePathToWidgetUnchecked(Widget, WidgetPath) )
+	{
+		ArrangedWidget = WidgetPath.FindArrangedWidget(Widget);
+		ArrangedWidget.Geometry.AbsolutePosition -= CurrentWindowRef->GetPositionInScreen();
+		return true;
+	}
+
+	return false;
+}
+
+bool SDesignerView::GetArrangedWidgetRelativeToDesigner(TSharedRef<SWidget> Widget, FArrangedWidget& ArrangedWidget) const
+{
+	FWidgetPath WidgetPath;
+	if ( FSlateApplication::Get().GeneratePathToWidgetUnchecked(Widget, WidgetPath) )
+	{
+		FArrangedWidget ArrangedDesigner = WidgetPath.FindArrangedWidget(this->AsShared());
+
+		ArrangedWidget = WidgetPath.FindArrangedWidget(Widget);
+		ArrangedWidget.Geometry.AbsolutePosition -= ArrangedDesigner.Geometry.AbsolutePosition;
+		return true;
+	}
+
+	return false;
+}
+
+FVector2D SDesignerView::GetCachedSelectionDesignerWidgetsLocation() const
+{
+	return CachedDesignerWidgetLocation;
+}
+
+FVector2D SDesignerView::GetSelectionDesignerWidgetsLocation() const
+{
+	if ( SelectedSlateWidget.IsValid() )
+	{
+		TSharedRef<SWidget> Widget = SelectedSlateWidget.Pin().ToSharedRef();
+
+		FArrangedWidget ArrangedWidget(SNullWidget::NullWidget, FGeometry());
+		GetArrangedWidgetRelativeToDesigner(Widget, ArrangedWidget);
+
+		return ArrangedWidget.Geometry.AbsolutePosition - FVector2D(0, 25);
+	}
+
+	return FVector2D(0, 0);
+}
+
+int32 SDesignerView::OnPaint(const FGeometry& AllottedGeometry, const FSlateRect& MyClippingRect, FSlateWindowElementList& OutDrawElements, int32 LayerId, const FWidgetStyle& InWidgetStyle, bool bParentEnabled) const
+{
+	SDesignSurface::OnPaint(AllottedGeometry, MyClippingRect, OutDrawElements, LayerId, InWidgetStyle, bParentEnabled);
+
+	LayerId += 1000;
+
+	TSet<FWidgetReference> Selected;
+	Selected.Add(SelectedWidget);
+
+	// Allow the extensions to paint anything they want.
+	for ( const TSharedRef<FDesignerExtension>& Ext : DesignerExtensions )
+	{
+		Ext->Paint(Selected, AllottedGeometry, MyClippingRect, OutDrawElements, LayerId, InWidgetStyle, bParentEnabled);
+	}
+
+	// Don't draw the hovered effect if it's also the selected widget
+	if ( HoveredSlateWidget.IsValid() && HoveredSlateWidget != SelectedSlateWidget )
+	{
+		TSharedRef<SWidget> Widget = HoveredSlateWidget.Pin().ToSharedRef();
+
+		FArrangedWidget ArrangedWidget(SNullWidget::NullWidget, FGeometry());
+		GetArrangedWidgetRelativeToWindow(Widget, ArrangedWidget);
+
+		// Draw hovered effect
+		// Azure = 0x007FFF
+		const FLinearColor HoveredTint(0, 0.5, 1, FMath::Clamp(HoverTime / HoveredAnimationTime, 0.0f, 1.0f));
+
+		FPaintGeometry HoveredGeometry(
+			ArrangedWidget.Geometry.AbsolutePosition,
+			ArrangedWidget.Geometry.Size * ArrangedWidget.Geometry.Scale,
+			ArrangedWidget.Geometry.Scale);
+
+		FSlateDrawElement::MakeBox(
+			OutDrawElements,
+			LayerId,
+			HoveredGeometry,
+			FCoreStyle::Get().GetBrush(TEXT("Debug.Border")),
+			MyClippingRect,
+			ESlateDrawEffect::None,
+			HoveredTint
+			);
+	}
+
+	if ( SelectedSlateWidget.IsValid() )
+	{
+		TSharedRef<SWidget> Widget = SelectedSlateWidget.Pin().ToSharedRef();
+
+		FArrangedWidget ArrangedWidget(SNullWidget::NullWidget, FGeometry());
+		GetArrangedWidgetRelativeToWindow(Widget, ArrangedWidget);
+
+		const FLinearColor Tint(0, 1, 0);
+
+		// Draw selection effect
+		FPaintGeometry SelectionGeometry(
+			ArrangedWidget.Geometry.AbsolutePosition,
+			ArrangedWidget.Geometry.Size * ArrangedWidget.Geometry.Scale,
+			ArrangedWidget.Geometry.Scale);
+
+		FSlateDrawElement::MakeBox(
+			OutDrawElements,
+			LayerId,
+			SelectionGeometry,
+			FCoreStyle::Get().GetBrush(TEXT("Debug.Border")),
+			MyClippingRect,
+			ESlateDrawEffect::None,
+			Tint
+		);
+
+		DrawDragHandles(SelectionGeometry, MyClippingRect, OutDrawElements, LayerId, InWidgetStyle);
+
+		//TODO UMG Move this to a function.
+		// Draw anchors.
+		if ( SelectedWidget.IsValid() && SelectedWidget.GetTemplate()->Slot )
+		{
+			const float X = AllottedGeometry.AbsolutePosition.X;
+			const float Y = AllottedGeometry.AbsolutePosition.Y;
+			const float Width = AllottedGeometry.Size.X;
+			const float Height = AllottedGeometry.Size.Y;
+
+			float Selection_X = SelectionGeometry.DrawPosition.X;
+			float Selection_Y = SelectionGeometry.DrawPosition.Y;
+			float Selection_Width = SelectionGeometry.DrawSize.X;
+			float Selection_Height = SelectionGeometry.DrawSize.Y;
+
+			const FVector2D LeftRightSize = FVector2D(32, 16);
+			const FVector2D TopBottomSize = FVector2D(16, 32);
+
+			// @TODO UMG - Don't use the curve editors brushes
+			const FSlateBrush* KeyBrush = FEditorStyle::GetBrush("CurveEd.CurveKey");
+			FLinearColor KeyColor = FLinearColor(1.0f, 0.0f, 0.0f, 1.f);
+
+			if ( UCanvasPanelSlot* CanvasSlot = Cast<UCanvasPanelSlot>(SelectedWidget.GetTemplate()->Slot) )
+			{
+				//Set("UMGEditor.AnchorCenter", new IMAGE_BRUSH("Icons/umg_anchor_center", Icon16x16));
+				//Set("UMGEditor.AnchorTopBottom", new IMAGE_BRUSH("Icons/umg_anchor_top_bottom", FVector2D(16, 32)));
+				//Set("UMGEditor.AnchorLeftRight", new IMAGE_BRUSH("Icons/umg_anchor_left_right", FVector2D(32, 16)));
+
+				// Left
+				FSlateDrawElement::MakeBox(
+					OutDrawElements,
+					++LayerId,
+					FPaintGeometry(FVector2D(X + Width * CanvasSlot->LayoutData.Anchors.Minimum.X - LeftRightSize.X * 0.5f, SelectionGeometry.DrawPosition.Y + Selection_Height * 0.5f - LeftRightSize.Y * 0.5f), LeftRightSize, 1.0f),
+					FEditorStyle::Get().GetBrush("UMGEditor.AnchorLeftRight"),
+					MyClippingRect,
+					ESlateDrawEffect::None,
+					KeyBrush->GetTint(InWidgetStyle) * InWidgetStyle.GetColorAndOpacityTint() * KeyColor
+					);
+
+				// Right
+				FSlateDrawElement::MakeBox(
+					OutDrawElements,
+					++LayerId,
+					FPaintGeometry(FVector2D(X + Width * CanvasSlot->LayoutData.Anchors.Maximum.X - LeftRightSize.X * 0.5f, SelectionGeometry.DrawPosition.Y + Selection_Height * 0.5f - LeftRightSize.Y * 0.5f), LeftRightSize, 1.0f),
+					FEditorStyle::Get().GetBrush("UMGEditor.AnchorLeftRight"),
+					MyClippingRect,
+					ESlateDrawEffect::None,
+					KeyBrush->GetTint(InWidgetStyle) * InWidgetStyle.GetColorAndOpacityTint() * KeyColor
+					);
+
+
+				// Top
+				FSlateDrawElement::MakeBox(
+					OutDrawElements,
+					++LayerId,
+					FPaintGeometry(FVector2D(SelectionGeometry.DrawPosition.X + Selection_Width * 0.5f - TopBottomSize.X * 0.5f, Y + Height * CanvasSlot->LayoutData.Anchors.Minimum.Y - TopBottomSize.Y * 0.5f), TopBottomSize, 1.0f),
+					FEditorStyle::Get().GetBrush("UMGEditor.AnchorTopBottom"),
+					MyClippingRect,
+					ESlateDrawEffect::None,
+					KeyBrush->GetTint(InWidgetStyle) * InWidgetStyle.GetColorAndOpacityTint() * KeyColor
+					);
+
+				// Bottom
+				FSlateDrawElement::MakeBox(
+					OutDrawElements,
+					++LayerId,
+					FPaintGeometry(FVector2D(SelectionGeometry.DrawPosition.X + Selection_Width * 0.5f - TopBottomSize.X * 0.5f, Y + Height * CanvasSlot->LayoutData.Anchors.Maximum.Y - TopBottomSize.Y * 0.5f), TopBottomSize, 1.0f),
+					FEditorStyle::Get().GetBrush("UMGEditor.AnchorTopBottom"),
+					MyClippingRect,
+					ESlateDrawEffect::None,
+					KeyBrush->GetTint(InWidgetStyle) * InWidgetStyle.GetColorAndOpacityTint() * KeyColor
+					);
+			}
+		}
+	}
+
+	return LayerId;
+}
+
+void SDesignerView::DrawDragHandles(const FPaintGeometry& SelectionGeometry, const FSlateRect& MyClippingRect, FSlateWindowElementList& OutDrawElements, int32 LayerId, const FWidgetStyle& InWidgetStyle) const
+{
+	if ( SelectedWidget.IsValid() && SelectedWidget.GetTemplate()->Slot )
+	{
+		float X = SelectionGeometry.DrawPosition.X;
+		float Y = SelectionGeometry.DrawPosition.Y;
+		float Width = SelectionGeometry.DrawSize.X;
+		float Height = SelectionGeometry.DrawSize.Y;
+
+		// @TODO UMG Handles should come from the slot/container to express how its slots can be transformed.
+		TArray<FVector2D> Handles;
+		Handles.Add(FVector2D(X, Y));					// Top - Left
+		Handles.Add(FVector2D(X + Width * 0.5f, Y));	// Top - Middle
+		Handles.Add(FVector2D(X + Width, Y));			// Top - Right
+
+		Handles.Add(FVector2D(X, Y + Height * 0.5f));			// Middle - Left
+		Handles.Add(FVector2D(X + Width, Y + Height * 0.5f));	// Middle - Right
+
+		Handles.Add(FVector2D(X, Y + Height));					// Bottom - Left
+		Handles.Add(FVector2D(X + Width * 0.5f, Y + Height));	// Bottom - Middle
+		Handles.Add(FVector2D(X + Width, Y + Height));			// Bottom - Right
+
+		const FVector2D HandleSize = FVector2D(10, 10);
+
+		// @TODO UMG - Don't use the curve editors brushes
+		const FSlateBrush* KeyBrush = FEditorStyle::GetBrush("CurveEd.CurveKey");
+		FLinearColor KeyColor = InWidgetStyle.GetColorAndOpacityTint();// IsEditingEnabled() ? InWidgetStyle.GetColorAndOpacityTint() : FLinearColor(0.1f, 0.1f, 0.1f, 1.f);
+
+		for ( int32 HandleIndex = 0; HandleIndex < Handles.Num(); HandleIndex++ )
+		{
+			const FVector2D& Handle = Handles[HandleIndex];
+			if ( !SelectedWidget.GetTemplate()->Slot->CanResize(DragDirections[HandleIndex]) )
+			{
+				// This isn't a valid handle
+				continue;
+			}
+
+			FSlateDrawElement::MakeBox(
+				OutDrawElements,
+				++LayerId,
+				FPaintGeometry(FVector2D(Handle.X - HandleSize.X * 0.5f, Handle.Y - HandleSize.Y * 0.5f), HandleSize, 1.0f),
+				KeyBrush,
+				MyClippingRect,
+				ESlateDrawEffect::None,
+				KeyBrush->GetTint(InWidgetStyle) * InWidgetStyle.GetColorAndOpacityTint() * KeyColor
+				);
+		}
+	}
+}
+
+SDesignerView::DragHandle SDesignerView::HitTestDragHandles(const FGeometry& AllottedGeometry, const FPointerEvent& PointerEvent) const
+{
+	FVector2D LocalPointer = AllottedGeometry.AbsoluteToLocal(PointerEvent.GetScreenSpacePosition());
+
+	if ( SelectedWidget.IsValid() && SelectedWidget.GetTemplate()->Slot && SelectedSlateWidget.IsValid() )
+	{
+		TSharedRef<SWidget> Widget = SelectedSlateWidget.Pin().ToSharedRef();
+
+		FArrangedWidget ArrangedWidget(SNullWidget::NullWidget, FGeometry());
+		GetArrangedWidget(Widget, ArrangedWidget);
+
+		FVector2D WidgetLocalPosition = AllottedGeometry.AbsoluteToLocal(ArrangedWidget.Geometry.AbsolutePosition);
+		float X = WidgetLocalPosition.X;
+		float Y = WidgetLocalPosition.Y;
+		float Width = ArrangedWidget.Geometry.Size.X;
+		float Height = ArrangedWidget.Geometry.Size.Y;
+
+		TArray<FVector2D> Handles;
+		Handles.Add(FVector2D(X, Y));					// Top - Left
+		Handles.Add(FVector2D(X + Width * 0.5f, Y));	// Top - Middle
+		Handles.Add(FVector2D(X + Width, Y));			// Top - Right
+
+		Handles.Add(FVector2D(X, Y + Height * 0.5f));			// Middle - Left
+		Handles.Add(FVector2D(X + Width, Y + Height * 0.5f));	// Middle - Right
+
+		Handles.Add(FVector2D(X, Y + Height));					// Bottom - Left
+		Handles.Add(FVector2D(X + Width * 0.5f, Y + Height));	// Bottom - Middle
+		Handles.Add(FVector2D(X + Width, Y + Height));			// Bottom - Right
+
+		const FVector2D HandleSize = FVector2D(10, 10);
+
+		int32 i = 0;
+		for ( FVector2D Handle : Handles )
+		{
+			FSlateRect Rect(FVector2D(Handle.X - HandleSize.X * 0.5f, Handle.Y - HandleSize.Y * 0.5f), Handle + HandleSize);
+			if ( Rect.ContainsPoint(LocalPointer) )
+			{
+				if ( !SelectedWidget.GetTemplate()->Slot->CanResize(DragDirections[i]) )
+				{
+					// This isn't a valid handle
+					break;
+				}
+
+				return (DragHandle)i;
+			}
+			i++;
+		}
+	}
+
+	return DH_NONE;
+}
+
+FCursorReply SDesignerView::OnCursorQuery(const FGeometry& MyGeometry, const FPointerEvent& CursorEvent) const
+{
+	// Hit test the drag handles
+	switch ( HitTestDragHandles(MyGeometry, CursorEvent) )
+	{
+	case DH_TOP_LEFT:
+	case DH_BOTTOM_RIGHT:
+		return FCursorReply::Cursor(EMouseCursor::ResizeSouthEast);
+	case DH_TOP_RIGHT:
+	case DH_BOTTOM_LEFT:
+		return FCursorReply::Cursor(EMouseCursor::ResizeSouthWest);
+	case DH_TOP_CENTER:
+	case DH_BOTTOM_CENTER:
+		return FCursorReply::Cursor(EMouseCursor::ResizeUpDown);
+	case DH_MIDDLE_LEFT:
+	case DH_MIDDLE_RIGHT:
+		return FCursorReply::Cursor(EMouseCursor::ResizeLeftRight);
+	}
+
+	return FCursorReply::Unhandled();
+}
+
+void SDesignerView::Tick(const FGeometry& AllottedGeometry, const double InCurrentTime, const float InDeltaTime)
+{
+	HoverTime += InDeltaTime;
+
+	UUserWidget* LatestPreviewWidget = BlueprintEditor.Pin()->GetPreview();
+
+	if ( LatestPreviewWidget != PreviewWidget )
+	{
+		PreviewWidget = LatestPreviewWidget;
+		if ( PreviewWidget )
+		{
+			TSharedRef<SWidget> CurrentWidget = PreviewWidget->MakeFullScreenWidget();
+			CurrentWidget->SlatePrepass();
+
+			if ( CurrentWidget != PreviewSlateWidget.Pin() )
+			{
+				PreviewSlateWidget = CurrentWidget;
+				PreviewSurface->SetContent(CurrentWidget);
+			}
+		}
+		else
+		{
+			ChildSlot
+			[
+				SNew(SHorizontalBox)
+				+ SHorizontalBox::Slot()
+				.HAlign(HAlign_Center)
+				.VAlign(VAlign_Center)
+				[
+					SNew(STextBlock)
+					.Text(LOCTEXT("NoWidgetPreview", "No Widget Preview"))
+				]
+			];
+		}
+	}
+
+	// Update the selected widget to match the selected template.
+	if ( LatestPreviewWidget )
+	{
+		if ( SelectedWidget.IsValid() )
+		{
+			// Set the selected widget so that we can draw the highlight
+			SelectedSlateWidget = LatestPreviewWidget->GetWidgetFromName(SelectedWidget.GetTemplate()->GetName());
+		}
+		else
+		{
+			SelectedSlateWidget.Reset();
+		}
+
+		if ( HoveredWidget.IsValid() )
+		{
+			HoveredSlateWidget = LatestPreviewWidget->GetWidgetFromName(HoveredWidget.GetTemplate()->GetName());
+		}
+		else
+		{
+			HoveredSlateWidget.Reset();
+		}
+	}
+
+	CachedDesignerWidgetLocation = GetSelectionDesignerWidgetsLocation();
+
+	SCompoundWidget::Tick(AllottedGeometry, InCurrentTime, InDeltaTime);
+}
+
+FReply SDesignerView::OnDragDetected(const FGeometry& MyGeometry, const FPointerEvent& MouseEvent)
+{
+	if ( SelectedWidget.IsValid() )
+	{
+		return FReply::Handled().BeginDragDrop(FSelectedWidgetDragDropOp::New(SelectedWidget));
+	}
+
+	return FReply::Unhandled();
+}
+
+void SDesignerView::OnDragEnter(const FGeometry& MyGeometry, const FDragDropEvent& DragDropEvent)
+{
+	//@TODO UMG Drop Feedback
+}
+
+void SDesignerView::OnDragLeave(const FDragDropEvent& DragDropEvent)
+{
+	TSharedPtr<FWidgetTemplateDragDropOp> DragDropOp = DragDropEvent.GetOperationAs<FWidgetTemplateDragDropOp>();
+	if ( DragDropOp.IsValid() )
+	{
+		DragDropOp->ResetToDefaultToolTip();
+	}
+	
+	if (DropPreviewWidget)
+	{
+		if ( DropPreviewParent )
+		{
+			DropPreviewParent->RemoveChild(DropPreviewWidget);
+		}
+
+		UWidgetBlueprint* BP = GetBlueprint();
+		BP->WidgetTree->RemoveWidget(DropPreviewWidget);
+		DropPreviewWidget = NULL;
+	}
+}
+
+FReply SDesignerView::OnDragOver(const FGeometry& MyGeometry, const FDragDropEvent& DragDropEvent)
+{
+	UWidgetBlueprint* BP = GetBlueprint();
+	
+	if (DropPreviewWidget)
+	{
+		if (DropPreviewParent)
+		{
+			DropPreviewParent->RemoveChild(DropPreviewWidget);
+		}
+		
+		BP->WidgetTree->RemoveWidget(DropPreviewWidget);
+		DropPreviewWidget = NULL;
+	}
+	
+	const bool bIsPreview = true;
+	DropPreviewWidget = ProcessDropAndAddWidget(MyGeometry, DragDropEvent, bIsPreview);
+	if ( DropPreviewWidget )
+	{
+		//@TODO UMG Drop Feedback
+		return FReply::Handled();
+	}
+
+	return FReply::Unhandled();
+}
+
+UWidget* SDesignerView::ProcessDropAndAddWidget(const FGeometry& MyGeometry, const FDragDropEvent& DragDropEvent, bool bIsPreview)
+{
+	// In order to prevent the GetWidgetAtCursor code from picking the widget we're about to move, we need to mark it
+	// as the drop preview widget before any other code can run.
+	TSharedPtr<FSelectedWidgetDragDropOp> SelectedDragDropOp = DragDropEvent.GetOperationAs<FSelectedWidgetDragDropOp>();
+	if ( SelectedDragDropOp.IsValid() )
+	{
+		DropPreviewWidget = SelectedDragDropOp->Widget.GetPreview();
+	}
+
+	UWidgetBlueprint* BP = GetBlueprint();
+
+	if ( DropPreviewWidget )
+	{
+		if ( DropPreviewParent )
+		{
+			DropPreviewParent->RemoveChild(DropPreviewWidget);
+		}
+
+		BP->WidgetTree->RemoveWidget(DropPreviewWidget);
+		DropPreviewWidget = NULL;
+	}
+
+	FArrangedWidget ArrangedWidget(SNullWidget::NullWidget, FGeometry());
+	FWidgetReference WidgetUnderCursor = GetWidgetAtCursor(MyGeometry, DragDropEvent, ArrangedWidget);
+
+	UWidget* Target = NULL;
+	if ( WidgetUnderCursor.IsValid() )
+	{
+		Target = bIsPreview ? WidgetUnderCursor.GetPreview() : WidgetUnderCursor.GetTemplate();
+	}
+
+	TSharedPtr<FWidgetTemplateDragDropOp> TemplateDragDropOp = DragDropEvent.GetOperationAs<FWidgetTemplateDragDropOp>();
+	if ( TemplateDragDropOp.IsValid() )
+	{
+		// If there's no root widget go ahead and add the widget into the root slot.
+		if ( BP->WidgetTree->RootWidget == NULL )
+		{
+			FScopedTransaction Transaction(LOCTEXT("Designer_AddWidget", "Add Widget"));
+
+			if ( !bIsPreview )
+			{
+				BP->WidgetTree->SetFlags(RF_Transactional);
+				BP->WidgetTree->Modify();
+			}
+
+			// TODO UMG This method isn't great, maybe the user widget should just be a canvas.
+
+			// Add it to the root if there are no other widgets to add it to.
+			UWidget* Widget = TemplateDragDropOp->Template->Create(BP->WidgetTree);
+			Widget->IsDesignTime(true);
+
+			BP->WidgetTree->RootWidget = Widget;
+
+			SelectedWidget = FWidgetReference::FromTemplate(BlueprintEditor.Pin(), Widget);
+
+			DropPreviewParent = NULL;
+
+			if ( bIsPreview )
+			{
+				Transaction.Cancel();
+			}
+
+			return Widget;
+		}
+		// If there's already a root widget we need to try and place our widget into a parent widget that we've picked against
+		else if ( Target && Target->IsA(UPanelWidget::StaticClass()) )
+		{
+			UPanelWidget* Parent = Cast<UPanelWidget>(Target);
+
+			FScopedTransaction Transaction(LOCTEXT("Designer_AddWidget", "Add Widget"));
+
+			// If this isn't a preview operation we need to modify a few things to properly undo the operation.
+			if ( !bIsPreview )
+			{
+				Parent->SetFlags(RF_Transactional);
+				Parent->Modify();
+
+				BP->WidgetTree->SetFlags(RF_Transactional);
+				BP->WidgetTree->Modify();
+			}
+
+			// Construct the widget and mark it for design time rendering.
+			UWidget* Widget = TemplateDragDropOp->Template->Create(BP->WidgetTree);
+			Widget->IsDesignTime(true);
+
+			// Determine local position inside the parent widget and add the widget to the slot.
+			FVector2D LocalPosition = ArrangedWidget.Geometry.AbsoluteToLocal(DragDropEvent.GetScreenSpacePosition());
+			if ( UPanelSlot* Slot = Parent->AddChild(Widget) )
+			{
+				TSharedPtr<SWidget> SlateWidget = Widget->GetWidget();
+				SlateWidget->SlatePrepass();
+				const FVector2D& WidgetDesiredSize = SlateWidget->GetDesiredSize();
+
+				static const FVector2D MinimumDefaultSize(20, 20);
+				FVector2D LocalSize = FVector2D(FMath::Max(WidgetDesiredSize.X, MinimumDefaultSize.X), FMath::Max(WidgetDesiredSize.Y, MinimumDefaultSize.Y));
+
+				Slot->SetDesiredPosition(LocalPosition);
+				Slot->SetDesiredSize(LocalSize);
+
+				DropPreviewParent = Parent;
+
+				if ( bIsPreview )
+				{
+					Transaction.Cancel();
+				}
+
+				return Widget;
+			}
+			else
+			{
+				// TODO UMG ERROR Slot can not be created because maybe the max children has been reached.
+				//          Maybe we can traverse the hierarchy and add it to the first parent that will accept it?
+			}
+
+			if ( bIsPreview )
+			{
+				Transaction.Cancel();
+			}
+		}
+	}
+
+	// Attempt to deal with moving widgets from a drag operation.
+	if ( SelectedDragDropOp.IsValid() )
+	{
+		if ( Target && Target->IsA(UPanelWidget::StaticClass()) )
+		{
+			UPanelWidget* NewParent = Cast<UPanelWidget>(Target);
+
+			FScopedTransaction Transaction(LOCTEXT("Designer_MoveWidget", "Move Widget"));
+
+			// If this isn't a preview operation we need to modify a few things to properly undo the operation.
+			if ( !bIsPreview )
+			{
+				NewParent->SetFlags(RF_Transactional);
+				NewParent->Modify();
+
+				BP->WidgetTree->SetFlags(RF_Transactional);
+				BP->WidgetTree->Modify();
+			}
+
+			UWidget* Widget = bIsPreview ? SelectedDragDropOp->Widget.GetPreview() : SelectedDragDropOp->Widget.GetTemplate();
+
+			if ( Widget->GetParent() )
+			{
+				if ( !bIsPreview )
+				{
+					Widget->GetParent()->Modify();
+				}
+
+				Widget->GetParent()->RemoveChild(Widget);
+			}
+
+			FVector2D LocalPosition = ArrangedWidget.Geometry.AbsoluteToLocal(DragDropEvent.GetScreenSpacePosition());
+			if ( UPanelSlot* Slot = NewParent->AddChild(Widget) )
+			{
+				FWidgetBlueprintEditorUtils::ImportPropertiesFromText(Slot, SelectedDragDropOp->ExportedSlotProperties);
+
+				//TODO UMG Migrate existing slot info
+				Slot->SetDesiredPosition(LocalPosition - SelectedWidgetContextMenuLocation);
+				//Slot->SetDesiredSize(FVector2D(150, 30));
+				//@TODO UMG When we add a child blindly we need to default the slot size to the preferred size of the widget if the container supports such things.
+				//@TODO UMG We may need a desired size canvas, where the slots have no size, they only give you position, alternatively, maybe slots that don't clip, so center is still easy.
+
+				DropPreviewParent = NewParent;
+
+				if ( bIsPreview )
+				{
+					Transaction.Cancel();
+				}
+
+				return Widget;
+			}
+			else
+			{
+				// TODO UMG ERROR Slot can not be created because maybe the max children has been reached.
+				//          Maybe we can traverse the hierarchy and add it to the first parent that will accept it?
+			}
+
+			if ( bIsPreview )
+			{
+				Transaction.Cancel();
+			}
+		}
+	}
+	
+	return NULL;
+}
+
+FReply SDesignerView::OnDrop(const FGeometry& MyGeometry, const FDragDropEvent& DragDropEvent)
+{
+	bMouseDown = false;
+	bMovingExistingWidget = false;
+
+	UWidgetBlueprint* BP = GetBlueprint();
+	
+	if (DropPreviewWidget)
+	{
+		if (DropPreviewParent)
+		{
+			DropPreviewParent->RemoveChild(DropPreviewWidget);
+		}
+		
+		BP->WidgetTree->RemoveWidget(DropPreviewWidget);
+		DropPreviewWidget = NULL;
+	}
+	
+	const bool bIsPreview = false;
+	UWidget* Widget = ProcessDropAndAddWidget(MyGeometry, DragDropEvent, bIsPreview);
+	if ( Widget )
+	{
+		FBlueprintEditorUtils::MarkBlueprintAsStructurallyModified(BP);
+
+		return FReply::Handled();
+	}
+	
+	return FReply::Unhandled();
+}
+
+void SDesignerView::BeginTransaction(const FText& SessionName)
+{
+	if ( ScopedTransaction == NULL )
+	{
+		ScopedTransaction = new FScopedTransaction(SessionName);
+	}
+
+	if ( SelectedWidget.IsValid() )
+	{
+		SelectedWidget.GetPreview()->Modify();
+		SelectedWidget.GetTemplate()->Modify();
+	}
+}
+
+void SDesignerView::EndTransaction()
+{
+	if ( ScopedTransaction != NULL )
+	{
+		delete ScopedTransaction;
+		ScopedTransaction = NULL;
+	}
+}
+
+#undef LOCTEXT_NAMESPACE
diff --git a/Engine/Source/Editor/UMGEditor/Private/Designer/SDesignerView.h b/Engine/Source/Editor/UMGEditor/Private/Designer/SDesignerView.h
new file mode 100644
--- /dev/null
+++ b/Engine/Source/Editor/UMGEditor/Private/Designer/SDesignerView.h
@@ -0,0 +1,160 @@
+// Copyright 1998-2014 Epic Games, Inc. All Rights Reserved.
+
+#pragma once
+
+#include "SCompoundWidget.h"
+#include "SNodePanel.h"
+#include "SDesignSurface.h"
+#include "DesignerExtension.h"
+#include "IUMGDesigner.h"
+
+class FDesignerExtension;
+
+/**
+ * The designer for widgets.  Allows for laying out widgets in a drag and drop environment.
+ */
+class SDesignerView : public SDesignSurface, public IUMGDesigner
+{
+public:
+
+	SLATE_BEGIN_ARGS( SDesignerView ) {}
+	SLATE_END_ARGS()
+
+	void Construct(const FArguments& InArgs, TSharedPtr<class FWidgetBlueprintEditor> InBlueprintEditor);
+	virtual ~SDesignerView();
+
+	// SWidget interface
+	virtual FReply OnMouseButtonDown(const FGeometry& MyGeometry, const FPointerEvent& MouseEvent) override;
+	virtual FReply OnMouseButtonUp(const FGeometry& MyGeometry, const FPointerEvent& MouseEvent) override;
+	virtual FReply OnMouseMove(const FGeometry& MyGeometry, const FPointerEvent& MouseEvent) override;
+	virtual void OnMouseEnter(const FGeometry& MyGeometry, const FPointerEvent& MouseEvent) override;
+	virtual void OnMouseLeave(const FPointerEvent& MouseEvent) override;
+
+	virtual FReply OnKeyDown(const FGeometry& MyGeometry, const FKeyboardEvent& InKeyboardEvent) override;
+
+	virtual FCursorReply OnCursorQuery(const FGeometry& MyGeometry, const FPointerEvent& CursorEvent) const override;
+
+	virtual int32 OnPaint(const FGeometry& AllottedGeometry, const FSlateRect& MyClippingRect, FSlateWindowElementList& OutDrawElements, int32 LayerId, const FWidgetStyle& InWidgetStyle, bool bParentEnabled) const override;
+	virtual void Tick(const FGeometry& AllottedGeometry, const double InCurrentTime, const float InDeltaTime) override;
+
+	virtual FReply OnDragDetected(const FGeometry& MyGeometry, const FPointerEvent& MouseEvent) override;
+	virtual void OnDragEnter(const FGeometry& MyGeometry, const FDragDropEvent& DragDropEvent) override;
+	virtual void OnDragLeave(const FDragDropEvent& DragDropEvent) override;
+	virtual FReply OnDragOver(const FGeometry& MyGeometry, const FDragDropEvent& DragDropEvent) override;
+	virtual FReply OnDrop(const FGeometry& MyGeometry, const FDragDropEvent& DragDropEvent) override;
+	// End of SWidget interface
+
+	void Register(TSharedRef<FDesignerExtension> Extension);
+
+private:
+	void CreateExtensionWidgetsForSelection();
+
+	/** Displays the context menu when you right click */
+	void ShowContextMenu(const FGeometry& MyGeometry, const FPointerEvent& MouseEvent);
+
+	void OnEditorSelectionChanged();
+
+	/** @returns Gets the widget under the cursor based on a mouse pointer event. */
+	FWidgetReference GetWidgetAtCursor(const FGeometry& MyGeometry, const FPointerEvent& MouseEvent, FArrangedWidget& ArrangedWidget);
+
+	/** Gets the blueprint being edited by the designer */
+	UWidgetBlueprint* GetBlueprint() const;
+
+	/** Called whenever the blueprint is structurally changed. */
+	void OnBlueprintChanged(UBlueprint* InBlueprint);
+	void OnObjectPropertyChanged(UObject* ObjectBeingModified, FPropertyChangedEvent& PropertyChangedEvent);
+
+	bool GetArrangedWidget(TSharedRef<SWidget> Widget, FArrangedWidget& ArrangedWidget) const;
+	bool GetArrangedWidgetRelativeToWindow(TSharedRef<SWidget> Widget, FArrangedWidget& ArrangedWidget) const;
+	bool GetArrangedWidgetRelativeToDesigner(TSharedRef<SWidget> Widget, FArrangedWidget& ArrangedWidget) const;
+
+	FVector2D GetSelectionDesignerWidgetsLocation() const;
+	FVector2D GetCachedSelectionDesignerWidgetsLocation() const;
+
+	void BeginTransaction(const FText& SessionName);
+	void EndTransaction();
+
+private:
+	enum DragHandle
+	{
+		DH_NONE = -1,
+
+		DH_TOP_LEFT = 0,
+		DH_TOP_CENTER,
+		DH_TOP_RIGHT,
+
+		DH_MIDDLE_LEFT,
+		DH_MIDDLE_RIGHT,
+
+		DH_BOTTOM_LEFT,
+		DH_BOTTOM_CENTER,
+		DH_BOTTOM_RIGHT,
+		
+		DH_MAX,
+	};
+
+	TArray< FVector2D > DragDirections;
+
+	/** Extensions for the designer to allow for custom widgets to be inserted onto the design surface as selection changes. */
+	TArray< TSharedRef<FDesignerExtension> > DesignerExtensions;
+
+	/** Temporary widgets that are destroyed every time selection changes and are built from the current selection by extensions. */
+	TArray< TSharedRef<SWidget> > ExtensionWidgets;
+
+private:
+	void DrawDragHandles(const FPaintGeometry& SelectionGeometry, const FSlateRect& MyClippingRect, FSlateWindowElementList& OutDrawElements, int32 LayerId, const FWidgetStyle& InWidgetStyle) const;
+	DragHandle HitTestDragHandles(const FGeometry& AllottedGeometry, const FPointerEvent& PointerEvent) const;
+
+	UWidget* ProcessDropAndAddWidget(const FGeometry& MyGeometry, const FDragDropEvent& DragDropEvent, bool bIsPreview);
+	
+private:
+	/** A reference to the BP Editor that owns this designer */
+	TWeakPtr<FWidgetBlueprintEditor> BlueprintEditor;
+
+	/** The transaction used to commit undoable actions from resize, move...etc */
+	FScopedTransaction* ScopedTransaction;
+
+	/** The current preview widget */
+	UUserWidget* PreviewWidget;
+
+	/** The current preview widget's slate widget */
+	TWeakPtr<SWidget> PreviewSlateWidget;
+	
+	UWidget* DropPreviewWidget;
+	UPanelWidget* DropPreviewParent;
+
+	TSharedPtr<SBorder> PreviewSurface;
+	TSharedPtr<SCanvas> ExtensionWidgetCanvas;
+
+	DragHandle CurrentHandle;
+
+	FVector2D CachedDesignerWidgetLocation;
+
+	/** The currently selected set of widgets */
+	TSet< FWidgetReference > SelectedWidgets;
+
+	/** TODO UMG Remove, after getting multiselection working. */
+	FWidgetReference SelectedWidget;
+
+	/** The location in selected widget local space where the context menu was summoned. */
+	FVector2D SelectedWidgetContextMenuLocation;
+
+	/** The currently selected slate widget, this is refreshed every frame in case it changes */
+	TArray< TWeakPtr<SWidget> > SelectedSlateWidgets;
+	TWeakPtr<SWidget> SelectedSlateWidget;
+
+	/**  */
+	bool bMouseDown;
+
+	/** An existing widget is being moved in its current container, or in to a new container. */
+	bool bMovingExistingWidget;
+
+	/** The wall clock time the user has been hovering over a single widget */
+	float HoverTime;
+
+	/** The current widget being hovered */
+	FWidgetReference HoveredWidget;
+
+	/** The current slate widget being hovered, this is refreshed every frame in case it changes */
+	TWeakPtr<SWidget> HoveredSlateWidget;
+};
diff --git a/Engine/Source/Editor/UMGEditor/Private/DragDrop/WidgetTemplateDragDropOp.cpp b/Engine/Source/Editor/UMGEditor/Private/DragDrop/WidgetTemplateDragDropOp.cpp
new file mode 100644
--- /dev/null
+++ b/Engine/Source/Editor/UMGEditor/Private/DragDrop/WidgetTemplateDragDropOp.cpp
@@ -0,0 +1,20 @@
+// Copyright 1998-2014 Epic Games, Inc. All Rights Reserved.
+
+#include "UMGEditorPrivatePCH.h"
+
+#include "WidgetTemplateDragDropOp.h"
+
+#define LOCTEXT_NAMESPACE "UMG"
+
+TSharedRef<FWidgetTemplateDragDropOp> FWidgetTemplateDragDropOp::New(const TSharedPtr<FWidgetTemplate>& InTemplate)
+{
+	TSharedRef<FWidgetTemplateDragDropOp> Operation = MakeShareable(new FWidgetTemplateDragDropOp());
+	Operation->Template = InTemplate;
+	Operation->DefaultHoverText = InTemplate->Name;
+	Operation->CurrentHoverText = InTemplate->Name;
+	Operation->Construct();
+
+	return Operation;
+}
+
+#undef LOCTEXT_NAMESPACE
diff --git a/Engine/Source/Editor/UMGEditor/Private/DragDrop/WidgetTemplateDragDropOp.h b/Engine/Source/Editor/UMGEditor/Private/DragDrop/WidgetTemplateDragDropOp.h
new file mode 100644
--- /dev/null
+++ b/Engine/Source/Editor/UMGEditor/Private/DragDrop/WidgetTemplateDragDropOp.h
@@ -0,0 +1,22 @@
+// Copyright 1998-2014 Epic Games, Inc. All Rights Reserved.
+
+#pragma once
+
+#include "DecoratedDragDropOp.h"
+#include "WidgetTemplate.h"
+
+/**
+ * This drag drop operation allows widget templates from the palate to be dragged and dropped into the designer
+ * or the widget hierarchy in order to spawn new widgets.
+ */
+class FWidgetTemplateDragDropOp : public FDecoratedDragDropOp
+{
+public:
+	DRAG_DROP_OPERATOR_TYPE(FWidgetTemplateDragDropOp, FDecoratedDragDropOp)
+
+	/** The template to create an instance */
+	TSharedPtr<FWidgetTemplate> Template;
+
+	/** Constructs the drag drop operation */
+	static TSharedRef<FWidgetTemplateDragDropOp> New(const TSharedPtr<FWidgetTemplate>& InTemplate);
+};
diff --git a/Engine/Source/Editor/UMGEditor/Private/Hierarchy/SHierarchyView.cpp b/Engine/Source/Editor/UMGEditor/Private/Hierarchy/SHierarchyView.cpp
new file mode 100644
--- /dev/null
+++ b/Engine/Source/Editor/UMGEditor/Private/Hierarchy/SHierarchyView.cpp
@@ -0,0 +1,342 @@
+// Copyright 1998-2014 Epic Games, Inc. All Rights Reserved.
+
+#include "UMGEditorPrivatePCH.h"
+#include "SHierarchyView.h"
+#include "SHierarchyViewItem.h"
+
+#include "UMGEditorActions.h"
+
+#include "PreviewScene.h"
+#include "SceneViewport.h"
+
+#include "BlueprintEditor.h"
+#include "SKismetInspector.h"
+#include "BlueprintEditorUtils.h"
+#include "WidgetTemplateClass.h"
+#include "WidgetBlueprintEditor.h"
+
+#define LOCTEXT_NAMESPACE "UMG"
+
+void SHierarchyView::Construct(const FArguments& InArgs, TSharedPtr<FWidgetBlueprintEditor> InBlueprintEditor, USimpleConstructionScript* InSCS)
+{
+	BlueprintEditor = InBlueprintEditor;
+	bRefreshRequested = false;
+	bIsFilterActive = false;
+
+	SearchBoxWidgetFilter = MakeShareable(new WidgetTextFilter(WidgetTextFilter::FItemToStringArray::CreateSP(this, &SHierarchyView::TransformWidgetToString)));
+
+	UWidgetBlueprint* Blueprint = GetBlueprint();
+	Blueprint->OnChanged().AddSP(this, &SHierarchyView::OnBlueprintChanged);
+
+	SAssignNew(WidgetTreeView, STreeView< UWidget* >)
+	.ItemHeight(20.0f)
+	.SelectionMode(ESelectionMode::Single)
+	.OnGetChildren(this, &SHierarchyView::WidgetHierarchy_OnGetChildren)
+	.OnGenerateRow(this, &SHierarchyView::WidgetHierarchy_OnGenerateRow)
+	.OnSelectionChanged(this, &SHierarchyView::WidgetHierarchy_OnSelectionChanged)
+	.OnContextMenuOpening(this, &SHierarchyView::WidgetHierarchy_OnContextMenuOpening)
+	.TreeItemsSource(&RootWidgets);
+
+	ChildSlot
+	[
+		SNew(SBorder)
+		.BorderImage(FEditorStyle::GetBrush("ToolPanel.GroupBorder"))
+		[
+			SNew(SVerticalBox)
+
+			+ SVerticalBox::Slot()
+			.Padding(4)
+			.AutoHeight()
+			[
+				SNew(SSearchBox)
+				.HintText(LOCTEXT("SearchWidgets", "Search Widgets"))
+				.OnTextChanged(this, &SHierarchyView::OnSearchChanged)
+			]
+
+			+ SVerticalBox::Slot()
+			.FillHeight(1.0f)
+			[
+				SNew(SScrollBorder, WidgetTreeView.ToSharedRef())
+				[
+					WidgetTreeView.ToSharedRef()
+				]
+			]
+		]
+	];
+
+	BlueprintEditor.Pin()->OnSelectedWidgetsChanged.AddRaw(this, &SHierarchyView::OnEditorSelectionChanged);
+
+	bRefreshRequested = true;
+}
+
+SHierarchyView::~SHierarchyView()
+{
+	UWidgetBlueprint* Blueprint = GetBlueprint();
+	if ( Blueprint )
+	{
+		Blueprint->OnChanged().RemoveAll(this);
+	}
+
+	if ( BlueprintEditor.IsValid() )
+	{
+		BlueprintEditor.Pin()->OnSelectedWidgetsChanged.RemoveAll(this);
+	}
+}
+
+void SHierarchyView::Tick(const FGeometry& AllottedGeometry, const double InCurrentTime, const float InDeltaTime)
+{
+	if ( bRefreshRequested )
+	{
+		bRefreshRequested = false;
+
+		RefreshTree();
+	}
+}
+
+FReply SHierarchyView::OnKeyDown(const FGeometry& MyGeometry, const FKeyboardEvent& InKeyboardEvent)
+{
+	BlueprintEditor.Pin()->PasteDropLocation = FVector2D(0, 0);
+
+	if ( BlueprintEditor.Pin()->WidgetCommandList->ProcessCommandBindings(InKeyboardEvent) )
+	{
+		return FReply::Handled();
+	}
+
+	return FReply::Unhandled();
+}
+
+void SHierarchyView::TransformWidgetToString(const UWidget* Widget, OUT TArray< FString >& Array)
+{
+	Array.Add( Widget->GetLabel() );
+}
+
+void SHierarchyView::OnSearchChanged(const FText& InFilterText)
+{
+	bRefreshRequested = true;
+	SearchBoxWidgetFilter->SetRawFilterText(InFilterText);
+}
+
+FText SHierarchyView::GetSearchText() const
+{
+	return SearchBoxWidgetFilter->GetRawFilterText();
+}
+
+void SHierarchyView::OnEditorSelectionChanged()
+{
+	WidgetTreeView->ClearSelection();
+
+	bool bFirst = true;
+
+	// Update the selection and expansion in the tree to match the new selection
+	const TSet<FWidgetReference>& SelectedWidgets = BlueprintEditor.Pin()->GetSelectedWidgets();
+	for ( const FWidgetReference& WidgetRef : SelectedWidgets )
+	{
+		UWidget* TemplateWidget = WidgetRef.GetTemplate();
+
+		if ( TemplateWidget )
+		{
+			WidgetTreeView->SetItemSelection(TemplateWidget, true);
+
+			// Attempt to scroll the first widget we find into view.
+			if ( bFirst )
+			{
+				bFirst = false;
+				WidgetTreeView->RequestScrollIntoView(TemplateWidget);
+			}
+
+			ExpandPathToWidget(TemplateWidget);
+		}
+	}
+}
+
+void SHierarchyView::ExpandPathToWidget(UWidget* TemplateWidget)
+{
+	// Expand the path leading to this widget in the tree.
+	UWidget* Parent = TemplateWidget->GetParent();
+	while ( Parent != NULL )
+	{
+		WidgetTreeView->SetItemExpansion(Parent, true);
+		Parent = Parent->GetParent();
+	}
+}
+
+UWidgetBlueprint* SHierarchyView::GetBlueprint() const
+{
+	if ( BlueprintEditor.IsValid() )
+	{
+		UBlueprint* BP = BlueprintEditor.Pin()->GetBlueprintObj();
+		return CastChecked<UWidgetBlueprint>(BP);
+	}
+
+	return NULL;
+}
+
+void SHierarchyView::OnBlueprintChanged(UBlueprint* InBlueprint)
+{
+	if ( InBlueprint )
+	{
+		RefreshTree();
+	}
+}
+
+void SHierarchyView::ShowDetailsForObjects(TArray<UWidget*> TemplateWidgets)
+{
+	TSet<FWidgetReference> SelectedWidgets;
+	for ( UWidget* TemplateWidget : TemplateWidgets )
+	{
+		FWidgetReference Selection = FWidgetReference::FromTemplate(BlueprintEditor.Pin(), TemplateWidget);
+		SelectedWidgets.Add(Selection);
+	}
+
+	BlueprintEditor.Pin()->SelectWidgets(SelectedWidgets);
+}
+
+TSharedPtr<SWidget> SHierarchyView::WidgetHierarchy_OnContextMenuOpening()
+{
+	FMenuBuilder MenuBuilder(true, NULL);
+
+	FWidgetBlueprintEditorUtils::CreateWidgetContextMenu(MenuBuilder, BlueprintEditor.Pin().ToSharedRef(), FVector2D(0, 0));
+
+	return MenuBuilder.MakeWidget();
+}
+
+void SHierarchyView::WidgetHierarchy_OnGetChildren(UWidget* InParent, TArray< UWidget* >& OutChildren)
+{
+	UPanelWidget* Widget = Cast<UPanelWidget>(InParent);
+	if ( Widget )
+	{
+		for ( int32 i = 0; i < Widget->GetChildrenCount(); i++ )
+		{
+			UWidget* Child = Widget->GetChildAt(i);
+			if ( Child )
+			{
+				if ( bIsFilterActive && !WidgetsPassingFilter.Contains(Child) )
+				{
+					continue;
+				}
+
+				OutChildren.Add(Child);
+			}
+		}
+	}
+}
+
+TSharedRef< ITableRow > SHierarchyView::WidgetHierarchy_OnGenerateRow(UWidget* InItem, const TSharedRef<STableViewBase>& OwnerTable)
+{
+	return SNew(SHierarchyViewItem, OwnerTable, BlueprintEditor.Pin(), InItem)
+		.HighlightText(this, &SHierarchyView::GetSearchText);
+}
+
+void SHierarchyView::WidgetHierarchy_OnSelectionChanged(UWidget* SelectedItem, ESelectInfo::Type SelectInfo)
+{
+	if ( SelectInfo != ESelectInfo::Direct )
+	{
+		TArray<UWidget*> Items;
+		Items.Add(SelectedItem);
+		ShowDetailsForObjects(Items);
+	}
+}
+
+FReply SHierarchyView::HandleDeleteSelected()
+{
+	TSet<FWidgetReference> SelectedWidgets = BlueprintEditor.Pin()->GetSelectedWidgets();
+	//TArray<UWidget*> SelectedWidgets = WidgetTreeView->GetSelectedItems();
+
+	// Remove the selected items from the widget cache
+	for ( FWidgetReference& Item : SelectedWidgets )
+	{
+		CachedExpandedWidgets.Remove(Item.GetTemplate());
+	}
+
+	FWidgetBlueprintEditorUtils::DeleteWidgets(GetBlueprint(), SelectedWidgets);
+
+	return FReply::Handled();
+}
+
+bool SHierarchyView::FilterWidgetHierarchy(UWidget* CurrentWidget)
+{
+	bool bAnyChildrenPass = false;
+
+	// Iterate over children and check to see if any of them pass the current filter
+	UPanelWidget* Widget = Cast<UPanelWidget>(CurrentWidget);
+	if ( Widget )
+	{
+		for ( int32 i = 0; i < Widget->GetChildrenCount(); i++ )
+		{
+			UWidget* Child = Widget->GetChildAt(i);
+			if ( Child )
+			{
+				bAnyChildrenPass |= FilterWidgetHierarchy(Child);
+			}
+		}
+	}
+
+	// Check to see if I pass the filter
+	const bool bWidgetPass = SearchBoxWidgetFilter->PassesFilter(CurrentWidget);
+
+	// If this particular widget passes the filter, expand every widget leading up to it.
+	if ( bWidgetPass )
+	{
+		ExpandPathToWidget(CurrentWidget);
+	}
+
+	// If either the widget or the children pass, add the current widget
+	if ( bWidgetPass || bAnyChildrenPass )
+	{
+		WidgetsPassingFilter.Add(CurrentWidget);
+		return true;
+	}
+
+	return false;
+}
+
+void SHierarchyView::RefreshTree()
+{
+	RootWidgets.Reset();
+
+	UWidgetBlueprint* Blueprint = GetBlueprint();
+
+	if ( Blueprint->WidgetTree->RootWidget )
+	{
+		bool bRootPassed = true;
+
+		const bool bWillFilterBeActive = !SearchBoxWidgetFilter->GetRawFilterText().IsEmpty();
+
+		// Save the expansion state when the filter becomes active
+		if ( !bIsFilterActive && bWillFilterBeActive )
+		{
+			CachedExpandedWidgets.Empty();
+			WidgetTreeView->GetExpandedItems(CachedExpandedWidgets);
+		}
+		// Restore the expansion state when the filter is removed
+		else if ( bIsFilterActive && !bWillFilterBeActive )
+		{
+			WidgetTreeView->ClearExpandedItems();
+			for ( UWidget* ExpandedWidget : CachedExpandedWidgets )
+			{
+				WidgetTreeView->SetItemExpansion(ExpandedWidget, true);
+			}
+			CachedExpandedWidgets.Empty();
+		}
+		
+		bIsFilterActive = bWillFilterBeActive;
+
+		WidgetsPassingFilter.Empty();
+		if ( bIsFilterActive )
+		{
+			bRootPassed = FilterWidgetHierarchy(Blueprint->WidgetTree->RootWidget);
+		}
+
+		if ( bRootPassed )
+		{
+			RootWidgets.Add(Blueprint->WidgetTree->RootWidget);
+		}
+	}
+
+	WidgetTreeView->RequestTreeRefresh();
+}
+
+
+//@TODO UMG Drop widgets onto the tree, when nothing is present, if there is a root node present, what happens then, let the root node attempt to place it?
+
+#undef LOCTEXT_NAMESPACE
diff --git a/Engine/Source/Editor/UMGEditor/Private/Hierarchy/SHierarchyView.h b/Engine/Source/Editor/UMGEditor/Private/Hierarchy/SHierarchyView.h
new file mode 100644
--- /dev/null
+++ b/Engine/Source/Editor/UMGEditor/Private/Hierarchy/SHierarchyView.h
@@ -0,0 +1,108 @@
+// Copyright 1998-2014 Epic Games, Inc. All Rights Reserved.
+
+#pragma once
+
+#include "SCompoundWidget.h"
+#include "BlueprintEditor.h"
+
+//TODO rename SUMGEditorHierarchy
+
+/**
+ * The tree view presenting the widget hierarchy.  This allows users to edit the hierarchy of widgets easily by dragging and 
+ * dropping them logically, which in some cases may be significantly easier than doing it visually in the widget designer.
+ */
+class SHierarchyView : public SCompoundWidget
+{
+public:
+	typedef TTextFilter< const UWidget* > WidgetTextFilter;
+
+public:
+	SLATE_BEGIN_ARGS( SHierarchyView ){}
+	SLATE_END_ARGS()
+
+	void Construct(const FArguments& InArgs, TSharedPtr<FWidgetBlueprintEditor> InBlueprintEditor, USimpleConstructionScript* InSCS);
+
+	virtual ~SHierarchyView();
+
+	// Begin SWidget
+	virtual void Tick(const FGeometry& AllottedGeometry, const double InCurrentTime, const float InDeltaTime);
+
+	virtual FReply OnKeyDown(const FGeometry& MyGeometry, const FKeyboardEvent& InKeyboardEvent) override;
+	// End SWidget
+
+private:
+	void BuildWrapWithMenu(FMenuBuilder& Menu);
+	TSharedPtr<SWidget> WidgetHierarchy_OnContextMenuOpening();
+	void WidgetHierarchy_OnGetChildren(UWidget* InParent, TArray< UWidget* >& OutChildren);
+	TSharedRef< ITableRow > WidgetHierarchy_OnGenerateRow(UWidget* InItem, const TSharedRef<STableViewBase>& OwnerTable);
+	void WidgetHierarchy_OnSelectionChanged(UWidget* SelectedItem, ESelectInfo::Type SelectInfo);
+
+private:
+	/** @returns the current blueprint being edited */
+	UWidgetBlueprint* GetBlueprint() const;
+
+	/** Expands every item in the tree leading to this widget */
+	void ExpandPathToWidget(UWidget* TemplateWidget);
+
+	/** Called when the blueprint is structurally changed. */
+	void OnBlueprintChanged(UBlueprint* InBlueprint);
+
+	/** Called when the selected widget has changed.  The treeview then needs to match the new selection. */
+	void OnEditorSelectionChanged();
+
+	/** Notifies the details panel that new widgets have been selected. */
+	void ShowDetailsForObjects(TArray<UWidget*> Widgets);
+
+	/** Rebuilds the tree structure based on the current filter options */
+	void RefreshTree();
+
+	/** Handles the menu clicking to delete the selected widgets. */
+	FReply HandleDeleteSelected();
+
+	/** Deletes the selected widgets from the hierarchy */
+	void DeleteSelected();
+
+	/** Wraps the selected widget using a user selected panel */
+	void WrapSelectedWidgets(UClass* WidgetClass);
+	
+	/** Called when the filter text is changed. */
+	void OnSearchChanged(const FText& InFilterText);
+
+	/** Gets the search text currently being used to filter the list, also used to highlight text */
+	FText GetSearchText() const;
+
+	/** Transforms the widget into a searchable string */
+	void TransformWidgetToString(const UWidget* Widget, OUT TArray< FString >& Array);
+
+	/** Builds a filterable set by recursing through the widget hierarchy */
+	bool FilterWidgetHierarchy(UWidget* CurrentWidget);
+
+private:
+
+	/** Cached pointer to the blueprint editor that owns this tree. */
+	TWeakPtr<class FWidgetBlueprintEditor> BlueprintEditor;
+
+	/** The tree item source.  We only ever have one root widget, but the tree expects an array. */
+	TArray< UWidget* > RootWidgets;
+
+	/** The widget hierarchy slate treeview widget */
+	TSharedPtr< STreeView< UWidget* > > WidgetTreeView;
+
+	/** The filter used by the search box */
+	TSharedPtr<WidgetTextFilter> SearchBoxWidgetFilter;
+
+	/** The set of widgets currently passing the filter.  Used to quickly filter children when the tree is expanded lazily. */
+	TSet< UWidget* > WidgetsPassingFilter;
+
+	/**
+	  * A cache of the expanded items.  We store this off before a user filters the list so that we can 
+	  * restore the expansion state after the filter is cleared.
+	  */
+	TSet< UWidget* > CachedExpandedWidgets;
+
+	/** Caches a value to know if a filter is currently active. */
+	bool bIsFilterActive;
+
+	/** Has a full refresh of the tree been requested?  This happens when the user is filtering the tree */
+	bool bRefreshRequested;
+};
diff --git a/Engine/Source/Editor/UMGEditor/Private/Hierarchy/SHierarchyViewItem.cpp b/Engine/Source/Editor/UMGEditor/Private/Hierarchy/SHierarchyViewItem.cpp
new file mode 100644
--- /dev/null
+++ b/Engine/Source/Editor/UMGEditor/Private/Hierarchy/SHierarchyViewItem.cpp
@@ -0,0 +1,165 @@
+// Copyright 1998-2014 Epic Games, Inc. All Rights Reserved.
+
+#include "UMGEditorPrivatePCH.h"
+#include "SHierarchyViewItem.h"
+
+#include "UMGEditorActions.h"
+#include "WidgetTemplateDragDropOp.h"
+
+#include "PreviewScene.h"
+#include "SceneViewport.h"
+
+#include "BlueprintEditor.h"
+#include "SKismetInspector.h"
+#include "BlueprintEditorUtils.h"
+
+#include "Kismet2NameValidators.h"
+
+#include "WidgetBlueprintEditor.h"
+
+#define LOCTEXT_NAMESPACE "UMG"
+
+void SHierarchyViewItem::Construct(const FArguments& InArgs, const TSharedRef< STableViewBase >& InOwnerTableView, TSharedPtr<FWidgetBlueprintEditor> InBlueprintEditor, UWidget* InItem)
+{
+	BlueprintEditor = InBlueprintEditor;
+	Item = InItem;
+
+	STableRow< UWidget* >::Construct(
+		STableRow< UWidget* >::FArguments()
+		.Padding(2.0f)
+		.Content()
+		[
+			SNew(SHorizontalBox)
+
+			// Widget icon
+			+ SHorizontalBox::Slot()
+			.AutoWidth()
+			[
+				SNew(SImage)
+				.ColorAndOpacity(FLinearColor(1,1,1,0.5))
+				.Image(Item->GetEditorIcon())
+			]
+
+			// Name of the widget
+			+ SHorizontalBox::Slot()
+			.FillWidth(1.0f)
+			.Padding(2, 0, 0, 0)
+			.VAlign(VAlign_Center)
+			[
+				SNew(SInlineEditableTextBlock)
+				.Font(this, &SHierarchyViewItem::GetItemFont)
+				.Text(this, &SHierarchyViewItem::GetItemText)
+				.HighlightText(InArgs._HighlightText)
+				.OnVerifyTextChanged(this, &SHierarchyViewItem::OnVerifyNameTextChanged)
+				.OnTextCommitted(this, &SHierarchyViewItem::OnNameTextCommited)
+				.IsSelected(this, &SHierarchyViewItem::IsSelectedExclusively)
+			]
+		],
+		InOwnerTableView);
+}
+
+bool SHierarchyViewItem::OnVerifyNameTextChanged(const FText& InText, FText& OutErrorMessage)
+{
+	FString NewName = InText.ToString();
+
+	UWidgetBlueprint* Blueprint = BlueprintEditor.Pin()->GetWidgetBlueprintObj();
+	UWidget* ExistingWidget = Blueprint->WidgetTree->FindWidget(NewName);
+
+	FKismetNameValidator Validator(Blueprint);
+
+	const bool bUniqueNameForVariable = ( EValidatorResult::Ok == Validator.IsValid(NewName) );
+	
+	if ( ( ExistingWidget != NULL && ExistingWidget != Item ) || !bUniqueNameForVariable )
+	{
+		OutErrorMessage = LOCTEXT("NameConflict", "NameConflict");
+		return false;
+	}
+
+	return true;
+}
+
+void SHierarchyViewItem::OnNameTextCommited(const FText& InText, ETextCommit::Type CommitInfo)
+{
+	FWidgetBlueprintEditorUtils::RenameWidget(BlueprintEditor.Pin().ToSharedRef(), Item->GetFName(), FName(*InText.ToString()));
+}
+
+FSlateFontInfo SHierarchyViewItem::GetItemFont() const
+{
+	if ( !Item->IsGeneratedName() && Item->bIsVariable )
+	{
+		// TODO UMG Hacky move into style area
+		return FSlateFontInfo(FPaths::EngineContentDir() / TEXT("Slate/Fonts/Roboto-Bold.ttf"), 10);
+	}
+	else
+	{
+		// TODO UMG Hacky move into style area
+		return FSlateFontInfo(FPaths::EngineContentDir() / TEXT("Slate/Fonts/Roboto-Normal.ttf"), 9);
+	}
+}
+
+FText SHierarchyViewItem::GetItemText() const
+{
+	return FText::FromString(Item->GetLabel());
+}
+
+//@TODO UMG Allow items in the tree to be dragged, and reordered, and reparented.
+
+void SHierarchyViewItem::OnDragEnter(const FGeometry& MyGeometry, const FDragDropEvent& DragDropEvent)
+{
+	//TSharedPtr<FDragDropOperation> Operation = DragDropEvent.GetOperation();
+	//if ( !Operation.IsValid() )
+	//{
+	//	return;
+	//}
+
+	//if ( Operation->IsOfType<FWidgetTemplateDragDropOp>() )
+	//{
+	//	TSharedPtr<FWidgetTemplateDragDropOp> DragDropOp = StaticCastSharedPtr<FWidgetTemplateDragDropOp>(Operation);
+	//	bDraggedOver = true;
+	//}
+}
+
+void SHierarchyViewItem::OnDragLeave(const FDragDropEvent& DragDropEvent)
+{
+	TSharedPtr<FWidgetTemplateDragDropOp> DragDropOp = DragDropEvent.GetOperationAs<FWidgetTemplateDragDropOp>();
+	if ( DragDropOp.IsValid() )
+	{
+		DragDropOp->ResetToDefaultToolTip();
+	}
+}
+
+FReply SHierarchyViewItem::OnDragOver(const FGeometry& MyGeometry, const FDragDropEvent& DragDropEvent)
+{
+	TSharedPtr<FWidgetTemplateDragDropOp> DragDropOp = DragDropEvent.GetOperationAs<FWidgetTemplateDragDropOp>();
+	if ( DragDropOp.IsValid() )
+	{
+		return FReply::Handled();
+	}
+
+	return FReply::Unhandled();
+}
+
+FReply SHierarchyViewItem::OnDrop(const FGeometry& MyGeometry, const FDragDropEvent& DragDropEvent)
+{
+	TSharedPtr<FWidgetTemplateDragDropOp> DragDropOp = DragDropEvent.GetOperationAs<FWidgetTemplateDragDropOp>();
+	if ( DragDropOp.IsValid() )
+	{
+		if ( Item->IsA(UPanelWidget::StaticClass()) )
+		{
+			UWidgetBlueprint* BP = CastChecked<UWidgetBlueprint>(BlueprintEditor.Pin()->GetBlueprintObj());			
+			UPanelWidget* Parent = Cast<UPanelWidget>(Item);
+
+			UWidget* Widget = DragDropOp->Template->Create(BP->WidgetTree);
+
+			Parent->AddChild(Widget);
+
+			FBlueprintEditorUtils::MarkBlueprintAsStructurallyModified(BP);
+		}
+
+		return FReply::Handled();
+	}
+
+	return FReply::Unhandled();
+}
+
+#undef LOCTEXT_NAMESPACE
diff --git a/Engine/Source/Editor/UMGEditor/Private/Hierarchy/SHierarchyViewItem.h b/Engine/Source/Editor/UMGEditor/Private/Hierarchy/SHierarchyViewItem.h
new file mode 100644
--- /dev/null
+++ b/Engine/Source/Editor/UMGEditor/Private/Hierarchy/SHierarchyViewItem.h
@@ -0,0 +1,48 @@
+// Copyright 1998-2014 Epic Games, Inc. All Rights Reserved.
+
+#pragma once
+
+#include "SCompoundWidget.h"
+#include "BlueprintEditor.h"
+
+/**
+ * An widget item in the hierarchy tree view.
+ */
+class SHierarchyViewItem : public STableRow< UWidget* >
+{
+public:
+	SLATE_BEGIN_ARGS( SHierarchyViewItem ){}
+		
+		/** The current text to highlight */
+		SLATE_ATTRIBUTE( FText, HighlightText )
+
+	SLATE_END_ARGS()
+
+	void Construct(const FArguments& InArgs, const TSharedRef< STableViewBase >& InOwnerTableView, TSharedPtr<FWidgetBlueprintEditor> InBlueprintEditor, UWidget* InItem);
+
+	// Begin SWidget
+	virtual void OnDragEnter(const FGeometry& MyGeometry, const FDragDropEvent& DragDropEvent) override;
+	virtual void OnDragLeave(const FDragDropEvent& DragDropEvent) override;
+	virtual FReply OnDragOver(const FGeometry& MyGeometry, const FDragDropEvent& DragDropEvent) override;
+	virtual FReply OnDrop(const FGeometry& MyGeometry, const FDragDropEvent& DragDropEvent) override;
+	// End SWidget
+	
+private:
+	/** Called when text is being committed to check for validity */
+	bool OnVerifyNameTextChanged(const FText& InText, FText& OutErrorMessage);
+
+	/* Called when text is committed on the node */
+	void OnNameTextCommited(const FText& InText, ETextCommit::Type CommitInfo);
+
+	/* Gets the font to use for the text item, bold for customized named items */
+	FSlateFontInfo GetItemFont() const;
+
+	/* @returns the widget name to use for the tree item */
+	FText GetItemText() const;
+
+	/* The blueprint editor this tree view item's tree is associated with. */
+	TWeakPtr<FWidgetBlueprintEditor> BlueprintEditor;
+
+	/* The template widget that this tree item represents */
+	UWidget* Item;
+};
diff --git a/Engine/Source/Editor/UMGEditor/Private/Palette/SPaletteView.cpp b/Engine/Source/Editor/UMGEditor/Private/Palette/SPaletteView.cpp
new file mode 100644
--- /dev/null
+++ b/Engine/Source/Editor/UMGEditor/Private/Palette/SPaletteView.cpp
@@ -0,0 +1,291 @@
+// Copyright 1998-2014 Epic Games, Inc. All Rights Reserved.
+
+#include "UMGEditorPrivatePCH.h"
+
+#include "SPaletteView.h"
+#include "UMGEditorActions.h"
+
+#include "PreviewScene.h"
+#include "SceneViewport.h"
+
+#include "BlueprintEditor.h"
+
+#include "DecoratedDragDropOp.h"
+#include "WidgetTemplateDragDropOp.h"
+
+#include "WidgetTemplate.h"
+#include "WidgetTemplateClass.h"
+#include "WidgetTemplateButton.h"
+#include "WidgetTemplateCheckBox.h"
+
+#define LOCTEXT_NAMESPACE "UMG"
+
+class SWidgetTemplateItem : public SCompoundWidget
+{
+public:
+
+	SLATE_BEGIN_ARGS(SWidgetTemplateItem) {}
+	SLATE_END_ARGS()
+
+	/**
+	* Constructs this widget
+	*
+	* @param InArgs    Declaration from which to construct the widget
+	*/
+	void Construct(const FArguments& InArgs, TSharedPtr<FWidgetTemplate> InTemplate)
+	{
+		Template = InTemplate;
+
+		ChildSlot
+		[
+			SNew(SHorizontalBox)
+
+			+ SHorizontalBox::Slot()
+			.AutoWidth()
+			.VAlign(VAlign_Center)
+			[
+				SNew(SImage)
+				.ColorAndOpacity(FLinearColor(1, 1, 1, 0.5))
+				.Image(Template->GetIcon())
+			]
+
+			+ SHorizontalBox::Slot()
+			.AutoWidth()
+			.Padding(2, 0, 0, 0)
+			.VAlign(VAlign_Center)
+			[
+				SNew(STextBlock)
+				.Text(Template->Name)
+			]
+		];
+	}
+
+private:
+	TSharedPtr<FWidgetTemplate> Template;
+};
+
+class FWidgetTemplateViewModel : public FWidgetViewModel
+{
+public:
+	virtual FText GetName() const
+	{
+		return Template->Name;
+	}
+
+	virtual TSharedRef<ITableRow> BuildRow(const TSharedRef<STableViewBase>& OwnerTable) override
+	{
+		return SNew(STableRow< TSharedPtr<FWidgetViewModel> >, OwnerTable)
+			.Padding(2.0f)
+			.OnDragDetected(this, &FWidgetTemplateViewModel::OnDraggingWidgetTemplateItem)
+			[
+				SNew(SWidgetTemplateItem, Template)
+			];
+	}
+
+	FReply OnDraggingWidgetTemplateItem(const FGeometry& MyGeometry, const FPointerEvent& MouseEvent)
+	{
+		return FReply::Handled().BeginDragDrop(FWidgetTemplateDragDropOp::New(Template));
+	}
+
+	TSharedPtr<FWidgetTemplate> Template;
+};
+
+class FWidgetHeaderViewModel : public FWidgetViewModel
+{
+public:
+	virtual FText GetName() const
+	{
+		return GroupName;
+	}
+	
+	virtual TSharedRef<ITableRow> BuildRow(const TSharedRef<STableViewBase>& OwnerTable) override
+	{
+		return SNew(STableRow< TSharedPtr<FWidgetViewModel> >, OwnerTable)
+		.Style( FEditorStyle::Get(), "UMGEditor.ToolboxHeader" )
+			.Padding(2.0f)
+			.ShowSelection(false)
+			[
+				SNew(STextBlock)
+				.Text(GroupName)
+			];
+	}
+
+	virtual void GetChildren(TArray< TSharedPtr<FWidgetViewModel> >& OutChildren) override
+	{
+		for ( TSharedPtr<FWidgetViewModel>& Child : Children )
+		{
+			OutChildren.Add(Child);
+		}
+	}
+
+	FText GroupName;
+	TArray< TSharedPtr<FWidgetViewModel> > Children;
+};
+
+void SPaletteView::Construct(const FArguments& InArgs, TSharedPtr<FBlueprintEditor> InBlueprintEditor, USimpleConstructionScript* InSCS)
+{
+	BlueprintEditor = InBlueprintEditor;
+
+	UBlueprint* BP = InBlueprintEditor->GetBlueprintObj();
+
+	BuildWidgetList();
+
+	//FCoreDelegates::OnObjectPropertyChanged.AddRaw(this, &SPaletteView::OnObjectPropertyChanged);
+
+	SAssignNew(WidgetTemplatesView, STreeView< TSharedPtr<FWidgetViewModel> >)
+	.SelectionMode(ESelectionMode::Single)
+	.OnGenerateRow(this, &SPaletteView::OnGenerateWidgetTemplateItem)
+	.OnGetChildren(this, &SPaletteView::OnGetChildren)
+	.TreeItemsSource(&WidgetViewModels);
+
+	ChildSlot
+	[
+		SNew(SVerticalBox)
+
+		+ SVerticalBox::Slot()
+		.Padding(4)
+		.AutoHeight()
+		[
+			SNew(SSearchBox)
+			.HintText(LOCTEXT("SearchTemplates", "Search Templates"))
+			.OnTextChanged(this, &SPaletteView::OnSearchChanged)
+		]
+
+		+ SVerticalBox::Slot()
+		.FillHeight(1.0f)
+		[
+			SNew(SScrollBorder, WidgetTemplatesView.ToSharedRef())
+			[
+				WidgetTemplatesView.ToSharedRef()
+			]
+		]
+	];
+
+	LoadItemExpanssion();
+}
+
+SPaletteView::~SPaletteView()
+{
+	//FCoreDelegates::OnObjectPropertyChanged.RemoveAll(this);
+	SaveItemExpansion();
+}
+
+void SPaletteView::OnSearchChanged(const FText& InFilterText)
+{
+	bRefreshRequested = true;
+	SearchText = InFilterText;
+}
+
+void SPaletteView::LoadItemExpanssion()
+{
+	// Restore the expansion state of the widget groups.
+	for ( TSharedPtr<FWidgetViewModel>& ViewModel : WidgetViewModels )
+	{
+		bool IsExpanded;
+		if ( GConfig->GetBool(TEXT("WidgetTemplatesExpanded"), *ViewModel->GetName().ToString(), IsExpanded, GEditorUserSettingsIni) && IsExpanded )
+		{
+			WidgetTemplatesView->SetItemExpansion(ViewModel, true);
+		}
+	}
+}
+
+void SPaletteView::SaveItemExpansion()
+{
+	// Restore the expansion state of the widget groups.
+	for ( TSharedPtr<FWidgetViewModel>& ViewModel : WidgetViewModels )
+	{
+		const bool IsExpanded = WidgetTemplatesView->IsItemExpanded(ViewModel);
+		GConfig->SetBool(TEXT("WidgetTemplatesExpanded"), *ViewModel->GetName().ToString(), IsExpanded, GEditorUserSettingsIni);
+	}
+}
+
+UWidgetBlueprint* SPaletteView::GetBlueprint() const
+{
+	if ( BlueprintEditor.IsValid() )
+	{
+		UBlueprint* BP = BlueprintEditor.Pin()->GetBlueprintObj();
+		return Cast<UWidgetBlueprint>(BP);
+	}
+
+	return NULL;
+}
+
+void SPaletteView::BuildWidgetList()
+{
+	WidgetViewModels.Reset();
+	WidgetTemplateCategories.Reset();
+
+	BuildClassWidgetList();
+	BuildSpecialWidgetList();
+
+	for ( auto& Entry : WidgetTemplateCategories )
+	{
+		TSharedPtr<FWidgetHeaderViewModel> Header = MakeShareable(new FWidgetHeaderViewModel());
+		Header->GroupName = FText::FromString(Entry.Key);
+
+		for ( auto& Template : Entry.Value )
+		{
+			TSharedPtr<FWidgetTemplateViewModel> TemplateViewModel = MakeShareable(new FWidgetTemplateViewModel());
+			TemplateViewModel->Template = Template;
+
+			Header->Children.Add(TemplateViewModel);
+		}
+
+		Header->Children.Sort([] (TSharedPtr<FWidgetViewModel> L, TSharedPtr<FWidgetViewModel> R) { return R->GetName().CompareTo(L->GetName()) > 0; });
+
+		WidgetViewModels.Add(Header);
+	}
+
+	WidgetViewModels.Sort([] (TSharedPtr<FWidgetViewModel> L, TSharedPtr<FWidgetViewModel> R) { return R->GetName().CompareTo(L->GetName()) > 0; });
+}
+
+void SPaletteView::BuildClassWidgetList()
+{
+	for ( TObjectIterator<UClass> ClassIt; ClassIt; ++ClassIt )
+	{
+		UClass* WidgetClass = *ClassIt;
+		if ( WidgetClass->IsChildOf(UWidget::StaticClass()) )
+		{
+			const bool bIsSkeletonClass = WidgetClass->HasAnyFlags(RF_Transient) && WidgetClass->HasAnyClassFlags(CLASS_CompiledFromBlueprint);
+			const bool bIsValidClass =
+				!( WidgetClass->HasAnyClassFlags(CLASS_Abstract | CLASS_Deprecated | CLASS_NewerVersionExists) || bIsSkeletonClass );
+			const bool bIsSameClass = WidgetClass->GetFName() == GetBlueprint()->GeneratedClass->GetFName();
+
+			if ( bIsValidClass && !bIsSameClass )
+			{
+				TSharedPtr<FWidgetTemplateClass> Template = MakeShareable(new FWidgetTemplateClass(WidgetClass));
+				Template->Name = WidgetClass->GetDisplayNameText();
+				Template->ToolTip = WidgetClass->GetDisplayNameText();
+
+				AddWidgetTemplate(Template);
+			}
+		}
+	}
+}
+
+void SPaletteView::BuildSpecialWidgetList()
+{
+	AddWidgetTemplate(MakeShareable(new FWidgetTemplateButton()));
+	AddWidgetTemplate(MakeShareable(new FWidgetTemplateCheckBox()));
+
+	//TODO Make this pluggable.
+}
+
+void SPaletteView::AddWidgetTemplate(TSharedPtr<FWidgetTemplate> Template)
+{
+	FString Category = Template->GetCategory().ToString();
+	WidgetTemplateArray& Group = WidgetTemplateCategories.FindOrAdd(Category);
+	Group.Add(Template);
+}
+
+void SPaletteView::OnGetChildren(TSharedPtr<FWidgetViewModel> Item, TArray< TSharedPtr<FWidgetViewModel> >& Children)
+{
+	return Item->GetChildren(Children);
+}
+
+TSharedRef<ITableRow> SPaletteView::OnGenerateWidgetTemplateItem(TSharedPtr<FWidgetViewModel> Item, const TSharedRef<STableViewBase>& OwnerTable)
+{
+	return Item->BuildRow(OwnerTable);
+}
+
+#undef LOCTEXT_NAMESPACE
diff --git a/Engine/Source/Editor/UMGEditor/Private/Palette/SPaletteView.h b/Engine/Source/Editor/UMGEditor/Private/Palette/SPaletteView.h
new file mode 100644
--- /dev/null
+++ b/Engine/Source/Editor/UMGEditor/Private/Palette/SPaletteView.h
@@ -0,0 +1,64 @@
+// Copyright 1998-2014 Epic Games, Inc. All Rights Reserved.
+
+#pragma once
+
+#include "SCompoundWidget.h"
+#include "BlueprintEditor.h"
+
+#include "WidgetTemplate.h"
+
+/** View model for the items in the widget template list */
+class FWidgetViewModel : public TSharedFromThis<FWidgetViewModel>
+{
+public:
+	virtual FText GetName() const = 0;
+
+	virtual TSharedRef<ITableRow> BuildRow(const TSharedRef<STableViewBase>& OwnerTable) = 0;
+
+	virtual void GetChildren(TArray< TSharedPtr<FWidgetViewModel> >& OutChildren)
+	{
+	}
+};
+
+/**  */
+class SPaletteView : public SCompoundWidget
+{
+public:
+	SLATE_BEGIN_ARGS( SPaletteView ){}
+	SLATE_END_ARGS()
+
+	void Construct(const FArguments& InArgs, TSharedPtr<FBlueprintEditor> InBlueprintEditor, USimpleConstructionScript* InSCS);
+	virtual ~SPaletteView();
+
+private:
+	UWidgetBlueprint* GetBlueprint() const;
+
+	void BuildWidgetList();
+	void BuildClassWidgetList();
+	void BuildSpecialWidgetList();
+
+	void OnGetChildren(TSharedPtr<FWidgetViewModel> Item, TArray< TSharedPtr<FWidgetViewModel> >& Children);
+	TSharedRef<ITableRow> OnGenerateWidgetTemplateItem(TSharedPtr<FWidgetViewModel> Item, const TSharedRef<STableViewBase>& OwnerTable);
+
+	/** Called when the filter text is changed. */
+	void OnSearchChanged(const FText& InFilterText);
+
+private:
+	void LoadItemExpanssion();
+	void SaveItemExpansion();
+
+	void AddWidgetTemplate(TSharedPtr<FWidgetTemplate> Template);
+
+	TWeakPtr<class FBlueprintEditor> BlueprintEditor;
+
+	typedef TArray< TSharedPtr<FWidgetTemplate> > WidgetTemplateArray;
+	TMap< FString, WidgetTemplateArray > WidgetTemplateCategories;
+
+	typedef TArray< TSharedPtr<FWidgetViewModel> > ViewModelsArray;
+	ViewModelsArray WidgetViewModels;
+
+	TSharedPtr< STreeView< TSharedPtr<FWidgetViewModel> > > WidgetTemplatesView;
+
+	bool bRefreshRequested;
+	FText SearchText;
+};
diff --git a/Engine/Source/Editor/UMGEditor/Private/SDesignSurface.cpp b/Engine/Source/Editor/UMGEditor/Private/SDesignSurface.cpp
deleted file mode 100644
diff --git a/Engine/Source/Editor/UMGEditor/Private/SDesignSurface.h b/Engine/Source/Editor/UMGEditor/Private/SDesignSurface.h
deleted file mode 100644
diff --git a/Engine/Source/Editor/UMGEditor/Private/SUMGDesigner.cpp b/Engine/Source/Editor/UMGEditor/Private/SUMGDesigner.cpp
deleted file mode 100644
diff --git a/Engine/Source/Editor/UMGEditor/Private/SUMGDesigner.h b/Engine/Source/Editor/UMGEditor/Private/SUMGDesigner.h
deleted file mode 100644
diff --git a/Engine/Source/Editor/UMGEditor/Private/SUMGEditorTree.cpp b/Engine/Source/Editor/UMGEditor/Private/SUMGEditorTree.cpp
deleted file mode 100644
diff --git a/Engine/Source/Editor/UMGEditor/Private/SUMGEditorTree.h b/Engine/Source/Editor/UMGEditor/Private/SUMGEditorTree.h
deleted file mode 100644
diff --git a/Engine/Source/Editor/UMGEditor/Private/SUMGEditorTreeItem.cpp b/Engine/Source/Editor/UMGEditor/Private/SUMGEditorTreeItem.cpp
deleted file mode 100644
diff --git a/Engine/Source/Editor/UMGEditor/Private/SUMGEditorTreeItem.h b/Engine/Source/Editor/UMGEditor/Private/SUMGEditorTreeItem.h
deleted file mode 100644
diff --git a/Engine/Source/Editor/UMGEditor/Private/SUMGEditorWidgetTemplates.cpp b/Engine/Source/Editor/UMGEditor/Private/SUMGEditorWidgetTemplates.cpp
deleted file mode 100644
diff --git a/Engine/Source/Editor/UMGEditor/Private/SUMGEditorWidgetTemplates.h b/Engine/Source/Editor/UMGEditor/Private/SUMGEditorWidgetTemplates.h
deleted file mode 100644
diff --git a/Engine/Source/Editor/UMGEditor/Private/TabFactory/DesignerTabSummoner.cpp b/Engine/Source/Editor/UMGEditor/Private/TabFactory/DesignerTabSummoner.cpp
--- a/Engine/Source/Editor/UMGEditor/Private/TabFactory/DesignerTabSummoner.cpp
+++ b/Engine/Source/Editor/UMGEditor/Private/TabFactory/DesignerTabSummoner.cpp
@@ -3,6 +3,7 @@
 #include "UMGEditorPrivatePCH.h"
 
 #include "DesignerTabSummoner.h"
+#include "SDesignerView.h"
 
 #define LOCTEXT_NAMESPACE "UMG"
 
@@ -29,7 +30,7 @@ TSharedRef<SWidget> FDesignerTabSummoner::CreateTabBody(const FWorkflowTabSpawnI
 		[
 			SNew(STutorialWrapper, TEXT("Designer"))
 			[
-				SNew(SUMGDesigner, BlueprintEditor.Pin())
+				SNew(SDesignerView, BlueprintEditor.Pin())
 			]
 		];
 }
diff --git a/Engine/Source/Editor/UMGEditor/Private/TabFactory/HierarchyTabSummoner.cpp b/Engine/Source/Editor/UMGEditor/Private/TabFactory/HierarchyTabSummoner.cpp
--- a/Engine/Source/Editor/UMGEditor/Private/TabFactory/HierarchyTabSummoner.cpp
+++ b/Engine/Source/Editor/UMGEditor/Private/TabFactory/HierarchyTabSummoner.cpp
@@ -3,6 +3,7 @@
 #include "UMGEditorPrivatePCH.h"
 
 #include "HierarchyTabSummoner.h"
+#include "SHierarchyView.h"
 
 #define LOCTEXT_NAMESPACE "UMG"
 
@@ -27,7 +28,7 @@ TSharedRef<SWidget> FHierarchyTabSummoner::CreateTabBody(const FWorkflowTabSpawn
 
 	return SNew(STutorialWrapper, TEXT("Hierarchy"))
 		[
-			SNew(SUMGEditorTree, BlueprintEditorPtr, BlueprintEditorPtr->GetBlueprintObj()->SimpleConstructionScript)
+			SNew(SHierarchyView, BlueprintEditorPtr, BlueprintEditorPtr->GetBlueprintObj()->SimpleConstructionScript)
 		];
 }
 
diff --git a/Engine/Source/Editor/UMGEditor/Private/TabFactory/PaletteTabSummoner.cpp b/Engine/Source/Editor/UMGEditor/Private/TabFactory/PaletteTabSummoner.cpp
--- a/Engine/Source/Editor/UMGEditor/Private/TabFactory/PaletteTabSummoner.cpp
+++ b/Engine/Source/Editor/UMGEditor/Private/TabFactory/PaletteTabSummoner.cpp
@@ -3,6 +3,7 @@
 #include "UMGEditorPrivatePCH.h"
 
 #include "PaletteTabSummoner.h"
+#include "SPaletteView.h"
 
 #define LOCTEXT_NAMESPACE "UMG"
 
@@ -27,7 +28,7 @@ TSharedRef<SWidget> FPaletteTabSummoner::CreateTabBody(const FWorkflowTabSpawnIn
 
 	return SNew(STutorialWrapper, TEXT("Palette"))
 		[
-			SNew(SUMGEditorWidgetTemplates, BlueprintEditorPtr, BlueprintEditorPtr->GetBlueprintObj()->SimpleConstructionScript)
+			SNew(SPaletteView, BlueprintEditorPtr, BlueprintEditorPtr->GetBlueprintObj()->SimpleConstructionScript)
 		];
 }
 
diff --git a/Engine/Source/Editor/UMGEditor/Private/WidgetTemplateDragDropOp.cpp b/Engine/Source/Editor/UMGEditor/Private/WidgetTemplateDragDropOp.cpp
deleted file mode 100644
diff --git a/Engine/Source/Editor/UMGEditor/Private/WidgetTemplateDragDropOp.h b/Engine/Source/Editor/UMGEditor/Private/WidgetTemplateDragDropOp.h
deleted file mode 100644
diff --git a/Engine/Source/Editor/UMGEditor/UMGEditor.Build.cs b/Engine/Source/Editor/UMGEditor/UMGEditor.Build.cs
--- a/Engine/Source/Editor/UMGEditor/UMGEditor.Build.cs
+++ b/Engine/Source/Editor/UMGEditor/UMGEditor.Build.cs
@@ -14,6 +14,10 @@ public class UMGEditor : ModuleRules
 				"Editor/UMGEditor/Private/Customizations",
 				"Editor/UMGEditor/Private/BlueprintModes",
 				"Editor/UMGEditor/Private/TabFactory",
+				"Editor/UMGEditor/Private/Designer",
+				"Editor/UMGEditor/Private/Hierarchy",
+				"Editor/UMGEditor/Private/Palette",
+				"Editor/UMGEditor/Private/DragDrop",
 			});
 
 		PrivateIncludePathModuleNames.AddRange(
