UMG - Fixing non unity build issues.

diff --git a/Engine/Source/Editor/UMGEditor/Private/BlueprintModes/WidgetDesignerApplicationMode.cpp b/Engine/Source/Editor/UMGEditor/Private/BlueprintModes/WidgetDesignerApplicationMode.cpp
--- a/Engine/Source/Editor/UMGEditor/Private/BlueprintModes/WidgetDesignerApplicationMode.cpp
+++ b/Engine/Source/Editor/UMGEditor/Private/BlueprintModes/WidgetDesignerApplicationMode.cpp
@@ -6,6 +6,7 @@
 #include "Editor/Kismet/Public/BlueprintEditorTabs.h"
 #include "Editor/Kismet/Public/SBlueprintEditorToolbar.h"
 #include "Editor/Kismet/Public/BlueprintEditorModes.h"
+#include "SKismetInspector.h"
 
 #include "WidgetBlueprintEditor.h"
 #include "WidgetDesignerApplicationMode.h"
diff --git a/Engine/Source/Editor/UMGEditor/Private/WidgetBlueprintEditorToolbar.cpp b/Engine/Source/Editor/UMGEditor/Private/WidgetBlueprintEditorToolbar.cpp
--- a/Engine/Source/Editor/UMGEditor/Private/WidgetBlueprintEditorToolbar.cpp
+++ b/Engine/Source/Editor/UMGEditor/Private/WidgetBlueprintEditorToolbar.cpp
@@ -1,9 +1,9 @@
 // Copyright 1998-2014 Epic Games, Inc. All Rights Reserved.
 
 #include "UMGEditorPrivatePCH.h"
-
 #include "WidgetBlueprintEditorToolbar.h"
 
+#include "WidgetBlueprintEditor.h"
 #include "SModeWidget.h"
 
 #define LOCTEXT_NAMESPACE "UMG"
