UMG - Fixing more non-unity build issues.

diff --git a/Engine/Source/Editor/UMGEditor/Private/Palette/SPaletteView.cpp b/Engine/Source/Editor/UMGEditor/Private/Palette/SPaletteView.cpp
--- a/Engine/Source/Editor/UMGEditor/Private/Palette/SPaletteView.cpp
+++ b/Engine/Source/Editor/UMGEditor/Private/Palette/SPaletteView.cpp
@@ -122,7 +122,7 @@ public:
 	TArray< TSharedPtr<FWidgetViewModel> > Children;
 };
 
-void SPaletteView::Construct(const FArguments& InArgs, TSharedPtr<FBlueprintEditor> InBlueprintEditor, USimpleConstructionScript* InSCS)
+void SPaletteView::Construct(const FArguments& InArgs, TSharedPtr<FBlueprintEditor> InBlueprintEditor)
 {
 	BlueprintEditor = InBlueprintEditor;
 
diff --git a/Engine/Source/Editor/UMGEditor/Private/Palette/SPaletteView.h b/Engine/Source/Editor/UMGEditor/Private/Palette/SPaletteView.h
--- a/Engine/Source/Editor/UMGEditor/Private/Palette/SPaletteView.h
+++ b/Engine/Source/Editor/UMGEditor/Private/Palette/SPaletteView.h
@@ -27,7 +27,7 @@ public:
 	SLATE_BEGIN_ARGS( SPaletteView ){}
 	SLATE_END_ARGS()
 
-	void Construct(const FArguments& InArgs, TSharedPtr<FBlueprintEditor> InBlueprintEditor, USimpleConstructionScript* InSCS);
+	void Construct(const FArguments& InArgs, TSharedPtr<FBlueprintEditor> InBlueprintEditor);
 	virtual ~SPaletteView();
 
 private:
diff --git a/Engine/Source/Editor/UMGEditor/Private/TabFactory/PaletteTabSummoner.cpp b/Engine/Source/Editor/UMGEditor/Private/TabFactory/PaletteTabSummoner.cpp
--- a/Engine/Source/Editor/UMGEditor/Private/TabFactory/PaletteTabSummoner.cpp
+++ b/Engine/Source/Editor/UMGEditor/Private/TabFactory/PaletteTabSummoner.cpp
@@ -28,7 +28,7 @@ TSharedRef<SWidget> FPaletteTabSummoner::CreateTabBody(const FWorkflowTabSpawnIn
 
 	return SNew(STutorialWrapper, TEXT("Palette"))
 		[
-			SNew(SPaletteView, BlueprintEditorPtr, BlueprintEditorPtr->GetBlueprintObj()->SimpleConstructionScript)
+			SNew(SPaletteView, BlueprintEditorPtr)
 		];
 }
 
diff --git a/Engine/Source/Editor/UMGEditor/Private/TabFactory/SequencerTabSummoner.cpp b/Engine/Source/Editor/UMGEditor/Private/TabFactory/SequencerTabSummoner.cpp
--- a/Engine/Source/Editor/UMGEditor/Private/TabFactory/SequencerTabSummoner.cpp
+++ b/Engine/Source/Editor/UMGEditor/Private/TabFactory/SequencerTabSummoner.cpp
@@ -3,6 +3,7 @@
 #include "UMGEditorPrivatePCH.h"
 
 #include "SequencerTabSummoner.h"
+#include "WidgetBlueprintEditor.h"
 
 #define LOCTEXT_NAMESPACE "UMG"
 
