UMG - More non-unity build fixes.

diff --git a/Engine/Source/Editor/UMGEditor/Private/Designer/SDesignerView.cpp b/Engine/Source/Editor/UMGEditor/Private/Designer/SDesignerView.cpp
--- a/Engine/Source/Editor/UMGEditor/Private/Designer/SDesignerView.cpp
+++ b/Engine/Source/Editor/UMGEditor/Private/Designer/SDesignerView.cpp
@@ -6,8 +6,6 @@
 #include "BlueprintEditor.h"
 #include "SKismetInspector.h"
 #include "BlueprintEditorUtils.h"
-#include "WidgetBlueprint.h"
-#include "WidgetBlueprintEditor.h"
 
 #include "WidgetTemplateDragDropOp.h"
 
diff --git a/Engine/Source/Editor/UMGEditor/Private/UMGEditorPrivatePCH.h b/Engine/Source/Editor/UMGEditor/Private/UMGEditorPrivatePCH.h
--- a/Engine/Source/Editor/UMGEditor/Private/UMGEditorPrivatePCH.h
+++ b/Engine/Source/Editor/UMGEditor/Private/UMGEditorPrivatePCH.h
@@ -24,6 +24,8 @@
 #include "WidgetBlueprintApplicationModes.h"
 #include "WidgetBlueprintApplicationMode.h"
 
+#include "WidgetReference.h"
+#include "WidgetBlueprintEditor.h"
 
 #include "AssetTypeActions_Base.h"
 
