UMG - Fixing another non-unity build problem.

diff --git a/Engine/Source/Editor/UMGEditor/Private/TabFactory/SequencerTabSummoner.cpp b/Engine/Source/Editor/UMGEditor/Private/TabFactory/SequencerTabSummoner.cpp
--- a/Engine/Source/Editor/UMGEditor/Private/TabFactory/SequencerTabSummoner.cpp
+++ b/Engine/Source/Editor/UMGEditor/Private/TabFactory/SequencerTabSummoner.cpp
@@ -4,6 +4,7 @@
 
 #include "SequencerTabSummoner.h"
 #include "WidgetBlueprintEditor.h"
+#include "ISequencer.h"
 
 #define LOCTEXT_NAMESPACE "UMG"
 
