UMG - No longer using the old blueprint component details panel.  There's now a separate one for the widget designer view that we can customize as need be.

diff --git a/Engine/Source/Editor/UMGEditor/Private/BlueprintModes/WidgetDesignerApplicationMode.cpp b/Engine/Source/Editor/UMGEditor/Private/BlueprintModes/WidgetDesignerApplicationMode.cpp
--- a/Engine/Source/Editor/UMGEditor/Private/BlueprintModes/WidgetDesignerApplicationMode.cpp
+++ b/Engine/Source/Editor/UMGEditor/Private/BlueprintModes/WidgetDesignerApplicationMode.cpp
@@ -19,6 +19,7 @@
 #include "HierarchyTabSummoner.h"
 #include "DesignerTabSummoner.h"
 #include "SequencerTabSummoner.h"
+#include "DetailsTabSummoner.h"
 
 /////////////////////////////////////////////////////
 // FWidgetDesignerApplicationMode
@@ -26,7 +27,7 @@
 FWidgetDesignerApplicationMode::FWidgetDesignerApplicationMode(TSharedPtr<FWidgetBlueprintEditor> InWidgetEditor)
 	: FWidgetBlueprintApplicationMode(InWidgetEditor, FWidgetBlueprintApplicationModes::DesignerMode)
 {
-	TabLayout = FTabManager::NewLayout( "WidgetBlueprintEditor_Designer_Layout_v1" )
+	TabLayout = FTabManager::NewLayout( "WidgetBlueprintEditor_Designer_Layout_v2" )
 	->AddArea
 	(
 		FTabManager::NewPrimaryArea()
@@ -75,7 +76,7 @@ FWidgetDesignerApplicationMode::FWidgetDesignerApplicationMode(TSharedPtr<FWidge
 				(
 					FTabManager::NewStack()
 					->SetSizeCoefficient( 0.5f )
-					->AddTab( FBlueprintEditorTabs::DetailsID, ETabState::OpenedTab )
+					->AddTab( FDetailsTabSummoner::TabID, ETabState::OpenedTab )
 				)
 			)
 		)
@@ -90,7 +91,8 @@ FWidgetDesignerApplicationMode::FWidgetDesignerApplicationMode(TSharedPtr<FWidge
 
 
 	// Add Tab Spawners
-	TabFactories.RegisterFactory(MakeShareable(new FSelectionDetailsSummoner(InWidgetEditor)));
+	//TabFactories.RegisterFactory(MakeShareable(new FSelectionDetailsSummoner(InWidgetEditor)));
+	TabFactories.RegisterFactory(MakeShareable(new FDetailsTabSummoner(InWidgetEditor)));
 	TabFactories.RegisterFactory(MakeShareable(new FDesignerTabSummoner(InWidgetEditor)));
 	TabFactories.RegisterFactory(MakeShareable(new FHierarchyTabSummoner(InWidgetEditor)));
 	TabFactories.RegisterFactory(MakeShareable(new FPaletteTabSummoner(InWidgetEditor)));
@@ -118,26 +120,16 @@ void FWidgetDesignerApplicationMode::PreDeactivateMode()
 {
 	//FWidgetBlueprintApplicationMode::PreDeactivateMode();
 
-	GetBlueprintEditor()->GetInspector()->EnableComponentDetailsCustomization(false);
+	//GetBlueprintEditor()->GetInspector()->EnableComponentDetailsCustomization(false);
 
-	static FName PropertyEditor("PropertyEditor");
-	FPropertyEditorModule& PropertyModule = FModuleManager::GetModuleChecked<FPropertyEditorModule>(PropertyEditor);
-	PropertyModule.UnregisterCustomPropertyTypeLayout(TEXT("CanvasPanelSlot"));
+	//static FName PropertyEditor("PropertyEditor");
+	//FPropertyEditorModule& PropertyModule = FModuleManager::GetModuleChecked<FPropertyEditorModule>(PropertyEditor);
+	//PropertyModule.UnregisterCustomPropertyTypeLayout(TEXT("CanvasPanelSlot"));
 }
 
 void FWidgetDesignerApplicationMode::PostActivateMode()
 {
 	//FWidgetBlueprintApplicationMode::PostActivateMode();
 
-	GetBlueprintEditor()->GetInspector()->EnableComponentDetailsCustomization(false);
-
-	TSharedRef<class SKismetInspector> Inspector = MyBlueprintEditor.Pin()->GetInspector();
-	FOnGetDetailCustomizationInstance LayoutDelegateDetails = FOnGetDetailCustomizationInstance::CreateStatic(&FBlueprintWidgetCustomization::MakeInstance, MyBlueprintEditor.Pin()->GetBlueprintObj());
-	Inspector->GetPropertyView()->RegisterInstancedCustomPropertyLayout(UWidget::StaticClass(), LayoutDelegateDetails);
-
-	FOnGetPropertyTypeCustomizationInstance CanvasSlotCustomization = FOnGetPropertyTypeCustomizationInstance::CreateStatic(&FCanvasSlotCustomization::MakeInstance, MyBlueprintEditor.Pin()->GetBlueprintObj());
-
-	static FName PropertyEditor("PropertyEditor");
-	FPropertyEditorModule& PropertyModule = FModuleManager::GetModuleChecked<FPropertyEditorModule>(PropertyEditor);
-	PropertyModule.RegisterCustomPropertyTypeLayout(TEXT("PanelSlot"), CanvasSlotCustomization);
+	//GetBlueprintEditor()->GetInspector()->EnableComponentDetailsCustomization(false);
 }
diff --git a/Engine/Source/Editor/UMGEditor/Private/Details/SWidgetDetailsView.cpp b/Engine/Source/Editor/UMGEditor/Private/Details/SWidgetDetailsView.cpp
new file mode 100644
--- /dev/null
+++ b/Engine/Source/Editor/UMGEditor/Private/Details/SWidgetDetailsView.cpp
@@ -0,0 +1,97 @@
+// Copyright 1998-2014 Epic Games, Inc. All Rights Reserved.
+
+#include "UMGEditorPrivatePCH.h"
+#include "SWidgetDetailsView.h"
+
+#include "BlueprintEditor.h"
+#include "IDetailsView.h"
+
+#define LOCTEXT_NAMESPACE "UMG"
+
+void SWidgetDetailsView::Construct(const FArguments& InArgs, TSharedPtr<FWidgetBlueprintEditor> InBlueprintEditor)
+{
+	BlueprintEditor = InBlueprintEditor;
+
+	// Create a property view
+	FPropertyEditorModule& EditModule = FModuleManager::Get().GetModuleChecked<FPropertyEditorModule>("PropertyEditor");
+
+	FNotifyHook* NotifyHook = InBlueprintEditor.Get();
+	FDetailsViewArgs DetailsViewArgs( /*bUpdateFromSelection=*/ false, /*bLockable=*/ false, /*bAllowSearch=*/ true, /*bObjectsUseNameArea=*/ true, /*bHideSelectionTip=*/ true, /*InNotifyHook=*/ NotifyHook, /*InSearchInitialKeyFocus=*/ false, /*InViewIdentifier=*/ NAME_None);
+	DetailsViewArgs.bHideActorNameArea = true;
+
+	PropertyView = EditModule.CreateDetailView(DetailsViewArgs);
+
+	ChildSlot
+	[
+		SNew(SVerticalBox)
+
+		+ SVerticalBox::Slot()
+		.FillHeight(1.0f)
+		[
+			PropertyView.ToSharedRef()
+		]
+	];
+	
+	BlueprintEditor.Pin()->OnSelectedWidgetsChanging.AddRaw(this, &SWidgetDetailsView::OnEditorSelectionChanging);
+	BlueprintEditor.Pin()->OnSelectedWidgetsChanged.AddRaw(this, &SWidgetDetailsView::OnEditorSelectionChanged);
+
+	RegisterCustomizations();
+}
+
+SWidgetDetailsView::~SWidgetDetailsView()
+{
+	if ( BlueprintEditor.IsValid() )
+	{
+		BlueprintEditor.Pin()->OnSelectedWidgetsChanging.RemoveAll(this);
+		BlueprintEditor.Pin()->OnSelectedWidgetsChanged.RemoveAll(this);
+	}
+}
+
+void SWidgetDetailsView::RegisterCustomizations()
+{
+	FOnGetDetailCustomizationInstance LayoutDelegateDetails = FOnGetDetailCustomizationInstance::CreateStatic(&FBlueprintWidgetCustomization::MakeInstance, BlueprintEditor.Pin()->GetBlueprintObj());
+	PropertyView->RegisterInstancedCustomPropertyLayout(UWidget::StaticClass(), LayoutDelegateDetails);
+
+	FOnGetPropertyTypeCustomizationInstance CanvasSlotCustomization = FOnGetPropertyTypeCustomizationInstance::CreateStatic(&FCanvasSlotCustomization::MakeInstance, BlueprintEditor.Pin()->GetBlueprintObj());
+
+	static FName PropertyEditor("PropertyEditor");
+	FPropertyEditorModule& PropertyModule = FModuleManager::GetModuleChecked<FPropertyEditorModule>(PropertyEditor);
+	PropertyModule.RegisterCustomPropertyTypeLayout(TEXT("PanelSlot"), CanvasSlotCustomization);
+}
+
+void SWidgetDetailsView::OnEditorSelectionChanging()
+{
+	ClearFocusIfOwned();
+}
+
+void SWidgetDetailsView::OnEditorSelectionChanged()
+{
+	TSet< FWidgetReference > SelectedWidgets = BlueprintEditor.Pin()->GetSelectedWidgets();
+
+	SelectedObjects.Empty();
+	for ( FWidgetReference& WidgetRef : SelectedWidgets )
+	{
+		SelectedObjects.Add(WidgetRef.GetPreview());
+	}
+
+	const bool bForceRefresh = false;
+	PropertyView->SetObjects(SelectedObjects, bForceRefresh);
+}
+
+void SWidgetDetailsView::ClearFocusIfOwned()
+{
+	static bool bIsReentrant = false;
+	if ( !bIsReentrant )
+	{
+		bIsReentrant = true;
+		// When the selection is changed, we may be potentially actively editing a property,
+		// if this occurs we need need to immediately clear keyboard focus
+		if ( FSlateApplication::Get().HasFocusedDescendants(AsShared()) )
+		{
+			FSlateApplication::Get().ClearKeyboardFocus(EKeyboardFocusCause::Mouse);
+		}
+		bIsReentrant = false;
+	}
+}
+
+#undef LOCTEXT_NAMESPACE
diff --git a/Engine/Source/Editor/UMGEditor/Private/Details/SWidgetDetailsView.h b/Engine/Source/Editor/UMGEditor/Private/Details/SWidgetDetailsView.h
new file mode 100644
--- /dev/null
+++ b/Engine/Source/Editor/UMGEditor/Private/Details/SWidgetDetailsView.h
@@ -0,0 +1,39 @@
+// Copyright 1998-2014 Epic Games, Inc. All Rights Reserved.
+
+#pragma once
+
+#include "SCompoundWidget.h"
+#include "BlueprintEditor.h"
+
+#include "WidgetTemplate.h"
+
+/**  */
+class SWidgetDetailsView : public SCompoundWidget
+{
+public:
+	SLATE_BEGIN_ARGS( SWidgetDetailsView ){}
+	SLATE_END_ARGS()
+
+	void Construct(const FArguments& InArgs, TSharedPtr<FWidgetBlueprintEditor> InBlueprintEditor);
+	virtual ~SWidgetDetailsView();
+	
+	TSharedPtr<class IDetailsView> GetPropertyView() const { return PropertyView; }
+	
+private:
+	void RegisterCustomizations();
+
+	void OnEditorSelectionChanging();
+	void OnEditorSelectionChanged();
+
+	void ClearFocusIfOwned();
+
+private:
+	/** The editor that owns this details view */
+	TWeakPtr<class FWidgetBlueprintEditor> BlueprintEditor;
+	
+	/** Property viewing widget */
+	TSharedPtr<class IDetailsView> PropertyView;
+
+	/** Selected objects for this detail view */
+	TArray< TWeakObjectPtr<UObject> > SelectedObjects;
+};
diff --git a/Engine/Source/Editor/UMGEditor/Private/TabFactory/DetailsTabSummoner.cpp b/Engine/Source/Editor/UMGEditor/Private/TabFactory/DetailsTabSummoner.cpp
new file mode 100644
--- /dev/null
+++ b/Engine/Source/Editor/UMGEditor/Private/TabFactory/DetailsTabSummoner.cpp
@@ -0,0 +1,35 @@
+// Copyright 1998-2014 Epic Games, Inc. All Rights Reserved.
+
+#include "UMGEditorPrivatePCH.h"
+
+#include "PaletteTabSummoner.h"
+#include "SWidgetDetailsView.h"
+
+#define LOCTEXT_NAMESPACE "UMG"
+
+const FName FDetailsTabSummoner::TabID(TEXT("WidgetDetails"));
+
+FDetailsTabSummoner::FDetailsTabSummoner(TSharedPtr<class FWidgetBlueprintEditor> InBlueprintEditor)
+		: FWorkflowTabFactory(TabID, InBlueprintEditor)
+		, BlueprintEditor(InBlueprintEditor)
+{
+	TabLabel = LOCTEXT("WidgetDetails_TabLabel", "Details");
+	//TabIcon = FEditorStyle::GetBrush("Kismet.Tabs.Components");
+
+	bIsSingleton = true;
+
+	ViewMenuDescription = LOCTEXT("WidgetDetails_ViewMenu_Desc", "Details");
+	ViewMenuTooltip = LOCTEXT("WidgetDetails_ViewMenu_ToolTip", "Show the Details");
+}
+
+TSharedRef<SWidget> FDetailsTabSummoner::CreateTabBody(const FWorkflowTabSpawnInfo& Info) const
+{
+	TSharedPtr<FWidgetBlueprintEditor> BlueprintEditorPtr = StaticCastSharedPtr<FWidgetBlueprintEditor>(BlueprintEditor.Pin());
+
+	return SNew(STutorialWrapper, TEXT("Details"))
+		[
+			SNew(SWidgetDetailsView, BlueprintEditorPtr)
+		];
+}
+
+#undef LOCTEXT_NAMESPACE 
diff --git a/Engine/Source/Editor/UMGEditor/Private/TabFactory/DetailsTabSummoner.h b/Engine/Source/Editor/UMGEditor/Private/TabFactory/DetailsTabSummoner.h
new file mode 100644
--- /dev/null
+++ b/Engine/Source/Editor/UMGEditor/Private/TabFactory/DetailsTabSummoner.h
@@ -0,0 +1,17 @@
+// Copyright 1998-2014 Epic Games, Inc. All Rights Reserved.
+
+#pragma once
+
+struct FDetailsTabSummoner : public FWorkflowTabFactory
+{
+public:
+	static const FName TabID;
+	
+public:
+	FDetailsTabSummoner(TSharedPtr<class FWidgetBlueprintEditor> InBlueprintEditor);
+	
+	virtual TSharedRef<SWidget> CreateTabBody(const FWorkflowTabSpawnInfo& Info) const override;
+	
+protected:
+	TWeakPtr<class FWidgetBlueprintEditor> BlueprintEditor;
+};
diff --git a/Engine/Source/Editor/UMGEditor/Private/WidgetBlueprintEditor.cpp b/Engine/Source/Editor/UMGEditor/Private/WidgetBlueprintEditor.cpp
--- a/Engine/Source/Editor/UMGEditor/Private/WidgetBlueprintEditor.cpp
+++ b/Engine/Source/Editor/UMGEditor/Private/WidgetBlueprintEditor.cpp
@@ -116,7 +116,9 @@ void FWidgetBlueprintEditor::SelectWidgets(const TSet<FWidgetReference>& Widgets
 		}
 	}
 
-	RefreshDetails(TempSelection);
+	//RefreshDetails(TempSelection);
+
+	OnSelectedWidgetsChanging.Broadcast();
 
 	// Finally change the selected widgets after we've updated the details panel 
 	// to ensure values that are pending are committed on focus loss, and migrated properly
@@ -180,8 +182,8 @@ void FWidgetBlueprintEditor::RefreshDetails(TSet<FWidgetReference>& Widgets)
 	UWidgetBlueprint* Blueprint = GetWidgetBlueprintObj();
 
 	// Update the details panel
-	SKismetInspector::FShowDetailsOptions Options(InspectorTitle, true);
-	GetInspector()->ShowDetailsForObjects(InspectorObjects, Options);
+//	SKismetInspector::FShowDetailsOptions Options(InspectorTitle, true);
+//	GetInspector()->ShowDetailsForObjects(InspectorObjects, Options);
 }
 
 void FWidgetBlueprintEditor::OnBlueprintChanged(UBlueprint* InBlueprint)
diff --git a/Engine/Source/Editor/UMGEditor/Private/WidgetBlueprintEditor.h b/Engine/Source/Editor/UMGEditor/Private/WidgetBlueprintEditor.h
--- a/Engine/Source/Editor/UMGEditor/Private/WidgetBlueprintEditor.h
+++ b/Engine/Source/Editor/UMGEditor/Private/WidgetBlueprintEditor.h
@@ -14,6 +14,7 @@ class ISequencer;
 class FWidgetBlueprintEditor : public FBlueprintEditor
 {
 public:
+	DECLARE_MULTICAST_DELEGATE(FOnSelectedWidgetsChanging)
 	DECLARE_MULTICAST_DELEGATE(FOnSelectedWidgetsChanged)
 
 	/** Called after the widget preview has been updated */
@@ -56,6 +57,9 @@ public:
 	TSharedPtr<class FWidgetBlueprintEditorToolbar> GetWidgetToolbarBuilder() { return WidgetToolbar; }
 
 public:
+	/** Fires whenever the selected set of widgets changing */
+	FOnSelectedWidgetsChanged OnSelectedWidgetsChanging;
+
 	/** Fires whenever the selected set of widgets changes */
 	FOnSelectedWidgetsChanged OnSelectedWidgetsChanged;
 
diff --git a/Engine/Source/Editor/UMGEditor/UMGEditor.Build.cs b/Engine/Source/Editor/UMGEditor/UMGEditor.Build.cs
--- a/Engine/Source/Editor/UMGEditor/UMGEditor.Build.cs
+++ b/Engine/Source/Editor/UMGEditor/UMGEditor.Build.cs
@@ -17,6 +17,7 @@ public class UMGEditor : ModuleRules
 				"Editor/UMGEditor/Private/Designer",
 				"Editor/UMGEditor/Private/Hierarchy",
 				"Editor/UMGEditor/Private/Palette",
+				"Editor/UMGEditor/Private/Details",
 				"Editor/UMGEditor/Private/DragDrop",
 			});
 
