UMG - Updating selection when the blueprint is changed to ensure everyone is watching the latest preview objects.

diff --git a/Engine/Source/Editor/UMGEditor/Private/WidgetBlueprintEditor.cpp b/Engine/Source/Editor/UMGEditor/Private/WidgetBlueprintEditor.cpp
--- a/Engine/Source/Editor/UMGEditor/Private/WidgetBlueprintEditor.cpp
+++ b/Engine/Source/Editor/UMGEditor/Private/WidgetBlueprintEditor.cpp
@@ -116,8 +116,6 @@ void FWidgetBlueprintEditor::SelectWidgets(const TSet<FWidgetReference>& Widgets
 		}
 	}
 
-	//RefreshDetails(TempSelection);
-
 	OnSelectedWidgetsChanging.Broadcast();
 
 	// Finally change the selected widgets after we've updated the details panel 
@@ -159,42 +157,19 @@ const TSet<FWidgetReference>& FWidgetBlueprintEditor::GetSelectedWidgets() const
 	return SelectedWidgets;
 }
 
-void FWidgetBlueprintEditor::RefreshDetails()
-{
-	RefreshDetails(SelectedWidgets);
-}
-
-void FWidgetBlueprintEditor::RefreshDetails(TSet<FWidgetReference>& Widgets)
-{
-	// Convert the selection set to an array of UObject* pointers
-	FString InspectorTitle = "Widget";// Widget->GetDisplayString();
-
-	TArray<UObject*> InspectorObjects;
-	for ( FWidgetReference& WidgetRef : Widgets )
-	{
-		UWidget* PreviewWidget = WidgetRef.GetPreview();
-		if ( PreviewWidget )
-		{
-			InspectorObjects.Add(PreviewWidget);
-		}
-	}
-
-	UWidgetBlueprint* Blueprint = GetWidgetBlueprintObj();
-
-	// Update the details panel
-//	SKismetInspector::FShowDetailsOptions Options(InspectorTitle, true);
-//	GetInspector()->ShowDetailsForObjects(InspectorObjects, Options);
-}
-
 void FWidgetBlueprintEditor::OnBlueprintChanged(UBlueprint* InBlueprint)
 {
 	if ( InBlueprint )
 	{
+		// Rebuilding the preview can force objects to be recreated, so the selection may need to be updated.
+		OnSelectedWidgetsChanging.Broadcast();
+
 		CleanSelection();
 
 		UpdatePreview(InBlueprint, true);
 
-		RefreshDetails();
+		// Fire the selection updated event to ensure everyone is watching the same widgets.
+		OnSelectedWidgetsChanged.Broadcast();
 	}
 }
 
diff --git a/Engine/Source/Editor/UMGEditor/Private/WidgetBlueprintEditor.h b/Engine/Source/Editor/UMGEditor/Private/WidgetBlueprintEditor.h
--- a/Engine/Source/Editor/UMGEditor/Private/WidgetBlueprintEditor.h
+++ b/Engine/Source/Editor/UMGEditor/Private/WidgetBlueprintEditor.h
@@ -86,12 +86,6 @@ private:
 	void PasteWidgets();
 
 private:
-	/** Updates the inspector to be viewing the currently selected set of widgets */
-	void RefreshDetails();
-
-	/** Updates the inspector to be viewing a set of widgets */
-	void RefreshDetails(TSet<FWidgetReference>& Widgets);
-
 	/** Called whenever the blueprint is structurally changed. */
 	virtual void OnBlueprintChanged(UBlueprint* InBlueprint) override;
 
