UMG - Fixing a non-unity issue.

diff --git a/Engine/Source/Editor/UMGEditor/Private/TabFactory/DetailsTabSummoner.cpp b/Engine/Source/Editor/UMGEditor/Private/TabFactory/DetailsTabSummoner.cpp
--- a/Engine/Source/Editor/UMGEditor/Private/TabFactory/DetailsTabSummoner.cpp
+++ b/Engine/Source/Editor/UMGEditor/Private/TabFactory/DetailsTabSummoner.cpp
@@ -2,7 +2,7 @@
 
 #include "UMGEditorPrivatePCH.h"
 
-#include "PaletteTabSummoner.h"
+#include "DetailsTabSummoner.h"
 #include "SWidgetDetailsView.h"
 
 #define LOCTEXT_NAMESPACE "UMG"
