UMG - Widget blueprint variables are now sorted.

diff --git a/Engine/Source/Editor/UMGEditor/Private/WidgetBlueprintCompiler.cpp b/Engine/Source/Editor/UMGEditor/Private/WidgetBlueprintCompiler.cpp
--- a/Engine/Source/Editor/UMGEditor/Private/WidgetBlueprintCompiler.cpp
+++ b/Engine/Source/Editor/UMGEditor/Private/WidgetBlueprintCompiler.cpp
@@ -217,169 +217,6 @@ void FWidgetBlueprintCompiler::ProcessWidgetNode(UWidgetGraphNode_Base* VisualWi
 	}
 }
 
-//void FWidgetBlueprintCompiler::CreateEvaluationHandler(UWidgetGraphNode_Base* VisualAnimNode, FEvaluationHandlerRecord& Record)
-//{
-//	// Shouldn't create a handler if there is nothing to work with
-//	check(Record.ServicedProperties.Num() > 0);
-//	check(Record.NodeVariableProperty != NULL);
-//
-//
-//	//@TODO: Want to name these better
-//	const FString FunctionName = FString::Printf(TEXT("%s_%s_%s"), *Record.EvaluationHandlerProperty->GetName(), *VisualAnimNode->GetOuter()->GetName(), *VisualAnimNode->GetName());
-//	Record.HandlerFunctionName = FName(*FunctionName);
-//	
-//
-//	//UWidgetBlueprint* Blueprint = WidgetBlueprint();
-//
-//
-//	// Create the stub graph and add it to the list of functions to compile
-//
-//	UEdGraph* ChildStubGraph = NewNamedObject<UEdGraph>(Blueprint, Record.HandlerFunctionName);
-//	Blueprint->EventGraphs.Add(ChildStubGraph);
-//	ChildStubGraph->Schema = UEdGraphSchema_K2::StaticClass();
-//	ChildStubGraph->SetFlags(RF_Transient);
-//	MessageLog.NotifyIntermediateObjectCreation(ChildStubGraph, VisualAnimNode);
-//
-//	FKismetFunctionContext* StubContext = CreateFunctionContext();
-//	StubContext->SourceGraph = ChildStubGraph;
-//
-//	// A stub graph has no visual representation and is thus not suited to be debugged via the debugger
-//	StubContext->bCreateDebugData = false;
-//	StubContext->MarkAsInternalOrCppUseOnly();
-//
-//
-//
-//
-//
-//
-//
-//
-//
-//
-//
-//	
-//
-//	// Add a custom event in the graph
-//	//UK2Node_CustomEvent* EntryNode = SpawnIntermediateNode<UK2Node_CustomEvent>(VisualAnimNode, ConsolidatedEventGraph);
-//	//EntryNode->bInternalEvent = true;
-//	//EntryNode->CustomFunctionName = Record.HandlerFunctionName;
-//	UK2Node_FunctionEntry* EntryNode = SpawnIntermediateNode<UK2Node_FunctionEntry>(VisualAnimNode, ChildStubGraph);
-//	//EntryNode->SignatureClass = UObject::StaticClass();
-//	//EntryNode->SignatureName = Schema->FN_ExecuteUbergraphBase;
-//	EntryNode->CustomGeneratedFunctionName = Record.HandlerFunctionName;
-//	EntryNode->AllocateDefaultPins();
-//
-//	// The ExecChain is the current exec output pin in the linear chain
-//	UEdGraphPin* ExecChain = Schema->FindExecutionPin(*EntryNode, EGPD_Output);
-//
-//	// Create a struct member write node to store the parameters into the animation node
-//	UK2Node_StructMemberSet* AssignmentNode = SpawnIntermediateNode<UK2Node_StructMemberSet>(VisualAnimNode, ChildStubGraph);
-//	AssignmentNode->VariableReference.SetSelfMember(Record.NodeVariableProperty->GetFName());
-//	AssignmentNode->StructType = Record.NodeVariableProperty->Struct;
-//	AssignmentNode->AllocateDefaultPins();
-//
-//	// Wire up the variable node execution wires
-//	//UEdGraphPin* ExecVariablesIn = Schema->FindExecutionPin(*AssignmentNode, EGPD_Input);
-//	//ExecChain->MakeLinkTo(ExecVariablesIn);
-//	//ExecChain = Schema->FindExecutionPin(*AssignmentNode, EGPD_Output);
-//
-//	UK2Node_FunctionResult* ReturnNode = SpawnIntermediateNode<UK2Node_FunctionResult>(VisualAnimNode, ChildStubGraph);
-//	ReturnNode->AllocateDefaultPins();
-//
-//	// Wire up the variable node execution wires
-//	UEdGraphPin* ExecReturnIn = Schema->FindExecutionPin(*ReturnNode, EGPD_Input);
-//	ExecChain->MakeLinkTo(ExecReturnIn);
-//	ExecChain = Schema->FindExecutionPin(*ReturnNode, EGPD_Output);
-//
-//	for ( auto TargetPinIt = AssignmentNode->Pins.CreateIterator(); TargetPinIt; ++TargetPinIt )
-//	{
-//		UEdGraphPin* TargetPin = *TargetPinIt;
-//		FString PropertyNameStr = TargetPin->PinName;
-//		FName PropertyName(*PropertyNameStr);
-//
-//		// Does it get serviced by this handler?
-//		if ( FAnimNodeSinglePropertyHandler* SourceInfo = Record.ServicedProperties.Find(PropertyName) )
-//		{
-//			if ( TargetPin->PinType.bIsArray )
-//			{
-//				// Grab the array that we need to set members for
-//				UK2Node_StructMemberGet* FetchArrayNode = SpawnIntermediateNode<UK2Node_StructMemberGet>(VisualAnimNode, ChildStubGraph);
-//				FetchArrayNode->VariableReference.SetSelfMember(Record.NodeVariableProperty->GetFName());
-//				FetchArrayNode->StructType = Record.NodeVariableProperty->Struct;
-//				FetchArrayNode->AllocatePinsForSingleMemberGet(PropertyName);
-//
-//				UEdGraphPin* ArrayVariableNode = FetchArrayNode->FindPin(PropertyNameStr);
-//
-//				if ( SourceInfo->ArrayPins.Num() > 0 )
-//				{
-//					// Set each element in the array
-//					for ( auto SourcePinIt = SourceInfo->ArrayPins.CreateIterator(); SourcePinIt; ++SourcePinIt )
-//					{
-//						int32 ArrayIndex = SourcePinIt.Key();
-//						UEdGraphPin* SourcePin = SourcePinIt.Value();
-//
-//						// Create an array element set node
-//						UK2Node_CallArrayFunction* ArrayNode = SpawnIntermediateNode<UK2Node_CallArrayFunction>(VisualAnimNode, ChildStubGraph);
-//						ArrayNode->FunctionReference.SetExternalMember(GET_FUNCTION_NAME_CHECKED(UKismetArrayLibrary, Array_Set), UKismetArrayLibrary::StaticClass());
-//						ArrayNode->AllocateDefaultPins();
-//
-//						// Connect the execution chain
-//						ExecChain->MakeLinkTo(ArrayNode->GetExecPin());
-//						ExecChain = ArrayNode->GetThenPin();
-//
-//						// Connect the input array
-//						UEdGraphPin* TargetArrayPin = ArrayNode->FindPinChecked(TEXT("TargetArray"));
-//						TargetArrayPin->MakeLinkTo(ArrayVariableNode);
-//						ArrayNode->PinConnectionListChanged(TargetArrayPin);
-//
-//						// Set the array index
-//						UEdGraphPin* TargetIndexPin = ArrayNode->FindPinChecked(TEXT("Index"));
-//						TargetIndexPin->DefaultValue = FString::FromInt(ArrayIndex);
-//
-//						// Wire up the data input
-//						UEdGraphPin* TargetItemPin = ArrayNode->FindPinChecked(TEXT("Item"));
-//						TargetItemPin->CopyPersistentDataFromOldPin(*SourcePin);
-//						MessageLog.NotifyIntermediateObjectCreation(TargetItemPin, SourcePin);
-//						SourcePin->BreakAllPinLinks();
-//					}
-//				}
-//			}
-//			else
-//			{
-//				// Single property
-//				if ( SourceInfo->SinglePin != NULL )
-//				{
-//					UEdGraphPin* SourcePin = SourceInfo->SinglePin;
-//
-//					// First, add this pin to the user-defined pins
-//					TSharedPtr<FUserPinInfo> NewPinInfo = MakeShareable(new FUserPinInfo());
-//					NewPinInfo->PinName = TEXT("return");
-//					NewPinInfo->PinType = SourcePin->PinType;
-//					//UserDefinedPins.Add(NewPinInfo);
-//					TargetPin = ReturnNode->CreatePinFromUserDefinition(NewPinInfo);
-//
-//					//PropertiesBeingSet.Add(FName(*SourcePin->PinName));
-//					//TargetPin->CopyPersistentDataFromOldPin(*SourcePin);
-//					MessageLog.NotifyIntermediateObjectCreation(TargetPin, SourcePin);
-//					SourcePin->BreakAllPinLinks();
-//				}
-//			}
-//		}
-//	}
-//
-//	// Remove any unused pins from the assignment node to avoid smashing constant values
-//	//for ( int32 PinIndex = 0; PinIndex < AssignmentNode->ShowPinForProperties.Num(); ++PinIndex )
-//	//{
-//	//	FOptionalPinFromProperty& TestProperty = AssignmentNode->ShowPinForProperties[PinIndex];
-//	//	TestProperty.bShowPin = PropertiesBeingSet.Contains(TestProperty.PropertyName);
-//	//}
-//
-//	AssignmentNode->ReconstructNode();
-//	ReturnNode->ReconstructNode();
-//
-//	// TODO
-//}
-
 void FWidgetBlueprintCompiler::CreateEvaluationHandler(UWidgetGraphNode_Base* VisualAnimNode, FEvaluationHandlerRecord& Record)
 {
 	// Shouldn't create a handler if there is nothing to work with
@@ -616,6 +453,18 @@ void FWidgetBlueprintCompiler::CreateClassVariablesFromBlueprint()
 	TArray<UWidget*> Widgets;
 	Blueprint->WidgetTree->GetAllWidgets(Widgets);
 
+	// Sort the widgets alphabetically
+	{
+		struct FWidgetSorter
+		{
+			bool operator()(const UWidget& A, const UWidget& B) const
+			{
+				return B.GetFName() < A.GetFName();
+			}
+		};
+		Widgets.Sort(FWidgetSorter());
+	}
+
 	for ( UWidget* Widget : Widgets )
 	{
 		// Skip non-variable widgets
