UMG - Updating delegate bindings when widgets are renamed.

diff --git a/Engine/Source/Editor/UMGEditor/Private/WidgetBlueprintEditorUtils.cpp b/Engine/Source/Editor/UMGEditor/Private/WidgetBlueprintEditorUtils.cpp
--- a/Engine/Source/Editor/UMGEditor/Private/WidgetBlueprintEditorUtils.cpp
+++ b/Engine/Source/Editor/UMGEditor/Private/WidgetBlueprintEditorUtils.cpp
@@ -81,6 +81,8 @@ bool FWidgetBlueprintEditorUtils::RenameWidget(TSharedRef<FWidgetBlueprintEditor
 
 	if ( Widget )
 	{
+		const FScopedTransaction Transaction(LOCTEXT("RenameWidget", "Rename Widget"));
+
 		// Rename Template
 		Blueprint->Modify();
 		Widget->Modify();
@@ -93,6 +95,7 @@ bool FWidgetBlueprintEditorUtils::RenameWidget(TSharedRef<FWidgetBlueprintEditor
 			WidgetPreview->Rename(*NewNameStr);
 		}
 
+		// Find and update all variable references in the graph
 		TArray<UK2Node_Variable*> WidgetVarNodes;
 		FBlueprintEditorUtils::GetAllNodesOfClass<UK2Node_Variable>(Blueprint, WidgetVarNodes);
 		for ( int32 It = 0; It < WidgetVarNodes.Num(); It++ )
@@ -120,16 +123,21 @@ bool FWidgetBlueprintEditorUtils::RenameWidget(TSharedRef<FWidgetBlueprintEditor
 			}
 		}
 
+		// Find and update all binding references in the widget blueprint
+		for ( FDelegateEditorBinding& Binding : Blueprint->Bindings )
+		{
+			if ( Binding.ObjectName == OldNameStr )
+			{
+				Binding.ObjectName = NewNameStr;
+			}
+		}
+
 		// Validate child blueprints and adjust variable names to avoid a potential name collision
 		FBlueprintEditorUtils::ValidateBlueprintChildVariables(Blueprint, NewName);
 
 		// Refresh references and flush editors
 		FBlueprintEditorUtils::MarkBlueprintAsStructurallyModified(Blueprint);
 		bRenamed = true;
-
-		// Refresh references and flush editors
-		FBlueprintEditorUtils::MarkBlueprintAsStructurallyModified(Blueprint);
-		bRenamed = true;
 	}
 
 	return bRenamed;
