UMG - Now showing a name and IsVariable checkbox at the top of the details panel in the designer.

diff --git a/Engine/Source/Editor/UMGEditor/Private/Details/SWidgetDetailsView.cpp b/Engine/Source/Editor/UMGEditor/Private/Details/SWidgetDetailsView.cpp
--- a/Engine/Source/Editor/UMGEditor/Private/Details/SWidgetDetailsView.cpp
+++ b/Engine/Source/Editor/UMGEditor/Private/Details/SWidgetDetailsView.cpp
@@ -5,6 +5,7 @@
 
 #include "BlueprintEditor.h"
 #include "IDetailsView.h"
+#include "Kismet2NameValidators.h"
 
 #define LOCTEXT_NAMESPACE "UMG"
 
@@ -25,13 +26,59 @@ void SWidgetDetailsView::Construct(const FArguments& InArgs, TSharedPtr<FWidgetB
 	[
 		SNew(SVerticalBox)
 
+		+ SVerticalBox::Slot()
+		.AutoHeight()
+		.Padding(0, 0, 0, 6)
+		[
+			SNew(SHorizontalBox)
+			.Visibility(this, &SWidgetDetailsView::GetNameAreaVisibility)
+
+			+ SHorizontalBox::Slot()
+			.AutoWidth()
+			.Padding(0, 0, 3, 0)
+			.VAlign(VAlign_Center)
+			[
+				SNew(SImage)
+				.Image(this, &SWidgetDetailsView::GetNameIcon)
+			]
+
+			+ SHorizontalBox::Slot()
+			.AutoWidth()
+			.Padding(0, 0, 6, 0)
+			[
+				SNew(SBox)
+				.WidthOverride(200.0f)
+				.VAlign(VAlign_Center)
+				[
+					SAssignNew(NameTextBox, SEditableTextBox)
+					.SelectAllTextWhenFocused(true)
+					.HintText(LOCTEXT("Name", "Name"))
+					.Text(this, &SWidgetDetailsView::GetNameText)
+					.OnTextChanged(this, &SWidgetDetailsView::HandleNameTextChanged)
+					.OnTextCommitted(this, &SWidgetDetailsView::HandleNameTextCommitted)
+				]
+			]
+
+			+ SHorizontalBox::Slot()
+			.AutoWidth()
+			[
+				SNew(SCheckBox)
+				.IsChecked(this, &SWidgetDetailsView::GetIsVariable)
+				.OnCheckStateChanged(this, &SWidgetDetailsView::HandleIsVariableChanged)
+				[
+					SNew(STextBlock)
+					.Text(LOCTEXT("IsVariable", "Is Variable"))
+				]
+			]
+		]
+
 		+ SVerticalBox::Slot()
 		.FillHeight(1.0f)
 		[
 			PropertyView.ToSharedRef()
 		]
 	];
-	
+
 	BlueprintEditor.Pin()->OnSelectedWidgetsChanging.AddRaw(this, &SWidgetDetailsView::OnEditorSelectionChanging);
 	BlueprintEditor.Pin()->OnSelectedWidgetsChanged.AddRaw(this, &SWidgetDetailsView::OnEditorSelectionChanged);
 
@@ -94,4 +141,142 @@ void SWidgetDetailsView::ClearFocusIfOwned()
 	}
 }
 
+EVisibility SWidgetDetailsView::GetNameAreaVisibility() const
+{
+	if ( SelectedObjects.Num() == 1 )
+	{
+		return EVisibility::Visible;
+	}
+	
+	return EVisibility::Collapsed;
+}
+
+const FSlateBrush* SWidgetDetailsView::GetNameIcon() const
+{
+	if ( SelectedObjects.Num() == 1 )
+	{
+		UWidget* Widget = Cast<UWidget>(SelectedObjects[0].Get());
+		if ( Widget )
+		{
+			return Widget->GetEditorIcon();
+		}
+	}
+
+	return NULL;
+}
+
+FText SWidgetDetailsView::GetNameText() const
+{
+	if ( SelectedObjects.Num() == 1 )
+	{
+		UWidget* Widget = Cast<UWidget>(SelectedObjects[0].Get());
+		if ( Widget )
+		{
+			return FText::FromString(Widget->GetLabel());
+		}
+	}
+	
+	return FText::GetEmpty();
+}
+
+
+void SWidgetDetailsView::HandleNameTextChanged(const FText& Text)
+{
+	FText OutErrorMessage;
+	if ( !HandleVerifyNameTextChanged(Text, OutErrorMessage) )
+	{
+		NameTextBox->SetError(OutErrorMessage);
+	}
+	else
+	{
+		NameTextBox->SetError(FText::GetEmpty());
+	}
+}
+
+bool SWidgetDetailsView::HandleVerifyNameTextChanged(const FText& InText, FText& OutErrorMessage)
+{
+	if ( SelectedObjects.Num() == 1 )
+	{
+		UWidget* Widget = Cast<UWidget>(SelectedObjects[0].Get());
+
+		FString NewName = InText.ToString();
+
+		UWidgetBlueprint* Blueprint = BlueprintEditor.Pin()->GetWidgetBlueprintObj();
+		UWidget* ExistingWidget = Blueprint->WidgetTree->FindWidget(NewName);
+
+		FKismetNameValidator Validator(Blueprint);
+
+		const bool bUniqueNameForVariable = ( EValidatorResult::Ok == Validator.IsValid(NewName) );
+
+		if ( ( ExistingWidget != NULL && ExistingWidget != Widget ) || !bUniqueNameForVariable )
+		{
+			OutErrorMessage = LOCTEXT("NameConflict", "NameConflict");
+			return false;
+		}
+
+		return true;
+	}
+	else
+	{
+		return false;
+	}
+}
+
+void SWidgetDetailsView::HandleNameTextCommitted(const FText& Text, ETextCommit::Type CommitType)
+{
+	static bool IsReentrant = false;
+
+	if ( !IsReentrant )
+	{
+		IsReentrant = true;
+		if ( SelectedObjects.Num() == 1 )
+		{
+			UWidget* Widget = Cast<UWidget>(SelectedObjects[0].Get());
+			FWidgetBlueprintEditorUtils::RenameWidget(BlueprintEditor.Pin().ToSharedRef(), Widget->GetFName(), FName(*Text.ToString()));
+		}
+		IsReentrant = false;
+	}
+}
+
+ESlateCheckBoxState::Type SWidgetDetailsView::GetIsVariable() const
+{
+	if ( SelectedObjects.Num() == 1 )
+	{
+		UWidget* Widget = Cast<UWidget>(SelectedObjects[0].Get());
+		if ( Widget )
+		{
+			return Widget->bIsVariable ? ESlateCheckBoxState::Checked : ESlateCheckBoxState::Unchecked;
+		}
+	}
+
+	return ESlateCheckBoxState::Unchecked;
+}
+
+void SWidgetDetailsView::HandleIsVariableChanged(ESlateCheckBoxState::Type CheckState)
+{
+	if ( SelectedObjects.Num() == 1 )
+	{
+		TSharedPtr<FWidgetBlueprintEditor> BPEditor = BlueprintEditor.Pin();
+
+		UWidget* Widget = Cast<UWidget>(SelectedObjects[0].Get());
+		UWidgetBlueprint* Blueprint = BlueprintEditor.Pin()->GetWidgetBlueprintObj();
+		
+		FWidgetReference WidgetRef = FWidgetReference::FromTemplate(BPEditor, Blueprint->WidgetTree->FindWidget(Widget->GetName()));
+		if ( WidgetRef.IsValid() )
+		{
+			UWidget* Template = WidgetRef.GetTemplate();
+			UWidget* Preview = WidgetRef.GetPreview();
+
+			const FScopedTransaction Transaction(LOCTEXT("VariableToggle", "Variable Toggle"));
+			Template->Modify();
+			Preview->Modify();
+
+			Template->bIsVariable = Preview->bIsVariable = CheckState == ESlateCheckBoxState::Checked ? true : false;
+
+			// Refresh references and flush editors
+			FBlueprintEditorUtils::MarkBlueprintAsStructurallyModified(Blueprint);
+		}
+	}
+}
+
 #undef LOCTEXT_NAMESPACE
diff --git a/Engine/Source/Editor/UMGEditor/Private/Details/SWidgetDetailsView.h b/Engine/Source/Editor/UMGEditor/Private/Details/SWidgetDetailsView.h
--- a/Engine/Source/Editor/UMGEditor/Private/Details/SWidgetDetailsView.h
+++ b/Engine/Source/Editor/UMGEditor/Private/Details/SWidgetDetailsView.h
@@ -7,7 +7,9 @@
 
 #include "WidgetTemplate.h"
 
-/**  */
+/**
+ * The details view used in the designer section of the widget blueprint editor.
+ */
 class SWidgetDetailsView : public SCompoundWidget
 {
 public:
@@ -27,9 +29,26 @@ private:
 
 	void ClearFocusIfOwned();
 
+	EVisibility GetNameAreaVisibility() const;
+
+	const FSlateBrush* GetNameIcon() const;
+
+	FText GetNameText() const;
+
+	bool HandleVerifyNameTextChanged(const FText& InText, FText& OutErrorMessage);
+
+	void HandleNameTextChanged(const FText& Text);
+	void HandleNameTextCommitted(const FText& Text, ETextCommit::Type CommitType);
+
+	ESlateCheckBoxState::Type GetIsVariable() const;
+
+	void HandleIsVariableChanged(ESlateCheckBoxState::Type CheckState);
+
 private:
 	/** The editor that owns this details view */
 	TWeakPtr<class FWidgetBlueprintEditor> BlueprintEditor;
+
+	TSharedPtr<SEditableTextBox> NameTextBox;
 	
 	/** Property viewing widget */
 	TSharedPtr<class IDetailsView> PropertyView;
diff --git a/Engine/Source/Runtime/UMG/Public/Components/Widget.h b/Engine/Source/Runtime/UMG/Public/Components/Widget.h
--- a/Engine/Source/Runtime/UMG/Public/Components/Widget.h
+++ b/Engine/Source/Runtime/UMG/Public/Components/Widget.h
@@ -54,7 +54,7 @@ public:
 	 * Allows controls to be exposed as variables in a blueprint.  Not all controls need to be exposed
 	 * as variables, so this allows only the most useful ones to end up being exposed.
 	 */
-	UPROPERTY(EditDefaultsOnly, Category=Variable)
+	UPROPERTY()
 	bool bIsVariable;
 
 	/** Flag if the Widget was created from a blueprint */
