UMG - Adding editor icons for spacer, widget switcher and uniform grid.  Working on widget behaviors when selected, the widget switcher now temporarily displays the selected widget, even if that's not the active widget index.

diff --git a/Engine/Source/Editor/UMGEditor/Private/Designer/SDesignerView.cpp b/Engine/Source/Editor/UMGEditor/Private/Designer/SDesignerView.cpp
--- a/Engine/Source/Editor/UMGEditor/Private/Designer/SDesignerView.cpp
+++ b/Engine/Source/Editor/UMGEditor/Private/Designer/SDesignerView.cpp
@@ -181,7 +181,27 @@ SDesignerView::~SDesignerView()
 
 void SDesignerView::OnEditorSelectionChanged()
 {
-	SelectedWidgets = BlueprintEditor.Pin()->GetSelectedWidgets();
+	TSet<FWidgetReference> PendingSelectedWidgets = BlueprintEditor.Pin()->GetSelectedWidgets();
+
+	// Notify all widgets that are no longer selected.
+	for ( FWidgetReference& WidgetRef : SelectedWidgets )
+	{
+		if ( WidgetRef.IsValid() && !PendingSelectedWidgets.Contains(WidgetRef) )
+		{
+			WidgetRef.GetPreview()->Deselect();
+		}
+	}
+
+	// Notify all widgets that are now selected.
+	for ( FWidgetReference& WidgetRef : PendingSelectedWidgets )
+	{
+		if ( WidgetRef.IsValid() && !SelectedWidgets.Contains(WidgetRef) )
+		{
+			WidgetRef.GetPreview()->Select();
+		}
+	}
+
+	SelectedWidgets = PendingSelectedWidgets;
 
 	if ( SelectedWidgets.Num() > 0 )
 	{
@@ -846,6 +866,17 @@ void SDesignerView::Tick(const FGeometry& AllottedGeometry, const double InCurre
 				PreviewSlateWidget = CurrentWidget;
 				PreviewSurface->SetContent(CurrentWidget);
 			}
+
+			// Notify all selected widgets that they are selected, because there are new preview objects
+			// state may have been lost so this will recreate it if the widget does something special when
+			// selected.
+			for ( FWidgetReference& WidgetRef : SelectedWidgets )
+			{
+				if ( WidgetRef.IsValid() )
+				{
+					WidgetRef.GetPreview()->Select();
+				}
+			}
 		}
 		else
 		{
diff --git a/Engine/Source/Editor/UMGEditor/Private/Palette/SPaletteView.cpp b/Engine/Source/Editor/UMGEditor/Private/Palette/SPaletteView.cpp
--- a/Engine/Source/Editor/UMGEditor/Private/Palette/SPaletteView.cpp
+++ b/Engine/Source/Editor/UMGEditor/Private/Palette/SPaletteView.cpp
@@ -20,11 +20,11 @@
 
 #define LOCTEXT_NAMESPACE "UMG"
 
-class SWidgetTemplateItem : public SCompoundWidget
+class SPaletteViewItem : public SCompoundWidget
 {
 public:
 
-	SLATE_BEGIN_ARGS(SWidgetTemplateItem) {}
+	SLATE_BEGIN_ARGS(SPaletteViewItem) {}
 	SLATE_END_ARGS()
 
 	/**
@@ -78,7 +78,7 @@ public:
 			.Padding(2.0f)
 			.OnDragDetected(this, &FWidgetTemplateViewModel::OnDraggingWidgetTemplateItem)
 			[
-				SNew(SWidgetTemplateItem, Template)
+				SNew(SPaletteViewItem, Template)
 			];
 	}
 
diff --git a/Engine/Source/Runtime/UMG/Private/Components/TextBlock.cpp b/Engine/Source/Runtime/UMG/Private/Components/TextBlock.cpp
--- a/Engine/Source/Runtime/UMG/Private/Components/TextBlock.cpp
+++ b/Engine/Source/Runtime/UMG/Private/Components/TextBlock.cpp
@@ -123,7 +123,7 @@ void UTextBlock::HandleTextCommitted(const FText& InText, ETextCommit::Type Comm
 	//     Need a way to recognize one particular widget and forward things to them!
 }
 
-void UTextBlock::OnDesignerDoubleClicked()
+void UTextBlock::OnBeginEdit()
 {
 	MyEditorTextBlock->EnterEditingMode();
 }
diff --git a/Engine/Source/Runtime/UMG/Private/Components/UniformGridPanel.cpp b/Engine/Source/Runtime/UMG/Private/Components/UniformGridPanel.cpp
--- a/Engine/Source/Runtime/UMG/Private/Components/UniformGridPanel.cpp
+++ b/Engine/Source/Runtime/UMG/Private/Components/UniformGridPanel.cpp
@@ -2,6 +2,8 @@
 
 #include "UMGPrivatePCH.h"
 
+#define LOCTEXT_NAMESPACE "UMG"
+
 /////////////////////////////////////////////////////
 // UUniformGridPanel
 
@@ -62,3 +64,16 @@ void UUniformGridPanel::SyncronizeProperties()
 
 	MyUniformGridPanel->SetSlotPadding(SlotPadding);
 }
+
+#if WITH_EDITOR
+
+const FSlateBrush* UUniformGridPanel::GetEditorIcon()
+{
+	return FUMGStyle::Get().GetBrush("Widget.UniformGrid");
+}
+
+#endif
+
+/////////////////////////////////////////////////////
+
+#undef LOCTEXT_NAMESPACE
diff --git a/Engine/Source/Runtime/UMG/Private/Components/Widget.cpp b/Engine/Source/Runtime/UMG/Private/Components/Widget.cpp
--- a/Engine/Source/Runtime/UMG/Private/Components/Widget.cpp
+++ b/Engine/Source/Runtime/UMG/Private/Components/Widget.cpp
@@ -202,6 +202,30 @@ void UWidget::PostEditChangeProperty(FPropertyChangedEvent& PropertyChangedEvent
 	}
 }
 
+void UWidget::Select()
+{
+	OnSelected();
+
+	UWidget* Parent = GetParent();
+	while ( Parent != NULL )
+	{
+		Parent->OnDescendantSelected(this);
+		Parent = Parent->GetParent();
+	}
+}
+
+void UWidget::Deselect()
+{
+	OnDeselected();
+
+	UWidget* Parent = GetParent();
+	while ( Parent != NULL )
+	{
+		Parent->OnDescendantDeselected(this);
+		Parent = Parent->GetParent();
+	}
+}
+
 #endif
 
 bool UWidget::Modify(bool bAlwaysMarkDirty)
@@ -329,4 +353,23 @@ void UWidget::GatherAllChildren(UWidget* Root, TSet<UWidget*>& Children)
 			GatherAllChildren(ChildWidget, Children);
 		}
 	}
-}
\ No newline at end of file
+}
+
+UWidget* UWidget::FindChildContainingDescendant(UWidget* Root, UWidget* Descendant)
+{
+	UWidget* Parent = Descendant->GetParent();
+
+	while ( Parent != NULL )
+	{
+		// If the Descendant's parent is the root, then the child containing the descendant is the descendant.
+		if ( Parent == Root )
+		{
+			return Descendant;
+		}
+
+		Descendant = Parent;
+		Parent = Parent->GetParent();
+	}
+
+	return NULL;
+}
diff --git a/Engine/Source/Runtime/UMG/Private/Components/WidgetSwitcher.cpp b/Engine/Source/Runtime/UMG/Private/Components/WidgetSwitcher.cpp
--- a/Engine/Source/Runtime/UMG/Private/Components/WidgetSwitcher.cpp
+++ b/Engine/Source/Runtime/UMG/Private/Components/WidgetSwitcher.cpp
@@ -2,6 +2,8 @@
 
 #include "UMGPrivatePCH.h"
 
+#define LOCTEXT_NAMESPACE "UMG"
+
 /////////////////////////////////////////////////////
 // UWidgetSwitcher
 
@@ -90,3 +92,39 @@ void UWidgetSwitcher::SyncronizeProperties()
 
 	SetActiveWidgetIndex(ActiveWidgetIndex);
 }
+
+#if WITH_EDITOR
+
+const FSlateBrush* UWidgetSwitcher::GetEditorIcon()
+{
+	return FUMGStyle::Get().GetBrush("Widget.WidgetSwitcher");
+}
+
+void UWidgetSwitcher::OnDescendantSelected(UWidget* DescendantWidget)
+{
+	// Temporarily sets the active child to the selected child to make
+	// dragging and dropping easier in the editor.
+	UWidget* SelectedChild = UWidget::FindChildContainingDescendant(this, DescendantWidget);
+	if ( SelectedChild )
+	{
+		int32 OverrideIndex = GetChildIndex(SelectedChild);
+		if ( OverrideIndex != -1 && MyWidgetSwitcher.IsValid() )
+		{
+			MyWidgetSwitcher->SetActiveWidgetIndex(OverrideIndex);
+		}
+	}
+}
+
+void UWidgetSwitcher::OnDescendantDeselected(UWidget* DescendantWidget)
+{
+	if ( MyWidgetSwitcher.IsValid() )
+	{
+		MyWidgetSwitcher->SetActiveWidgetIndex(ActiveWidgetIndex);
+	}
+}
+
+#endif
+
+/////////////////////////////////////////////////////
+
+#undef LOCTEXT_NAMESPACE
diff --git a/Engine/Source/Runtime/UMG/Private/UMGStyle.cpp b/Engine/Source/Runtime/UMG/Private/UMGStyle.cpp
--- a/Engine/Source/Runtime/UMG/Private/UMGStyle.cpp
+++ b/Engine/Source/Runtime/UMG/Private/UMGStyle.cpp
@@ -45,7 +45,7 @@ TSharedRef< FSlateStyleSet > FUMGStyle::Create()
 	
 	Style->Set("Widget", new IMAGE_BRUSH(TEXT("Widget"), Icon16x16));
 	Style->Set("Widget.CheckBox", new IMAGE_BRUSH(TEXT("CheckBox"), Icon16x16));
-	Style->Set("ClassIcon.Button", new IMAGE_BRUSH(TEXT("Button"), Icon16x16));
+	Style->Set("Widget.Button", new IMAGE_BRUSH(TEXT("Button"), Icon16x16));
 	Style->Set("Widget.EditableTextBox", new IMAGE_BRUSH(TEXT("EditableTextBox"), Icon16x16));
 	Style->Set("Widget.EditableText", new IMAGE_BRUSH(TEXT("EditableText"), Icon16x16));
 	Style->Set("Widget.HorizontalBox", new IMAGE_BRUSH(TEXT("HorizontalBox"), Icon16x16));
@@ -60,6 +60,8 @@ TSharedRef< FSlateStyleSet > FUMGStyle::Create()
 	Style->Set("Widget.ProgressBar", new IMAGE_BRUSH(TEXT("ProgressBar"), Icon16x16));
 	Style->Set("Widget.MenuAnchor", new IMAGE_BRUSH(TEXT("MenuAnchor"), Icon16x16));
 	Style->Set("Widget.ScrollBar", new IMAGE_BRUSH(TEXT("ScrollBar"), Icon16x16));
+	Style->Set("Widget.UniformGrid", new IMAGE_BRUSH(TEXT("UniformGrid"), Icon16x16));
+	Style->Set("Widget.WidgetSwitcher", new IMAGE_BRUSH(TEXT("WidgetSwitcher"), Icon16x16));
 
 	Style->Set("Widget.UserWidget", new IMAGE_BRUSH(TEXT("UserWidget"), Icon16x16));
 
diff --git a/Engine/Source/Runtime/UMG/Public/Components/TextBlock.h b/Engine/Source/Runtime/UMG/Public/Components/TextBlock.h
--- a/Engine/Source/Runtime/UMG/Public/Components/TextBlock.h
+++ b/Engine/Source/Runtime/UMG/Public/Components/TextBlock.h
@@ -98,7 +98,7 @@ public:
 
 	void HandleTextCommitted(const FText& InText, ETextCommit::Type CommitteType);
 
-	virtual void OnDesignerDoubleClicked() override;
+	virtual void OnBeginEdit() override;
 #endif
 
 protected:
diff --git a/Engine/Source/Runtime/UMG/Public/Components/UniformGridPanel.h b/Engine/Source/Runtime/UMG/Public/Components/UniformGridPanel.h
--- a/Engine/Source/Runtime/UMG/Public/Components/UniformGridPanel.h
+++ b/Engine/Source/Runtime/UMG/Public/Components/UniformGridPanel.h
@@ -26,6 +26,10 @@ class UMG_API UUniformGridPanel : public UPanelWidget
 	virtual void SyncronizeProperties() override;
 	// End of UWidget interface
 
+#if WITH_EDITOR
+	virtual const FSlateBrush* GetEditorIcon() override;
+#endif
+
 protected:
 
 	// UPanelWidget
diff --git a/Engine/Source/Runtime/UMG/Public/Components/Widget.h b/Engine/Source/Runtime/UMG/Public/Components/Widget.h
--- a/Engine/Source/Runtime/UMG/Public/Components/Widget.h
+++ b/Engine/Source/Runtime/UMG/Public/Components/Widget.h
@@ -171,12 +171,6 @@ public:
 	/** Gets the editor icon */
 	virtual const FSlateBrush* GetEditorIcon();
 
-	//TODO UMG These are not used yet, and may never be.
-	virtual void OnDesignerSelected() { }
-	virtual void OnDesignerDeselected() { }
-	virtual void OnDesignerClicked() { }
-	virtual void OnDesignerDoubleClicked() { }
-
 	/** Gets a widget representing the tiny preview of the toolbox */
 	virtual TSharedRef<SWidget> GetToolboxPreviewWidget() const;
 	
@@ -186,6 +180,20 @@ public:
 	// UObject interface
 	virtual void PostEditChangeProperty(struct FPropertyChangedEvent& PropertyChangedEvent) override;
 	// End of UObject interface
+
+	// Begin Designer contextual events
+	void Select();
+	void Deselect();
+
+	virtual void OnSelected() { }
+	virtual void OnDeselected() { }
+
+	virtual void OnDescendantSelected(UWidget* DescendantWidget) { }
+	virtual void OnDescendantDeselected(UWidget* DescendantWidget) { }
+
+	virtual void OnBeginEdit() { }
+	virtual void OnEndEdit() { }
+	// End Designer contextual events
 #endif
 
 	// Utility methods
@@ -197,6 +205,7 @@ public:
 
 	static void GatherChildren(UWidget* Root, TSet<UWidget*>& Children);
 	static void GatherAllChildren(UWidget* Root, TSet<UWidget*>& Children);
+	static UWidget* FindChildContainingDescendant(UWidget* Root, UWidget* Descendant);
 
 protected:
 	/** Function implemented by all subclasses of UWidget is called when the underlying SWidget needs to be constructed. */
diff --git a/Engine/Source/Runtime/UMG/Public/Components/WidgetSwitcher.h b/Engine/Source/Runtime/UMG/Public/Components/WidgetSwitcher.h
--- a/Engine/Source/Runtime/UMG/Public/Components/WidgetSwitcher.h
+++ b/Engine/Source/Runtime/UMG/Public/Components/WidgetSwitcher.h
@@ -31,6 +31,12 @@ class UMG_API UWidgetSwitcher : public UPanelWidget
 	
 	void SyncronizeProperties();
 
+#if WITH_EDITOR
+	virtual const FSlateBrush* GetEditorIcon() override;
+	virtual void OnDescendantSelected(UWidget* DescendantWidget) override;
+	virtual void OnDescendantDeselected(UWidget* DescendantWidget) override;
+#endif
+
 protected:
 
 	// UPanelWidget
