UMG - Now providing cursor feedback when hovering over an area that wont accept the droppable widget.

diff --git a/Engine/Source/Editor/UMGEditor/Private/Designer/SDesignerView.cpp b/Engine/Source/Editor/UMGEditor/Private/Designer/SDesignerView.cpp
--- a/Engine/Source/Editor/UMGEditor/Private/Designer/SDesignerView.cpp
+++ b/Engine/Source/Editor/UMGEditor/Private/Designer/SDesignerView.cpp
@@ -939,10 +939,11 @@ void SDesignerView::OnDragEnter(const FGeometry& MyGeometry, const FDragDropEven
 
 void SDesignerView::OnDragLeave(const FDragDropEvent& DragDropEvent)
 {
-	TSharedPtr<FWidgetTemplateDragDropOp> DragDropOp = DragDropEvent.GetOperationAs<FWidgetTemplateDragDropOp>();
-	if ( DragDropOp.IsValid() )
+	TSharedPtr<FDecoratedDragDropOp> DecoratedDragDropOp = DragDropEvent.GetOperationAs<FDecoratedDragDropOp>();
+	if ( DecoratedDragDropOp.IsValid() )
 	{
-		DragDropOp->ResetToDefaultToolTip();
+		DecoratedDragDropOp->SetCursorOverride(TOptional<EMouseCursor::Type>());
+		DecoratedDragDropOp->ResetToDefaultToolTip();
 	}
 	
 	if (DropPreviewWidget)
@@ -1019,6 +1020,8 @@ UWidget* SDesignerView::ProcessDropAndAddWidget(const FGeometry& MyGeometry, con
 	TSharedPtr<FWidgetTemplateDragDropOp> TemplateDragDropOp = DragDropEvent.GetOperationAs<FWidgetTemplateDragDropOp>();
 	if ( TemplateDragDropOp.IsValid() )
 	{
+		TemplateDragDropOp->SetCursorOverride(TOptional<EMouseCursor::Type>());
+
 		// If there's no root widget go ahead and add the widget into the root slot.
 		if ( BP->WidgetTree->RootWidget == NULL )
 		{
@@ -1095,6 +1098,8 @@ UWidget* SDesignerView::ProcessDropAndAddWidget(const FGeometry& MyGeometry, con
 			}
 			else
 			{
+				TemplateDragDropOp->SetCursorOverride(EMouseCursor::SlashedCircle);
+
 				// TODO UMG ERROR Slot can not be created because maybe the max children has been reached.
 				//          Maybe we can traverse the hierarchy and add it to the first parent that will accept it?
 			}
@@ -1104,11 +1109,17 @@ UWidget* SDesignerView::ProcessDropAndAddWidget(const FGeometry& MyGeometry, con
 				Transaction.Cancel();
 			}
 		}
+		else
+		{
+			TemplateDragDropOp->SetCursorOverride(EMouseCursor::SlashedCircle);
+		}
 	}
 
 	// Attempt to deal with moving widgets from a drag operation.
 	if ( SelectedDragDropOp.IsValid() )
 	{
+		SelectedDragDropOp->SetCursorOverride(TOptional<EMouseCursor::Type>());
+
 		if ( Target && Target->IsA(UPanelWidget::StaticClass()) )
 		{
 			UPanelWidget* NewParent = Cast<UPanelWidget>(Target);
@@ -1159,6 +1170,8 @@ UWidget* SDesignerView::ProcessDropAndAddWidget(const FGeometry& MyGeometry, con
 			}
 			else
 			{
+				SelectedDragDropOp->SetCursorOverride(EMouseCursor::SlashedCircle);
+
 				// TODO UMG ERROR Slot can not be created because maybe the max children has been reached.
 				//          Maybe we can traverse the hierarchy and add it to the first parent that will accept it?
 			}
@@ -1168,6 +1181,10 @@ UWidget* SDesignerView::ProcessDropAndAddWidget(const FGeometry& MyGeometry, con
 				Transaction.Cancel();
 			}
 		}
+		else
+		{
+			SelectedDragDropOp->SetCursorOverride(EMouseCursor::SlashedCircle);
+		}
 	}
 	
 	return NULL;
