UMG - The widget blueprint factory no longer appears in the new asset menu if UMG is not enabled.

diff --git a/Engine/Source/Editor/UMGEditor/Classes/WidgetBlueprintFactory.h b/Engine/Source/Editor/UMGEditor/Classes/WidgetBlueprintFactory.h
--- a/Engine/Source/Editor/UMGEditor/Classes/WidgetBlueprintFactory.h
+++ b/Engine/Source/Editor/UMGEditor/Classes/WidgetBlueprintFactory.h
@@ -19,6 +19,7 @@ class UWidgetBlueprintFactory : public UFactory
 
 	// Begin UFactory Interface
 	virtual bool ConfigureProperties() override;
+	virtual bool ShouldShowInNewMenu() const override;
 	virtual UObject* FactoryCreateNew(UClass* Class, UObject* InParent, FName Name, EObjectFlags Flags, UObject* Context, FFeedbackContext* Warn, FName CallingContext) override;
 	virtual UObject* FactoryCreateNew(UClass* Class, UObject* InParent, FName Name, EObjectFlags Flags, UObject* Context, FFeedbackContext* Warn) override;
 	// Begin UFactory Interface	
diff --git a/Engine/Source/Editor/UMGEditor/Private/WidgetBlueprintFactory.cpp b/Engine/Source/Editor/UMGEditor/Private/WidgetBlueprintFactory.cpp
--- a/Engine/Source/Editor/UMGEditor/Private/WidgetBlueprintFactory.cpp
+++ b/Engine/Source/Editor/UMGEditor/Private/WidgetBlueprintFactory.cpp
@@ -24,7 +24,12 @@ bool UWidgetBlueprintFactory::ConfigureProperties()
 {
 	// TODO Make config dialog like the anim blueprints.
 	return true;
-};
+}
+
+bool UWidgetBlueprintFactory::ShouldShowInNewMenu() const
+{
+	return GetDefault<UEditorExperimentalSettings>()->bUnrealMotionGraphics;
+}
 
 UObject* UWidgetBlueprintFactory::FactoryCreateNew(UClass* Class, UObject* InParent, FName Name, EObjectFlags Flags, UObject* Context, FFeedbackContext* Warn, FName CallingContext)
 {
