UMG - Rename Preview before renaming the template widget so the preview widget can be found after the template is renamed.

diff --git a/Engine/Source/Editor/UMGEditor/Private/WidgetBlueprintEditorUtils.cpp b/Engine/Source/Editor/UMGEditor/Private/WidgetBlueprintEditorUtils.cpp
--- a/Engine/Source/Editor/UMGEditor/Private/WidgetBlueprintEditorUtils.cpp
+++ b/Engine/Source/Editor/UMGEditor/Private/WidgetBlueprintEditorUtils.cpp
@@ -86,16 +86,19 @@ bool FWidgetBlueprintEditorUtils::RenameWidget(TSharedRef<FWidgetBlueprintEditor
 		// Rename Template
 		Blueprint->Modify();
 		Widget->Modify();
-		Widget->Rename(*NewNameStr);
 
-		// Rename Preview
+
+		// Rename Preview before renaming the template widget so the preview widget can be found
 		UWidget* WidgetPreview = FWidgetReference::FromTemplate(BlueprintEditor, Widget).GetPreview();
-		if ( WidgetPreview )
+		if(WidgetPreview)
 		{
 			WidgetPreview->Rename(*NewNameStr);
 		}
 
 		// Find and update all variable references in the graph
+
+		Widget->Rename(*NewNameStr);
+	
 		TArray<UK2Node_Variable*> WidgetVarNodes;
 		FBlueprintEditorUtils::GetAllNodesOfClass<UK2Node_Variable>(Blueprint, WidgetVarNodes);
 		for ( int32 It = 0; It < WidgetVarNodes.Num(); It++ )
