UMG - The notification system for property changes and migrating widget properties back to the template is now divorced from the one the blueprint editor uses for detecting general changes to the nodes - Fixes a lot of crash bugs.  Also now calling the parent blueprint editors Blueprint updated function fixing a lot of issues with the graph not updating.

diff --git a/Engine/Source/Editor/UMGEditor/Private/Details/SWidgetDetailsView.cpp b/Engine/Source/Editor/UMGEditor/Private/Details/SWidgetDetailsView.cpp
--- a/Engine/Source/Editor/UMGEditor/Private/Details/SWidgetDetailsView.cpp
+++ b/Engine/Source/Editor/UMGEditor/Private/Details/SWidgetDetailsView.cpp
@@ -16,8 +16,10 @@ void SWidgetDetailsView::Construct(const FArguments& InArgs, TSharedPtr<FWidgetB
 	// Create a property view
 	FPropertyEditorModule& EditModule = FModuleManager::Get().GetModuleChecked<FPropertyEditorModule>("PropertyEditor");
 
-	FNotifyHook* NotifyHook = InBlueprintEditor.Get();
+	FNotifyHook* NotifyHook = this;
 	FDetailsViewArgs DetailsViewArgs( /*bUpdateFromSelection=*/ false, /*bLockable=*/ false, /*bAllowSearch=*/ true, /*bObjectsUseNameArea=*/ true, /*bHideSelectionTip=*/ true, /*InNotifyHook=*/ NotifyHook, /*InSearchInitialKeyFocus=*/ false, /*InViewIdentifier=*/ NAME_None);
+
+	// We hide the actor name area because we're providing our own.
 	DetailsViewArgs.bHideActorNameArea = true;
 
 	PropertyView = EditModule.CreateDetailView(DetailsViewArgs);
@@ -279,4 +281,23 @@ void SWidgetDetailsView::HandleIsVariableChanged(ESlateCheckBoxState::Type Check
 	}
 }
 
+void SWidgetDetailsView::NotifyPreChange(FEditPropertyChain* PropertyAboutToChange)
+{
+	TSharedPtr<FWidgetBlueprintEditor> Editor = BlueprintEditor.Pin();
+
+	const bool bIsModify = true;
+	Editor->MigrateFromChain(PropertyAboutToChange, bIsModify);
+}
+
+void SWidgetDetailsView::NotifyPostChange(const FPropertyChangedEvent& PropertyChangedEvent, FEditPropertyChain* PropertyThatChanged)
+{
+	if ( PropertyChangedEvent.ChangeType != EPropertyChangeType::Interactive )
+	{
+		TSharedPtr<FWidgetBlueprintEditor> Editor = BlueprintEditor.Pin();
+
+		const bool bIsModify = false;
+		Editor->MigrateFromChain(PropertyThatChanged, bIsModify);
+	}
+}
+
 #undef LOCTEXT_NAMESPACE
diff --git a/Engine/Source/Editor/UMGEditor/Private/Details/SWidgetDetailsView.h b/Engine/Source/Editor/UMGEditor/Private/Details/SWidgetDetailsView.h
--- a/Engine/Source/Editor/UMGEditor/Private/Details/SWidgetDetailsView.h
+++ b/Engine/Source/Editor/UMGEditor/Private/Details/SWidgetDetailsView.h
@@ -10,7 +10,7 @@
 /**
  * The details view used in the designer section of the widget blueprint editor.
  */
-class SWidgetDetailsView : public SCompoundWidget
+class SWidgetDetailsView : public SCompoundWidget, public FNotifyHook
 {
 public:
 	SLATE_BEGIN_ARGS( SWidgetDetailsView ){}
@@ -19,12 +19,22 @@ public:
 	void Construct(const FArguments& InArgs, TSharedPtr<FWidgetBlueprintEditor> InBlueprintEditor);
 	virtual ~SWidgetDetailsView();
 	
+	/** Gets the property view for this details panel */
 	TSharedPtr<class IDetailsView> GetPropertyView() const { return PropertyView; }
+
+	// FNotifyHook interface
+	virtual void NotifyPreChange(FEditPropertyChain* PropertyAboutToChange) override;
+	virtual void NotifyPostChange(const FPropertyChangedEvent& PropertyChangedEvent, FEditPropertyChain* PropertyThatChanged) override;
+	// End of FNotifyHook interface
 	
 private:
+	/** Registers the designer specific customizations */
 	void RegisterCustomizations();
 
+	/** Handles the widget selection changing event in the editor, updates the details panel accordingly. */
 	void OnEditorSelectionChanging();
+
+	/** Handles the widget selection changed event in the editor, updates the details panel accordingly. */
 	void OnEditorSelectionChanged();
 
 	void ClearFocusIfOwned();
diff --git a/Engine/Source/Editor/UMGEditor/Private/WidgetBlueprintEditor.cpp b/Engine/Source/Editor/UMGEditor/Private/WidgetBlueprintEditor.cpp
--- a/Engine/Source/Editor/UMGEditor/Private/WidgetBlueprintEditor.cpp
+++ b/Engine/Source/Editor/UMGEditor/Private/WidgetBlueprintEditor.cpp
@@ -159,6 +159,8 @@ const TSet<FWidgetReference>& FWidgetBlueprintEditor::GetSelectedWidgets() const
 
 void FWidgetBlueprintEditor::OnBlueprintChanged(UBlueprint* InBlueprint)
 {
+	FBlueprintEditor::OnBlueprintChanged(InBlueprint);
+
 	if ( InBlueprint )
 	{
 		// Rebuilding the preview can force objects to be recreated, so the selection may need to be updated.
@@ -340,11 +342,6 @@ void FWidgetBlueprintEditor::AddReferencedObjects( FReferenceCollector& Collecto
 	Collector.AddReferencedObject( Preview );
 }
 
-void FWidgetBlueprintEditor::NotifyPreChange(FEditPropertyChain* PropertyAboutToChange)
-{
-	MigrateFromChain(PropertyAboutToChange, true);
-}
-
 void FWidgetBlueprintEditor::MigrateFromChain(FEditPropertyChain* PropertyThatChanged, bool bIsModify)
 {
 	UWidgetBlueprint* Blueprint = GetWidgetBlueprintObj();
@@ -371,16 +368,6 @@ void FWidgetBlueprintEditor::MigrateFromChain(FEditPropertyChain* PropertyThatCh
 	}
 }
 
-void FWidgetBlueprintEditor::NotifyPostChange(const FPropertyChangedEvent& PropertyChangedEvent, FEditPropertyChain* PropertyThatChanged)
-{
-	//Super::NotifyPostChange(PropertyChangedEvent, PropertyThatChanged);
-	
-	if ( PropertyChangedEvent.ChangeType != EPropertyChangeType::Interactive )
-	{
-		MigrateFromChain(PropertyThatChanged, false);
-	}
-}
-
 UWidgetBlueprint* FWidgetBlueprintEditor::GetWidgetBlueprintObj() const
 {
 	return Cast<UWidgetBlueprint>(GetBlueprintObj());
diff --git a/Engine/Source/Editor/UMGEditor/Private/WidgetBlueprintEditor.h b/Engine/Source/Editor/UMGEditor/Private/WidgetBlueprintEditor.h
--- a/Engine/Source/Editor/UMGEditor/Private/WidgetBlueprintEditor.h
+++ b/Engine/Source/Editor/UMGEditor/Private/WidgetBlueprintEditor.h
@@ -27,8 +27,6 @@ public:
 	void InitWidgetBlueprintEditor(const EToolkitMode::Type Mode, const TSharedPtr< IToolkitHost >& InitToolkitHost, const TArray<UBlueprint*>& InBlueprints, bool bShouldOpenInDefaultsMode);
 
 	virtual void Tick(float DeltaTime) override;
-	virtual void NotifyPreChange(class FEditPropertyChain* PropertyAboutToChange) override;
-	virtual void NotifyPostChange(const FPropertyChangedEvent& PropertyChangedEvent, class FEditPropertyChain* PropertyThatChanged) override;
 
 	/** FGCObjectInterface */
 	virtual void AddReferencedObjects( FReferenceCollector& Collector ) override;
@@ -56,6 +54,9 @@ public:
 
 	TSharedPtr<class FWidgetBlueprintEditorToolbar> GetWidgetToolbarBuilder() { return WidgetToolbar; }
 
+	/** Migrate a property change from the preview GUI to the template GUI. */
+	void MigrateFromChain(FEditPropertyChain* PropertyThatChanged, bool bIsModify);
+
 public:
 	/** Fires whenever the selected set of widgets changing */
 	FOnSelectedWidgetsChanged OnSelectedWidgetsChanging;
@@ -94,10 +95,7 @@ private:
 
 	/** Tick the current preview GUI object */
 	void UpdatePreview(UBlueprint* InBlueprint, bool bInForceFullUpdate);
-
-	/** Migrate a property change from the preview GUI to the template GUI. */
-	void MigrateFromChain(FEditPropertyChain* PropertyThatChanged, bool bIsModify);
-
+	
 	/**
 	 * Gets the default movie scene which is used when there is no 
 	 * animation data present on the widget blueprint
