UMG - Now allowing simplified reparenting in the hierarchy view via drag/drop.

diff --git a/Engine/Source/Editor/UMGEditor/Private/Designer/SDesignerView.cpp b/Engine/Source/Editor/UMGEditor/Private/Designer/SDesignerView.cpp
--- a/Engine/Source/Editor/UMGEditor/Private/Designer/SDesignerView.cpp
+++ b/Engine/Source/Editor/UMGEditor/Private/Designer/SDesignerView.cpp
@@ -1153,11 +1153,7 @@ UWidget* SDesignerView::ProcessDropAndAddWidget(const FGeometry& MyGeometry, con
 			{
 				FWidgetBlueprintEditorUtils::ImportPropertiesFromText(Slot, SelectedDragDropOp->ExportedSlotProperties);
 
-				//TODO UMG Migrate existing slot info
 				Slot->SetDesiredPosition(LocalPosition - SelectedWidgetContextMenuLocation);
-				//Slot->SetDesiredSize(FVector2D(150, 30));
-				//@TODO UMG When we add a child blindly we need to default the slot size to the preferred size of the widget if the container supports such things.
-				//@TODO UMG We may need a desired size canvas, where the slots have no size, they only give you position, alternatively, maybe slots that don't clip, so center is still easy.
 
 				DropPreviewParent = NewParent;
 
diff --git a/Engine/Source/Editor/UMGEditor/Private/Hierarchy/SHierarchyView.cpp b/Engine/Source/Editor/UMGEditor/Private/Hierarchy/SHierarchyView.cpp
--- a/Engine/Source/Editor/UMGEditor/Private/Hierarchy/SHierarchyView.cpp
+++ b/Engine/Source/Editor/UMGEditor/Private/Hierarchy/SHierarchyView.cpp
@@ -223,7 +223,9 @@ void SHierarchyView::WidgetHierarchy_OnGetChildren(UWidget* InParent, TArray< UW
 
 TSharedRef< ITableRow > SHierarchyView::WidgetHierarchy_OnGenerateRow(UWidget* InItem, const TSharedRef<STableViewBase>& OwnerTable)
 {
-	return SNew(SHierarchyViewItem, OwnerTable, BlueprintEditor.Pin(), InItem)
+	FWidgetReference WidgetRef = FWidgetReference::FromTemplate(BlueprintEditor.Pin(), InItem);
+
+	return SNew(SHierarchyViewItem, OwnerTable, BlueprintEditor.Pin(), WidgetRef)
 		.HighlightText(this, &SHierarchyView::GetSearchText);
 }
 
diff --git a/Engine/Source/Editor/UMGEditor/Private/Hierarchy/SHierarchyViewItem.cpp b/Engine/Source/Editor/UMGEditor/Private/Hierarchy/SHierarchyViewItem.cpp
--- a/Engine/Source/Editor/UMGEditor/Private/Hierarchy/SHierarchyViewItem.cpp
+++ b/Engine/Source/Editor/UMGEditor/Private/Hierarchy/SHierarchyViewItem.cpp
@@ -19,7 +19,68 @@
 
 #define LOCTEXT_NAMESPACE "UMG"
 
-void SHierarchyViewItem::Construct(const FArguments& InArgs, const TSharedRef< STableViewBase >& InOwnerTableView, TSharedPtr<FWidgetBlueprintEditor> InBlueprintEditor, UWidget* InItem)
+/**
+ * 
+ */
+class FHierarchyWidgetDragDropOp : public FDecoratedDragDropOp
+{
+public:
+	DRAG_DROP_OPERATOR_TYPE(FHierarchyWidgetDragDropOp, FDecoratedDragDropOp)
+
+	virtual ~FHierarchyWidgetDragDropOp();
+
+	/** The slot properties for the old slot the widget was in, is used to attempt to reapply the same layout information */
+	TMap<FName, FString> ExportedSlotProperties;
+
+	/** The widget being dragged and dropped */
+	FWidgetReference Widget;
+
+	/** The original parent of the widget. */
+	UWidget* WidgetParent;
+
+	/** The widget being dragged and dropped */
+	FScopedTransaction* Transaction;
+
+	/** Constructs a new drag/drop operation */
+	static TSharedRef<FHierarchyWidgetDragDropOp> New(UWidgetBlueprint* Blueprint, FWidgetReference InWidget);
+};
+
+TSharedRef<FHierarchyWidgetDragDropOp> FHierarchyWidgetDragDropOp::New(UWidgetBlueprint* Blueprint, FWidgetReference InWidget)
+{
+	TSharedRef<FHierarchyWidgetDragDropOp> Operation = MakeShareable(new FHierarchyWidgetDragDropOp());
+	Operation->Widget = InWidget;
+	Operation->DefaultHoverText = FText::FromString(InWidget.GetTemplate()->GetLabel());
+	Operation->CurrentHoverText = FText::FromString(InWidget.GetTemplate()->GetLabel());
+	Operation->Construct();
+
+	FWidgetBlueprintEditorUtils::ExportPropertiesToText(InWidget.GetTemplate()->Slot, Operation->ExportedSlotProperties);
+
+	Operation->Transaction = new FScopedTransaction(LOCTEXT("Designer_MoveWidget", "Move Widget"));
+
+	Blueprint->WidgetTree->SetFlags(RF_Transactional);
+	Blueprint->WidgetTree->Modify();
+
+	UWidget* Widget = Operation->Widget.GetTemplate();
+	Widget->Modify();
+
+	Operation->WidgetParent = Widget->GetParent();
+
+	if ( Operation->WidgetParent )
+	{
+		Operation->WidgetParent->Modify();
+	}
+
+	return Operation;
+}
+
+FHierarchyWidgetDragDropOp::~FHierarchyWidgetDragDropOp()
+{
+	delete Transaction;
+}
+
+//////////////////////////////////////////////////////////////////////////
+
+void SHierarchyViewItem::Construct(const FArguments& InArgs, const TSharedRef< STableViewBase >& InOwnerTableView, TSharedPtr<FWidgetBlueprintEditor> InBlueprintEditor, FWidgetReference InItem)
 {
 	BlueprintEditor = InBlueprintEditor;
 	Item = InItem;
@@ -27,6 +88,10 @@ void SHierarchyViewItem::Construct(const FArguments& InArgs, const TSharedRef< S
 	STableRow< UWidget* >::Construct(
 		STableRow< UWidget* >::FArguments()
 		.Padding(2.0f)
+		.OnDragDetected(this, &SHierarchyViewItem::HandleDragDetected)
+		.OnDragEnter(this, &SHierarchyViewItem::HandleDragEnter)
+		.OnDragLeave(this, &SHierarchyViewItem::HandleDragLeave)
+		.OnDrop(this, &SHierarchyViewItem::HandleDrop)
 		.Content()
 		[
 			SNew(SHorizontalBox)
@@ -37,7 +102,7 @@ void SHierarchyViewItem::Construct(const FArguments& InArgs, const TSharedRef< S
 			[
 				SNew(SImage)
 				.ColorAndOpacity(FLinearColor(1,1,1,0.5))
-				.Image(Item->GetEditorIcon())
+				.Image(Item.GetTemplate()->GetEditorIcon())
 			]
 
 			// Name of the widget
@@ -69,7 +134,7 @@ bool SHierarchyViewItem::OnVerifyNameTextChanged(const FText& InText, FText& Out
 
 	const bool bUniqueNameForVariable = ( EValidatorResult::Ok == Validator.IsValid(NewName) );
 	
-	if ( ( ExistingWidget != NULL && ExistingWidget != Item ) || !bUniqueNameForVariable )
+	if ( ( ExistingWidget != NULL && ExistingWidget != Item.GetTemplate() ) || !bUniqueNameForVariable )
 	{
 		OutErrorMessage = LOCTEXT("NameConflict", "NameConflict");
 		return false;
@@ -80,12 +145,12 @@ bool SHierarchyViewItem::OnVerifyNameTextChanged(const FText& InText, FText& Out
 
 void SHierarchyViewItem::OnNameTextCommited(const FText& InText, ETextCommit::Type CommitInfo)
 {
-	FWidgetBlueprintEditorUtils::RenameWidget(BlueprintEditor.Pin().ToSharedRef(), Item->GetFName(), FName(*InText.ToString()));
+	FWidgetBlueprintEditorUtils::RenameWidget(BlueprintEditor.Pin().ToSharedRef(), Item.GetTemplate()->GetFName(), FName(*InText.ToString()));
 }
 
 FSlateFontInfo SHierarchyViewItem::GetItemFont() const
 {
-	if ( !Item->IsGeneratedName() && Item->bIsVariable )
+	if ( !Item.GetTemplate()->IsGeneratedName() && Item.GetTemplate()->bIsVariable )
 	{
 		// TODO UMG Hacky move into style area
 		return FSlateFontInfo(FPaths::EngineContentDir() / TEXT("Slate/Fonts/Roboto-Bold.ttf"), 10);
@@ -99,57 +164,65 @@ FSlateFontInfo SHierarchyViewItem::GetItemFont() const
 
 FText SHierarchyViewItem::GetItemText() const
 {
-	return FText::FromString(Item->GetLabel());
+	return FText::FromString(Item.GetTemplate()->GetLabel());
 }
 
-//@TODO UMG Allow items in the tree to be dragged, and reordered, and reparented.
+FReply SHierarchyViewItem::HandleDragDetected(const FGeometry& MyGeometry, const FPointerEvent& MouseEvent)
+{
+	const bool bIsRoot = Item.GetTemplate()->GetParent() == NULL;
 
-void SHierarchyViewItem::OnDragEnter(const FGeometry& MyGeometry, const FDragDropEvent& DragDropEvent)
+	if ( !bIsRoot )
+	{
+		return FReply::Handled().BeginDragDrop(FHierarchyWidgetDragDropOp::New(BlueprintEditor.Pin()->GetWidgetBlueprintObj(), Item));
+	}
+
+	return FReply::Unhandled();
+}
+
+void SHierarchyViewItem::HandleDragEnter(FDragDropEvent const& DragDropEvent)
 {
-	//TSharedPtr<FDragDropOperation> Operation = DragDropEvent.GetOperation();
-	//if ( !Operation.IsValid() )
-	//{
-	//	return;
-	//}
-
-	//if ( Operation->IsOfType<FWidgetTemplateDragDropOp>() )
-	//{
-	//	TSharedPtr<FWidgetTemplateDragDropOp> DragDropOp = StaticCastSharedPtr<FWidgetTemplateDragDropOp>(Operation);
-	//	bDraggedOver = true;
-	//}
 }
 
-void SHierarchyViewItem::OnDragLeave(const FDragDropEvent& DragDropEvent)
+void SHierarchyViewItem::HandleDragLeave(const FDragDropEvent& DragDropEvent)
 {
-	TSharedPtr<FWidgetTemplateDragDropOp> DragDropOp = DragDropEvent.GetOperationAs<FWidgetTemplateDragDropOp>();
-	if ( DragDropOp.IsValid() )
+	TSharedPtr<FDecoratedDragDropOp> DecoratedDragDropOp = DragDropEvent.GetOperationAs<FDecoratedDragDropOp>();
+	if ( DecoratedDragDropOp.IsValid() )
 	{
-		DragDropOp->ResetToDefaultToolTip();
+		DecoratedDragDropOp->SetCursorOverride(TOptional<EMouseCursor::Type>());
+		DecoratedDragDropOp->ResetToDefaultToolTip();
 	}
 }
 
 FReply SHierarchyViewItem::OnDragOver(const FGeometry& MyGeometry, const FDragDropEvent& DragDropEvent)
 {
-	TSharedPtr<FWidgetTemplateDragDropOp> DragDropOp = DragDropEvent.GetOperationAs<FWidgetTemplateDragDropOp>();
-	if ( DragDropOp.IsValid() )
+	TSharedPtr<FWidgetTemplateDragDropOp> TemplateDragDropOp = DragDropEvent.GetOperationAs<FWidgetTemplateDragDropOp>();
+	if ( TemplateDragDropOp.IsValid() )
 	{
+		// TODO UMG - Show feedback
 		return FReply::Handled();
 	}
 
+	TSharedPtr<FHierarchyWidgetDragDropOp> HierarchyDragDropOp = DragDropEvent.GetOperationAs<FHierarchyWidgetDragDropOp>();
+	if ( HierarchyDragDropOp.IsValid() )
+	{
+		bool bIsDrop = false;
+		return ProcessDragDrop(DragDropEvent, bIsDrop);
+	}
+
 	return FReply::Unhandled();
 }
 
-FReply SHierarchyViewItem::OnDrop(const FGeometry& MyGeometry, const FDragDropEvent& DragDropEvent)
+FReply SHierarchyViewItem::HandleDrop(const FDragDropEvent& DragDropEvent)
 {
-	TSharedPtr<FWidgetTemplateDragDropOp> DragDropOp = DragDropEvent.GetOperationAs<FWidgetTemplateDragDropOp>();
-	if ( DragDropOp.IsValid() )
+	TSharedPtr<FWidgetTemplateDragDropOp> TemplateDragDropOp = DragDropEvent.GetOperationAs<FWidgetTemplateDragDropOp>();
+	if ( TemplateDragDropOp.IsValid() )
 	{
-		if ( Item->IsA(UPanelWidget::StaticClass()) )
+		if ( Item.GetTemplate()->IsA(UPanelWidget::StaticClass()) )
 		{
 			UWidgetBlueprint* BP = CastChecked<UWidgetBlueprint>(BlueprintEditor.Pin()->GetBlueprintObj());			
-			UPanelWidget* Parent = Cast<UPanelWidget>(Item);
+			UPanelWidget* Parent = Cast<UPanelWidget>(Item.GetTemplate());
 
-			UWidget* Widget = DragDropOp->Template->Create(BP->WidgetTree);
+			UWidget* Widget = TemplateDragDropOp->Template->Create(BP->WidgetTree);
 
 			Parent->AddChild(Widget);
 
@@ -159,6 +232,78 @@ FReply SHierarchyViewItem::OnDrop(const FGeometry& MyGeometry, const FDragDropEv
 		return FReply::Handled();
 	}
 
+	bool bIsDrop = true;
+	return ProcessDragDrop(DragDropEvent, bIsDrop);
+}
+
+FReply SHierarchyViewItem::ProcessDragDrop(const FDragDropEvent& DragDropEvent, bool bIsDrop)
+{
+	UWidgetBlueprint* Blueprint = BlueprintEditor.Pin()->GetWidgetBlueprintObj();
+
+	TSharedPtr<FHierarchyWidgetDragDropOp> HierarchyDragDropOp = DragDropEvent.GetOperationAs<FHierarchyWidgetDragDropOp>();
+	if ( HierarchyDragDropOp.IsValid() )
+	{
+		HierarchyDragDropOp->SetCursorOverride(TOptional<EMouseCursor::Type>());
+
+		if ( Item.IsValid() )
+		{
+			const bool bIsPanel = Item.GetTemplate()->IsA(UPanelWidget::StaticClass());
+			const bool bIsDraggedObject = Item.GetTemplate() == HierarchyDragDropOp->Widget.GetTemplate();
+			const bool bIsChildOfDraggedObject = Item.GetTemplate()->IsChildOf(HierarchyDragDropOp->Widget.GetTemplate());
+
+			if ( bIsPanel && !bIsDraggedObject && !bIsChildOfDraggedObject )
+			{
+				UWidget* Widget = bIsDrop ? HierarchyDragDropOp->Widget.GetTemplate() : HierarchyDragDropOp->Widget.GetPreview();
+
+				if ( Widget->GetParent() )
+				{
+					Widget->GetParent()->RemoveChild(Widget);
+				}
+
+				UPanelWidget* NewParent = Cast<UPanelWidget>( bIsDrop ? Item.GetTemplate() : Item.GetPreview() );
+
+				if ( bIsDrop && NewParent != HierarchyDragDropOp->WidgetParent )
+				{
+					NewParent->SetFlags(RF_Transactional);
+					NewParent->Modify();
+				}
+
+				if ( UPanelSlot* Slot = NewParent->AddChild(Widget) )
+				{
+					FWidgetBlueprintEditorUtils::ImportPropertiesFromText(Slot, HierarchyDragDropOp->ExportedSlotProperties);
+
+					Slot->SetDesiredPosition(FVector2D::ZeroVector);
+
+					if ( bIsDrop )
+					{
+						FBlueprintEditorUtils::MarkBlueprintAsStructurallyModified(Blueprint);
+					}
+					else
+					{
+						Item.GetPreview()->GetWidget()->SlatePrepass();
+					}
+				}
+				else
+				{
+					HierarchyDragDropOp->SetCursorOverride(EMouseCursor::SlashedCircle);
+
+					// TODO UMG ERROR Slot can not be created because maybe the max children has been reached.
+					//          Maybe we can traverse the hierarchy and add it to the first parent that will accept it?
+				}
+			}
+			else
+			{
+				HierarchyDragDropOp->SetCursorOverride(EMouseCursor::SlashedCircle);
+			}
+		}
+		else
+		{
+			HierarchyDragDropOp->SetCursorOverride(EMouseCursor::SlashedCircle);
+		}
+
+		return FReply::Handled();
+	}
+
 	return FReply::Unhandled();
 }
 
diff --git a/Engine/Source/Editor/UMGEditor/Private/Hierarchy/SHierarchyViewItem.h b/Engine/Source/Editor/UMGEditor/Private/Hierarchy/SHierarchyViewItem.h
--- a/Engine/Source/Editor/UMGEditor/Private/Hierarchy/SHierarchyViewItem.h
+++ b/Engine/Source/Editor/UMGEditor/Private/Hierarchy/SHierarchyViewItem.h
@@ -18,16 +18,21 @@ public:
 
 	SLATE_END_ARGS()
 
-	void Construct(const FArguments& InArgs, const TSharedRef< STableViewBase >& InOwnerTableView, TSharedPtr<FWidgetBlueprintEditor> InBlueprintEditor, UWidget* InItem);
+	void Construct(const FArguments& InArgs, const TSharedRef< STableViewBase >& InOwnerTableView, TSharedPtr<FWidgetBlueprintEditor> InBlueprintEditor, FWidgetReference InItem);
 
 	// Begin SWidget
-	virtual void OnDragEnter(const FGeometry& MyGeometry, const FDragDropEvent& DragDropEvent) override;
-	virtual void OnDragLeave(const FDragDropEvent& DragDropEvent) override;
 	virtual FReply OnDragOver(const FGeometry& MyGeometry, const FDragDropEvent& DragDropEvent) override;
-	virtual FReply OnDrop(const FGeometry& MyGeometry, const FDragDropEvent& DragDropEvent) override;
 	// End SWidget
-	
+
 private:
+	FReply HandleDragDetected(const FGeometry& MyGeometry, const FPointerEvent& MouseEvent);
+	void HandleDragEnter(const FDragDropEvent& DragDropEvent);
+	void HandleDragLeave(const FDragDropEvent& DragDropEvent);
+	FReply HandleDrop(FDragDropEvent const& DragDropEvent);
+
+	/** Process Drag / Drop events providing either feedback, or generating the final drop action */
+	FReply ProcessDragDrop(const FDragDropEvent& DragDropEvent, bool bIsDrop);
+
 	/** Called when text is being committed to check for validity */
 	bool OnVerifyNameTextChanged(const FText& InText, FText& OutErrorMessage);
 
@@ -44,5 +49,5 @@ private:
 	TWeakPtr<FWidgetBlueprintEditor> BlueprintEditor;
 
 	/* The template widget that this tree item represents */
-	UWidget* Item;
+	FWidgetReference Item;
 };
diff --git a/Engine/Source/Runtime/UMG/Private/Components/Widget.cpp b/Engine/Source/Runtime/UMG/Private/Components/Widget.cpp
--- a/Engine/Source/Runtime/UMG/Private/Components/Widget.cpp
+++ b/Engine/Source/Runtime/UMG/Private/Components/Widget.cpp
@@ -238,6 +238,21 @@ bool UWidget::Modify(bool bAlwaysMarkDirty)
 	return Modified;
 }
 
+bool UWidget::IsChildOf(UWidget* PossibleParent)
+{
+	UPanelWidget* Parent = GetParent();
+	if ( Parent == NULL )
+	{
+		return false;
+	}
+	else if ( Parent == PossibleParent )
+	{
+		return true;
+	}
+	
+	return Parent->IsChildOf(PossibleParent);
+}
+
 TSharedRef<SWidget> UWidget::RebuildWidget()
 {
 	ensureMsg(false, TEXT("You must implement RebuildWidget() in your child class"));
diff --git a/Engine/Source/Runtime/UMG/Public/Components/Widget.h b/Engine/Source/Runtime/UMG/Public/Components/Widget.h
--- a/Engine/Source/Runtime/UMG/Public/Components/Widget.h
+++ b/Engine/Source/Runtime/UMG/Public/Components/Widget.h
@@ -157,6 +157,9 @@ public:
 
 	/** Mark this object as modified, also mark the slot as modified. */
 	virtual bool Modify(bool bAlwaysMarkDirty = true);
+
+	/** @return true if this widget is a child of the PossibleParent */
+	bool IsChildOf(UWidget* PossibleParent);
 	
 #if WITH_EDITOR
 	/** Is the label generated or provided by the user? */
