UMG - Now updating multicast events on rename of widgets.

diff --git a/Engine/Source/Editor/UMGEditor/Private/WidgetBlueprintEditorUtils.cpp b/Engine/Source/Editor/UMGEditor/Private/WidgetBlueprintEditorUtils.cpp
--- a/Engine/Source/Editor/UMGEditor/Private/WidgetBlueprintEditorUtils.cpp
+++ b/Engine/Source/Editor/UMGEditor/Private/WidgetBlueprintEditorUtils.cpp
@@ -99,11 +99,11 @@ bool FWidgetBlueprintEditorUtils::RenameWidget(TSharedRef<FWidgetBlueprintEditor
 
 		Widget->Rename(*NewNameStr);
 	
+		// Update Variable References
 		TArray<UK2Node_Variable*> WidgetVarNodes;
 		FBlueprintEditorUtils::GetAllNodesOfClass<UK2Node_Variable>(Blueprint, WidgetVarNodes);
-		for ( int32 It = 0; It < WidgetVarNodes.Num(); It++ )
+		for ( UK2Node_Variable* TestNode : WidgetVarNodes )
 		{
-			UK2Node_Variable* TestNode = WidgetVarNodes[It];
 			if ( TestNode && ( OldName == TestNode->GetVarName() ) )
 			{
 				UEdGraphPin* TestPin = TestNode->FindPin(OldNameStr);
@@ -126,6 +126,17 @@ bool FWidgetBlueprintEditorUtils::RenameWidget(TSharedRef<FWidgetBlueprintEditor
 			}
 		}
 
+		// Update Event References to member variables
+		TArray<UK2Node_ComponentBoundEvent*> EventNodes;
+		FBlueprintEditorUtils::GetAllNodesOfClass<UK2Node_ComponentBoundEvent>(Blueprint, EventNodes);
+		for ( UK2Node_ComponentBoundEvent* EventNode : EventNodes )
+		{
+			if ( EventNode->ComponentPropertyName == OldName )
+			{
+				EventNode->ComponentPropertyName = NewName;
+			}
+		}
+
 		// Find and update all binding references in the widget blueprint
 		for ( FDelegateEditorBinding& Binding : Blueprint->Bindings )
 		{
