UMG - Improved widget rename validation and error reporting.

diff --git a/Engine/Source/Editor/UMGEditor/Private/Details/SWidgetDetailsView.cpp b/Engine/Source/Editor/UMGEditor/Private/Details/SWidgetDetailsView.cpp
--- a/Engine/Source/Editor/UMGEditor/Private/Details/SWidgetDetailsView.cpp
+++ b/Engine/Source/Editor/UMGEditor/Private/Details/SWidgetDetailsView.cpp
@@ -199,20 +199,34 @@ bool SWidgetDetailsView::HandleVerifyNameTextChanged(const FText& InText, FText&
 {
 	if ( SelectedObjects.Num() == 1 )
 	{
-		UWidget* Widget = Cast<UWidget>(SelectedObjects[0].Get());
+		UWidget* PreviewWidget = Cast<UWidget>(SelectedObjects[0].Get());
 
 		FString NewName = InText.ToString();
 
 		UWidgetBlueprint* Blueprint = BlueprintEditor.Pin()->GetWidgetBlueprintObj();
-		UWidget* ExistingWidget = Blueprint->WidgetTree->FindWidget(NewName);
+		UWidget* TemplateWidget = Blueprint->WidgetTree->FindWidget(NewName);
+
+		bool bIsSameWidget = false;
+		if ( TemplateWidget != NULL )
+		{
+			if ( FWidgetReference::FromTemplate(BlueprintEditor.Pin(), TemplateWidget).GetPreview() != PreviewWidget )
+			{
+				OutErrorMessage = LOCTEXT("ExistingWidgetName", "Existing Widget Name");
+				return false;
+			}
+			else
+			{
+				bIsSameWidget = true;
+			}
+		}
 
 		FKismetNameValidator Validator(Blueprint);
 
 		const bool bUniqueNameForVariable = ( EValidatorResult::Ok == Validator.IsValid(NewName) );
 
-		if ( ( ExistingWidget != NULL && ExistingWidget != Widget ) || !bUniqueNameForVariable )
+		if ( !bUniqueNameForVariable && !bIsSameWidget )
 		{
-			OutErrorMessage = LOCTEXT("NameConflict", "NameConflict");
+			OutErrorMessage = LOCTEXT("ExistingVariableName", "Existing Variable Name");
 			return false;
 		}
 
@@ -233,8 +247,12 @@ void SWidgetDetailsView::HandleNameTextCommitted(const FText& Text, ETextCommit:
 		IsReentrant = true;
 		if ( SelectedObjects.Num() == 1 )
 		{
-			UWidget* Widget = Cast<UWidget>(SelectedObjects[0].Get());
-			FWidgetBlueprintEditorUtils::RenameWidget(BlueprintEditor.Pin().ToSharedRef(), Widget->GetFName(), FName(*Text.ToString()));
+			FText DummyText;
+			if ( HandleVerifyNameTextChanged(Text, DummyText) )
+			{
+				UWidget* Widget = Cast<UWidget>(SelectedObjects[0].Get());
+				FWidgetBlueprintEditorUtils::RenameWidget(BlueprintEditor.Pin().ToSharedRef(), Widget->GetFName(), FName(*Text.ToString()));
+			}
 		}
 		IsReentrant = false;
 	}
diff --git a/Engine/Source/Editor/UMGEditor/Private/Hierarchy/SHierarchyViewItem.cpp b/Engine/Source/Editor/UMGEditor/Private/Hierarchy/SHierarchyViewItem.cpp
--- a/Engine/Source/Editor/UMGEditor/Private/Hierarchy/SHierarchyViewItem.cpp
+++ b/Engine/Source/Editor/UMGEditor/Private/Hierarchy/SHierarchyViewItem.cpp
@@ -128,15 +128,29 @@ bool SHierarchyViewItem::OnVerifyNameTextChanged(const FText& InText, FText& Out
 	FString NewName = InText.ToString();
 
 	UWidgetBlueprint* Blueprint = BlueprintEditor.Pin()->GetWidgetBlueprintObj();
-	UWidget* ExistingWidget = Blueprint->WidgetTree->FindWidget(NewName);
+	UWidget* ExistingTemplate = Blueprint->WidgetTree->FindWidget(NewName);
+
+	bool bIsSameWidget = false;
+	if ( ExistingTemplate != NULL )
+	{
+		if ( Item.GetTemplate() == ExistingTemplate )
+		{
+			OutErrorMessage = LOCTEXT("ExistingWidgetName", "Existing Widget Name");
+			return false;
+		}
+		else
+		{
+			bIsSameWidget = true;
+		}
+	}
 
 	FKismetNameValidator Validator(Blueprint);
 
 	const bool bUniqueNameForVariable = ( EValidatorResult::Ok == Validator.IsValid(NewName) );
-	
-	if ( ( ExistingWidget != NULL && ExistingWidget != Item.GetTemplate() ) || !bUniqueNameForVariable )
+
+	if ( !bUniqueNameForVariable && !bIsSameWidget )
 	{
-		OutErrorMessage = LOCTEXT("NameConflict", "NameConflict");
+		OutErrorMessage = LOCTEXT("ExistingVariableName", "Existing Variable Name");
 		return false;
 	}
 
