UMG - Now have a customization for horizontal and vertical alignment.  Buttons now have a slot object of their own, old assets will load, but things may appear differently.

diff --git a/Engine/Source/Editor/EditorStyle/Private/SlateEditorStyle.cpp b/Engine/Source/Editor/EditorStyle/Private/SlateEditorStyle.cpp
--- a/Engine/Source/Editor/EditorStyle/Private/SlateEditorStyle.cpp
+++ b/Engine/Source/Editor/EditorStyle/Private/SlateEditorStyle.cpp
@@ -5399,6 +5399,16 @@ void FSlateEditorStyle::FStyle::SetupUMGEditorStyles()
 		.SetTextColor( DefaultForeground )
 		.SetSelectedTextColor( InvertedForeground )
 		);
+
+	Set("HorizontalAlignment_Left", new IMAGE_BRUSH("Icons/UMG/Alignment/Horizontal_Left", Icon20x20));
+	Set("HorizontalAlignment_Center", new IMAGE_BRUSH("Icons/UMG/Alignment/Horizontal_Center", Icon20x20));
+	Set("HorizontalAlignment_Right", new IMAGE_BRUSH("Icons/UMG/Alignment/Horizontal_Right", Icon20x20));
+	Set("HorizontalAlignment_Fill", new IMAGE_BRUSH("Icons/UMG/Alignment/Horizontal_Fill", Icon20x20));
+
+	Set("VerticalAlignment_Top", new IMAGE_BRUSH("Icons/UMG/Alignment/Vertical_Top", Icon20x20));
+	Set("VerticalAlignment_Center", new IMAGE_BRUSH("Icons/UMG/Alignment/Vertical_Center", Icon20x20));
+	Set("VerticalAlignment_Bottom", new IMAGE_BRUSH("Icons/UMG/Alignment/Vertical_Bottom", Icon20x20));
+	Set("VerticalAlignment_Fill", new IMAGE_BRUSH("Icons/UMG/Alignment/Vertical_Fill", Icon20x20));
 }
 
 void FSlateEditorStyle::FStyle::SetupTranslationEditorStyles()
diff --git a/Engine/Source/Editor/UMGEditor/Private/Customizations/HorizontalAlignmentCustomization.cpp b/Engine/Source/Editor/UMGEditor/Private/Customizations/HorizontalAlignmentCustomization.cpp
new file mode 100644
--- /dev/null
+++ b/Engine/Source/Editor/UMGEditor/Private/Customizations/HorizontalAlignmentCustomization.cpp
@@ -0,0 +1,113 @@
+// Copyright 1998-2014 Epic Games, Inc. All Rights Reserved.
+
+#include "UMGEditorPrivatePCH.h"
+
+#include "PropertyEditing.h"
+#include "ObjectEditorUtils.h"
+#include "WidgetGraphSchema.h"
+#include "ScopedTransaction.h"
+#include "BlueprintEditorUtils.h"
+
+#define LOCTEXT_NAMESPACE "UMG"
+
+void FHorizontalAlignmentCustomization::CustomizeHeader(TSharedRef<IPropertyHandle> PropertyHandle, FDetailWidgetRow& HeaderRow, IPropertyTypeCustomizationUtils& CustomizationUtils)
+{
+	const FMargin OuterPadding(2);
+	const FMargin ContentPadding(2);
+
+	HeaderRow
+	.NameContent()
+	[
+		PropertyHandle->CreatePropertyNameWidget()
+	]
+	.ValueContent()
+	[
+		SNew(SHorizontalBox)
+
+		+ SHorizontalBox::Slot()
+		.AutoWidth()
+		.Padding(OuterPadding)
+		[
+			SNew( SCheckBox )
+			.Style(FEditorStyle::Get(), "ToggleButtonCheckbox")
+			.ToolTipText(LOCTEXT("HAlignLeft", "Horizontally Align Left"))
+			.Padding(ContentPadding)
+			.OnCheckStateChanged(this, &FHorizontalAlignmentCustomization::HandleCheckStateChanged, PropertyHandle, HAlign_Left)
+			.IsChecked(this, &FHorizontalAlignmentCustomization::GetCheckState, PropertyHandle, HAlign_Left)
+			[
+				SNew(SImage)
+				.Image(FEditorStyle::GetBrush("HorizontalAlignment_Left"))
+			]
+		]
+
+		+ SHorizontalBox::Slot()
+		.AutoWidth()
+		.Padding(OuterPadding)
+		[
+			SNew(SCheckBox)
+			.Style(FEditorStyle::Get(), "ToggleButtonCheckbox")
+			.ToolTipText(LOCTEXT("HAlignCenter", "Horizontally Align Center"))
+			.Padding(ContentPadding)
+			.OnCheckStateChanged(this, &FHorizontalAlignmentCustomization::HandleCheckStateChanged, PropertyHandle, HAlign_Center)
+			.IsChecked(this, &FHorizontalAlignmentCustomization::GetCheckState, PropertyHandle, HAlign_Center)
+			[
+				SNew(SImage)
+				.Image(FEditorStyle::GetBrush("HorizontalAlignment_Center"))
+			]
+		]
+
+		+ SHorizontalBox::Slot()
+		.AutoWidth()
+		.Padding(OuterPadding)
+		[
+			SNew(SCheckBox)
+			.Style(FEditorStyle::Get(), "ToggleButtonCheckbox")
+			.ToolTipText(LOCTEXT("HAlignRight", "Horizontally Align Right"))
+			.Padding(ContentPadding)
+			.OnCheckStateChanged(this, &FHorizontalAlignmentCustomization::HandleCheckStateChanged, PropertyHandle, HAlign_Right)
+			.IsChecked(this, &FHorizontalAlignmentCustomization::GetCheckState, PropertyHandle, HAlign_Right)
+			[
+				SNew(SImage)
+				.Image(FEditorStyle::GetBrush("HorizontalAlignment_Right"))
+			]
+		]
+
+		+ SHorizontalBox::Slot()
+		.AutoWidth()
+		.Padding(OuterPadding)
+		[
+			SNew(SCheckBox)
+			.Style(FEditorStyle::Get(), "ToggleButtonCheckbox")
+			.ToolTipText(LOCTEXT("HAlignFill", "Horizontally Align Fill"))
+			.Padding(ContentPadding)
+			.OnCheckStateChanged(this, &FHorizontalAlignmentCustomization::HandleCheckStateChanged, PropertyHandle, HAlign_Fill)
+			.IsChecked(this, &FHorizontalAlignmentCustomization::GetCheckState, PropertyHandle, HAlign_Fill)
+			[
+				SNew(SImage)
+				.Image(FEditorStyle::GetBrush("HorizontalAlignment_Fill"))
+			]
+		]
+	];
+}
+
+void FHorizontalAlignmentCustomization::CustomizeChildren(TSharedRef<IPropertyHandle> PropertyHandle, IDetailChildrenBuilder& ChildBuilder, IPropertyTypeCustomizationUtils& CustomizationUtils)
+{
+}
+
+void FHorizontalAlignmentCustomization::HandleCheckStateChanged(ESlateCheckBoxState::Type InCheckboxState, TSharedRef<IPropertyHandle> PropertyHandle, EHorizontalAlignment ToAlignment)
+{
+	PropertyHandle->SetValue((uint8)ToAlignment);
+}
+
+ESlateCheckBoxState::Type FHorizontalAlignmentCustomization::GetCheckState(TSharedRef<IPropertyHandle> PropertyHandle, EHorizontalAlignment ForAlignment) const
+{
+	uint8 Value;
+	if ( PropertyHandle->GetValue(Value) )
+	{
+		return Value == ForAlignment ? ESlateCheckBoxState::Checked : ESlateCheckBoxState::Unchecked;
+	}
+
+	return ESlateCheckBoxState::Unchecked;
+}
+
+#undef LOCTEXT_NAMESPACE
diff --git a/Engine/Source/Editor/UMGEditor/Private/Customizations/HorizontalAlignmentCustomization.h b/Engine/Source/Editor/UMGEditor/Private/Customizations/HorizontalAlignmentCustomization.h
new file mode 100644
--- /dev/null
+++ b/Engine/Source/Editor/UMGEditor/Private/Customizations/HorizontalAlignmentCustomization.h
@@ -0,0 +1,25 @@
+// Copyright 1998-2014 Epic Games, Inc. All Rights Reserved.
+
+#pragma once
+
+class FHorizontalAlignmentCustomization : public IPropertyTypeCustomization
+{
+public:
+	/** Makes a new instance of this detail layout class for a specific detail view requesting it */
+	static TSharedRef<class IPropertyTypeCustomization> MakeInstance()
+	{
+		return MakeShareable(new FHorizontalAlignmentCustomization());
+	}
+
+	FHorizontalAlignmentCustomization()
+	{
+	}
+	
+	virtual void CustomizeHeader(TSharedRef<IPropertyHandle> PropertyHandle, class FDetailWidgetRow& HeaderRow, IPropertyTypeCustomizationUtils& CustomizationUtils) override;
+
+	virtual void CustomizeChildren(TSharedRef<IPropertyHandle> PropertyHandle, class IDetailChildrenBuilder& ChildBuilder, IPropertyTypeCustomizationUtils& CustomizationUtils) override;
+	
+	void HandleCheckStateChanged(ESlateCheckBoxState::Type InCheckboxState, TSharedRef<IPropertyHandle> PropertyHandle, EHorizontalAlignment ToAlignment);
+
+	ESlateCheckBoxState::Type GetCheckState(TSharedRef<IPropertyHandle> PropertyHandle, EHorizontalAlignment ForAlignment) const;
+};
diff --git a/Engine/Source/Editor/UMGEditor/Private/Customizations/VerticalAlignmentCustomization.cpp b/Engine/Source/Editor/UMGEditor/Private/Customizations/VerticalAlignmentCustomization.cpp
new file mode 100644
--- /dev/null
+++ b/Engine/Source/Editor/UMGEditor/Private/Customizations/VerticalAlignmentCustomization.cpp
@@ -0,0 +1,113 @@
+// Copyright 1998-2014 Epic Games, Inc. All Rights Reserved.
+
+#include "UMGEditorPrivatePCH.h"
+
+#include "PropertyEditing.h"
+#include "ObjectEditorUtils.h"
+#include "WidgetGraphSchema.h"
+#include "ScopedTransaction.h"
+#include "BlueprintEditorUtils.h"
+
+#define LOCTEXT_NAMESPACE "UMG"
+
+void FVerticalAlignmentCustomization::CustomizeHeader(TSharedRef<IPropertyHandle> PropertyHandle, FDetailWidgetRow& HeaderRow, IPropertyTypeCustomizationUtils& CustomizationUtils)
+{
+	const FMargin OuterPadding(2);
+	const FMargin ContentPadding(2);
+
+	HeaderRow
+	.NameContent()
+	[
+		PropertyHandle->CreatePropertyNameWidget()
+	]
+	.ValueContent()
+	[
+		SNew(SHorizontalBox)
+
+		+ SHorizontalBox::Slot()
+		.AutoWidth()
+		.Padding(OuterPadding)
+		[
+			SNew( SCheckBox )
+			.Style(FEditorStyle::Get(), "ToggleButtonCheckbox")
+			.ToolTipText(LOCTEXT("VAlignTop", "Vertically Align Top"))
+			.Padding(ContentPadding)
+			.OnCheckStateChanged(this, &FVerticalAlignmentCustomization::HandleCheckStateChanged, PropertyHandle, VAlign_Top)
+			.IsChecked(this, &FVerticalAlignmentCustomization::GetCheckState, PropertyHandle, VAlign_Top)
+			[
+				SNew(SImage)
+				.Image(FEditorStyle::GetBrush("VerticalAlignment_Top"))
+			]
+		]
+
+		+ SHorizontalBox::Slot()
+		.AutoWidth()
+		.Padding(OuterPadding)
+		[
+			SNew(SCheckBox)
+			.Style(FEditorStyle::Get(), "ToggleButtonCheckbox")
+			.ToolTipText(LOCTEXT("VAlignCenter", "Vertically Align Center"))
+			.Padding(ContentPadding)
+			.OnCheckStateChanged(this, &FVerticalAlignmentCustomization::HandleCheckStateChanged, PropertyHandle, VAlign_Center)
+			.IsChecked(this, &FVerticalAlignmentCustomization::GetCheckState, PropertyHandle, VAlign_Center)
+			[
+				SNew(SImage)
+				.Image(FEditorStyle::GetBrush("VerticalAlignment_Center"))
+			]
+		]
+
+		+ SHorizontalBox::Slot()
+		.AutoWidth()
+		.Padding(OuterPadding)
+		[
+			SNew(SCheckBox)
+			.Style(FEditorStyle::Get(), "ToggleButtonCheckbox")
+			.ToolTipText(LOCTEXT("VAlignBottom", "Vertically Align Bottom"))
+			.Padding(ContentPadding)
+			.OnCheckStateChanged(this, &FVerticalAlignmentCustomization::HandleCheckStateChanged, PropertyHandle, VAlign_Bottom)
+			.IsChecked(this, &FVerticalAlignmentCustomization::GetCheckState, PropertyHandle, VAlign_Bottom)
+			[
+				SNew(SImage)
+				.Image(FEditorStyle::GetBrush("VerticalAlignment_Bottom"))
+			]
+		]
+
+		+ SHorizontalBox::Slot()
+		.AutoWidth()
+		.Padding(OuterPadding)
+		[
+			SNew(SCheckBox)
+			.Style(FEditorStyle::Get(), "ToggleButtonCheckbox")
+			.ToolTipText(LOCTEXT("VAlignFill", "Vertically Align Fill"))
+			.Padding(ContentPadding)
+			.OnCheckStateChanged(this, &FVerticalAlignmentCustomization::HandleCheckStateChanged, PropertyHandle, VAlign_Fill)
+			.IsChecked(this, &FVerticalAlignmentCustomization::GetCheckState, PropertyHandle, VAlign_Fill)
+			[
+				SNew(SImage)
+				.Image(FEditorStyle::GetBrush("VerticalAlignment_Fill"))
+			]
+		]
+	];
+}
+
+void FVerticalAlignmentCustomization::CustomizeChildren(TSharedRef<IPropertyHandle> PropertyHandle, IDetailChildrenBuilder& ChildBuilder, IPropertyTypeCustomizationUtils& CustomizationUtils)
+{
+}
+
+void FVerticalAlignmentCustomization::HandleCheckStateChanged(ESlateCheckBoxState::Type InCheckboxState, TSharedRef<IPropertyHandle> PropertyHandle, EVerticalAlignment ToAlignment)
+{
+	PropertyHandle->SetValue((uint8)ToAlignment);
+}
+
+ESlateCheckBoxState::Type FVerticalAlignmentCustomization::GetCheckState(TSharedRef<IPropertyHandle> PropertyHandle, EVerticalAlignment ForAlignment) const
+{
+	uint8 Value;
+	if ( PropertyHandle->GetValue(Value) )
+	{
+		return Value == ForAlignment ? ESlateCheckBoxState::Checked : ESlateCheckBoxState::Unchecked;
+	}
+
+	return ESlateCheckBoxState::Unchecked;
+}
+
+#undef LOCTEXT_NAMESPACE
diff --git a/Engine/Source/Editor/UMGEditor/Private/Customizations/VerticalAlignmentCustomization.h b/Engine/Source/Editor/UMGEditor/Private/Customizations/VerticalAlignmentCustomization.h
new file mode 100644
--- /dev/null
+++ b/Engine/Source/Editor/UMGEditor/Private/Customizations/VerticalAlignmentCustomization.h
@@ -0,0 +1,25 @@
+// Copyright 1998-2014 Epic Games, Inc. All Rights Reserved.
+
+#pragma once
+
+class FVerticalAlignmentCustomization : public IPropertyTypeCustomization
+{
+public:
+	/** Makes a new instance of this detail layout class for a specific detail view requesting it */
+	static TSharedRef<class IPropertyTypeCustomization> MakeInstance()
+	{
+		return MakeShareable(new FVerticalAlignmentCustomization());
+	}
+
+	FVerticalAlignmentCustomization()
+	{
+	}
+	
+	virtual void CustomizeHeader(TSharedRef<IPropertyHandle> PropertyHandle, class FDetailWidgetRow& HeaderRow, IPropertyTypeCustomizationUtils& CustomizationUtils) override;
+
+	virtual void CustomizeChildren(TSharedRef<IPropertyHandle> PropertyHandle, class IDetailChildrenBuilder& ChildBuilder, IPropertyTypeCustomizationUtils& CustomizationUtils) override;
+	
+	void HandleCheckStateChanged(ESlateCheckBoxState::Type InCheckboxState, TSharedRef<IPropertyHandle> PropertyHandle, EVerticalAlignment ToAlignment);
+
+	ESlateCheckBoxState::Type GetCheckState(TSharedRef<IPropertyHandle> PropertyHandle, EVerticalAlignment ForAlignment) const;
+};
diff --git a/Engine/Source/Editor/UMGEditor/Private/Details/SWidgetDetailsView.cpp b/Engine/Source/Editor/UMGEditor/Private/Details/SWidgetDetailsView.cpp
--- a/Engine/Source/Editor/UMGEditor/Private/Details/SWidgetDetailsView.cpp
+++ b/Engine/Source/Editor/UMGEditor/Private/Details/SWidgetDetailsView.cpp
@@ -106,6 +106,9 @@ void SWidgetDetailsView::RegisterCustomizations()
 	static FName PropertyEditor("PropertyEditor");
 	FPropertyEditorModule& PropertyModule = FModuleManager::GetModuleChecked<FPropertyEditorModule>(PropertyEditor);
 	PropertyModule.RegisterCustomPropertyTypeLayout(TEXT("PanelSlot"), CanvasSlotCustomization);
+
+	PropertyModule.RegisterCustomPropertyTypeLayout(TEXT("EHorizontalAlignment"), FOnGetPropertyTypeCustomizationInstance::CreateStatic(&FHorizontalAlignmentCustomization::MakeInstance));
+	PropertyModule.RegisterCustomPropertyTypeLayout(TEXT("EVerticalAlignment"), FOnGetPropertyTypeCustomizationInstance::CreateStatic(&FVerticalAlignmentCustomization::MakeInstance));
 }
 
 void SWidgetDetailsView::OnEditorSelectionChanging()
diff --git a/Engine/Source/Editor/UMGEditor/Private/UMGEditorPrivatePCH.h b/Engine/Source/Editor/UMGEditor/Private/UMGEditorPrivatePCH.h
--- a/Engine/Source/Editor/UMGEditor/Private/UMGEditorPrivatePCH.h
+++ b/Engine/Source/Editor/UMGEditor/Private/UMGEditorPrivatePCH.h
@@ -37,6 +37,8 @@
 
 #include "DetailCustomizations.h"
 #include "CanvasSlotCustomization.h"
+#include "HorizontalAlignmentCustomization.h"
+#include "VerticalAlignmentCustomization.h"
 
 #include "WidgetTemplate.h"
 #include "WidgetTemplateDragDropOp.h"
diff --git a/Engine/Source/Editor/UMGEditor/Private/WidgetBlueprintCompiler.cpp b/Engine/Source/Editor/UMGEditor/Private/WidgetBlueprintCompiler.cpp
--- a/Engine/Source/Editor/UMGEditor/Private/WidgetBlueprintCompiler.cpp
+++ b/Engine/Source/Editor/UMGEditor/Private/WidgetBlueprintCompiler.cpp
@@ -480,6 +480,7 @@ void FWidgetBlueprintCompiler::CreateClassVariablesFromBlueprint()
 		if ( WidgetProperty != NULL )
 		{
 			WidgetProperty->SetMetaData(TEXT("Category"), *Blueprint->GetName());
+			//WidgetProperty->SetMetaData(TEXT("Category"), TEXT("Widget"));
 			WidgetProperty->SetPropertyFlags(CPF_BlueprintVisible);
 
 			WidgetToMemberVariableMap.Add(Widget, WidgetProperty);
diff --git a/Engine/Source/Runtime/UMG/Private/Components/Button.cpp b/Engine/Source/Runtime/UMG/Private/Components/Button.cpp
--- a/Engine/Source/Runtime/UMG/Private/Components/Button.cpp
+++ b/Engine/Source/Runtime/UMG/Private/Components/Button.cpp
@@ -10,13 +10,9 @@
 UButton::UButton(const FPostConstructInitializeProperties& PCIP)
 	: Super(PCIP)
 {
-	HorizontalAlignment = HAlign_Center;
-	VerticalAlignment = VAlign_Center;
-
 	SButton::FArguments ButtonDefaults;
 
 	Style = NULL;
-	ContentPadding = ButtonDefaults._ContentPadding.Get();
 
 	DesiredSizeScale = ButtonDefaults._DesiredSizeScale.Get();
 	ContentScale = ButtonDefaults._ContentScale.Get();
@@ -32,7 +28,7 @@ TSharedRef<SWidget> UButton::RebuildWidget()
 
 	if ( GetChildrenCount() > 0 )
 	{
-		MyButton->SetContent(GetContentSlot()->Content ? GetContentSlot()->Content->GetWidget() : SNullWidget::NullWidget);
+		Cast<UButtonSlot>(GetContentSlot())->BuildSlot(MyButton.ToSharedRef());
 	}
 	
 	return MyButton.ToSharedRef();
@@ -62,29 +58,31 @@ void UButton::SyncronizeProperties()
 		OptionalHoveredSound = HoveredSound;
 	}
 
-	MyButton->SetButtonStyle(StylePtr);
-	MyButton->SetHAlign(HorizontalAlignment);
-	MyButton->SetVAlign(VerticalAlignment);
-	MyButton->SetContentPadding(ContentPadding);
-	
+	MyButton->SetButtonStyle( StylePtr );
+
 	MyButton->SetForegroundColor( ForegroundColor );
 	MyButton->SetColorAndOpacity( ColorAndOpacity );
 	MyButton->SetBorderBackgroundColor( BackgroundColor );
 	
-	MyButton->SetDesiredSizeScale(DesiredSizeScale);
-	MyButton->SetContentScale(ContentScale);
-	MyButton->SetPressedSound(OptionalPressedSound);
-	MyButton->SetHoveredSound(OptionalHoveredSound);
+	MyButton->SetDesiredSizeScale( DesiredSizeScale );
+	MyButton->SetContentScale( ContentScale );
+	MyButton->SetPressedSound( OptionalPressedSound );
+	MyButton->SetHoveredSound( OptionalHoveredSound );
 
 	MyButton->SetOnClicked(BIND_UOBJECT_DELEGATE(FOnClicked, HandleOnClicked));
 }
 
+UClass* UButton::GetSlotClass() const
+{
+	return UButtonSlot::StaticClass();
+}
+
 void UButton::OnSlotAdded(UPanelSlot* Slot)
 {
 	// Add the child to the live slot if it already exists
 	if ( MyButton.IsValid() )
 	{
-		MyButton->SetContent(Slot->Content ? Slot->Content->GetWidget() : SNullWidget::NullWidget);
+		Cast<UButtonSlot>(Slot)->BuildSlot(MyButton.ToSharedRef());
 	}
 }
 
@@ -124,15 +122,6 @@ void UButton::SetForegroundColor(FLinearColor InForegroundColor)
 	}
 }
 
-void UButton::SetContentPadding(FMargin InContentPadding)
-{
-	ContentPadding = InContentPadding;
-	if ( MyButton.IsValid() )
-	{
-		MyButton->SetContentPadding(InContentPadding);
-	}
-}
-
 FReply UButton::HandleOnClicked()
 {
 	if ( OnClickedEvent.IsBound() )
@@ -148,6 +137,24 @@ bool UButton::IsPressed() const
 	return MyButton->IsPressed();
 }
 
+void UButton::PostLoad()
+{
+	Super::PostLoad();
+
+	if ( GetChildrenCount() > 0 )
+	{
+		//TODO UMG Pre-Release Upgrade, now buttons have slots of their own.  Convert existing slot to new slot.
+		UButtonSlot* ButtonSlot = Cast<UButtonSlot>(GetContentSlot());
+		if ( ButtonSlot == NULL )
+		{
+			ButtonSlot = ConstructObject<UButtonSlot>(UButtonSlot::StaticClass(), this);
+			ButtonSlot->Content = GetContentSlot()->Content;
+			ButtonSlot->Content->Slot = ButtonSlot;
+			Slots[0] = ButtonSlot;
+		}
+	}
+}
+
 #if WITH_EDITOR
 
 const FSlateBrush* UButton::GetEditorIcon()
diff --git a/Engine/Source/Runtime/UMG/Private/Components/ButtonSlot.cpp b/Engine/Source/Runtime/UMG/Private/Components/ButtonSlot.cpp
new file mode 100644
--- /dev/null
+++ b/Engine/Source/Runtime/UMG/Private/Components/ButtonSlot.cpp
@@ -0,0 +1,60 @@
+// Copyright 1998-2014 Epic Games, Inc. All Rights Reserved.
+
+#include "UMGPrivatePCH.h"
+
+/////////////////////////////////////////////////////
+// UButtonSlot
+
+UButtonSlot::UButtonSlot(const FPostConstructInitializeProperties& PCIP)
+	: Super(PCIP)
+{
+	Padding = FMargin(4, 2);
+
+	HorizontalAlignment = HAlign_Center;
+	VerticalAlignment = VAlign_Center;
+}
+
+void UButtonSlot::BuildSlot(TSharedRef<SButton> InButton)
+{
+	Button = InButton;
+
+	Button->SetContentPadding(Padding);
+	Button->SetHAlign(HorizontalAlignment);
+	Button->SetVAlign(VerticalAlignment);
+
+	Button->SetContent(Content ? Content->GetWidget() : SNullWidget::NullWidget);
+}
+
+void UButtonSlot::SetPadding(FMargin InPadding)
+{
+	Padding = InPadding;
+	if ( Button.IsValid() )
+	{
+		Button->SetContentPadding(InPadding);
+	}
+}
+
+void UButtonSlot::SetHorizontalAlignment(EHorizontalAlignment InHorizontalAlignment)
+{
+	HorizontalAlignment = InHorizontalAlignment;
+	if ( Button.IsValid() )
+	{
+		Button->SetHAlign(InHorizontalAlignment);
+	}
+}
+
+void UButtonSlot::SetVerticalAlignment(EVerticalAlignment InVerticalAlignment)
+{
+	VerticalAlignment = InVerticalAlignment;
+	if ( Button.IsValid() )
+	{
+		Button->SetVAlign(InVerticalAlignment);
+	}
+}
+
+void UButtonSlot::SyncronizeProperties()
+{
+	SetPadding(Padding);
+	SetHorizontalAlignment(HorizontalAlignment);
+	SetVerticalAlignment(VerticalAlignment);
+}
diff --git a/Engine/Source/Runtime/UMG/Public/Components/Button.h b/Engine/Source/Runtime/UMG/Public/Components/Button.h
--- a/Engine/Source/Runtime/UMG/Public/Components/Button.h
+++ b/Engine/Source/Runtime/UMG/Public/Components/Button.h
@@ -17,18 +17,6 @@ public:
 	UPROPERTY(EditDefaultsOnly, Category=Style, meta=( DisplayThumbnail = "true" ))
 	USlateWidgetStyleAsset* Style;
 
-	/** Horizontal positioning of the content within the button */
-	UPROPERTY(EditDefaultsOnly, Category="Content Layout")
-	TEnumAsByte<EHorizontalAlignment> HorizontalAlignment;
-
-	/** Vertical positioning of the content within the button */
-	UPROPERTY(EditDefaultsOnly, Category="Content Layout")
-	TEnumAsByte<EVerticalAlignment> VerticalAlignment;
-
-	/** The padding to add around the button content. */
-	UPROPERTY(EditDefaultsOnly, Category="Content Layout")
-	FMargin ContentPadding;
-
 	/** The scaling factor for the button border */
 	UPROPERTY(EditDefaultsOnly, Category=Appearance, AdvancedDisplay)
 	FVector2D DesiredSizeScale;
@@ -70,10 +58,6 @@ public:
 	/**  */
 	UFUNCTION(BlueprintCallable, Category="Appearance")
 	void SetForegroundColor(FLinearColor InForegroundColor);
-
-	/**  */
-	UFUNCTION(BlueprintCallable, Category="Appearance")
-	void SetContentPadding(FMargin InContentPadding);
 	
 	/**  */
 	UFUNCTION(BlueprintCallable, Category="Widget")
@@ -83,6 +67,10 @@ public:
 	virtual void SyncronizeProperties() override;
 	// End of UWidget interface
 
+	// Begin UObject interface
+	virtual void PostLoad() override;
+	// End of UObject interface
+
 #if WITH_EDITOR
 	virtual const FSlateBrush* GetEditorIcon() override;
 #endif
@@ -90,6 +78,7 @@ public:
 protected:
 
 	// UPanelWidget
+	virtual UClass* GetSlotClass() const override;
 	virtual void OnSlotAdded(UPanelSlot* Slot) override;
 	virtual void OnSlotRemoved(UPanelSlot* Slot) override;
 	// End UPanelWidget
diff --git a/Engine/Source/Runtime/UMG/Public/Components/ButtonSlot.h b/Engine/Source/Runtime/UMG/Public/Components/ButtonSlot.h
new file mode 100644
--- /dev/null
+++ b/Engine/Source/Runtime/UMG/Public/Components/ButtonSlot.h
@@ -0,0 +1,47 @@
+// Copyright 1998-2014 Epic Games, Inc. All Rights Reserved.
+
+#pragma once
+
+#include "SlateWrapperTypes.h"
+
+#include "ButtonSlot.generated.h"
+
+/** The Slot for the UButtonSlot, contains the widget displayed in a button's single slot */
+UCLASS()
+class UMG_API UButtonSlot : public UPanelSlot
+{
+	GENERATED_UCLASS_BODY()
+	
+	/** The padding area between the slot and the content it contains. */
+	UPROPERTY(EditDefaultsOnly, Category=Layout)
+	FMargin Padding;
+
+	/** The alignment of the object horizontally. */
+	UPROPERTY(EditDefaultsOnly, Category=Layout)
+	TEnumAsByte<EHorizontalAlignment> HorizontalAlignment;
+
+	/** The alignment of the object vertically. */
+	UPROPERTY(EditDefaultsOnly, Category=Layout)
+	TEnumAsByte<EVerticalAlignment> VerticalAlignment;
+
+	UFUNCTION(BlueprintCallable, Category="Appearance")
+	void SetPadding(FMargin InPadding);
+
+	UFUNCTION(BlueprintCallable, Category="Appearance")
+	void SetHorizontalAlignment(EHorizontalAlignment InHorizontalAlignment);
+
+	UFUNCTION(BlueprintCallable, Category="Appearance")
+	void SetVerticalAlignment(EVerticalAlignment InVerticalAlignment);
+
+	// UPanelSlot interface
+	virtual void SyncronizeProperties() override;
+	// End of UPanelSlot interface
+
+	/** Builds the underlying slot for the slate button. */
+	void BuildSlot(TSharedRef<SButton> InButton);
+
+private:
+
+	/** A pointer to the button to allow us to adjust the size, padding...etc at runtime. */
+	TSharedPtr<SButton> Button;
+};
diff --git a/Engine/Source/Runtime/UMG/Public/Components/ContentWidget.h b/Engine/Source/Runtime/UMG/Public/Components/ContentWidget.h
--- a/Engine/Source/Runtime/UMG/Public/Components/ContentWidget.h
+++ b/Engine/Source/Runtime/UMG/Public/Components/ContentWidget.h
@@ -9,12 +9,6 @@ class UMG_API UContentWidget : public UPanelWidget
 {
 	GENERATED_UCLASS_BODY()
 
-	//UFUNCTION()
-	//virtual UWidget* GetContent();
-	//
-	//UFUNCTION()
-	//virtual void SetContent(UWidget* InContent);
-
 protected:
 
 	// UPanelWidget
diff --git a/Engine/Source/Runtime/UMG/Public/UMG.h b/Engine/Source/Runtime/UMG/Public/UMG.h
--- a/Engine/Source/Runtime/UMG/Public/UMG.h
+++ b/Engine/Source/Runtime/UMG/Public/UMG.h
@@ -54,9 +54,11 @@
 #include "WidgetSwitcherSlot.h"
 #include "WidgetSwitcher.h"
 
+#include "ButtonSlot.h"
+#include "Button.h"
+
 #include "TextBlock.h"
 #include "Border.h"
-#include "Button.h"
 #include "CheckBox.h"
 #include "EditableText.h"
 #include "EditableTextBox.h"
