UMG - Adding a details customization for the FSlateChildSize struct.

diff --git a/Engine/Source/Editor/UMGEditor/Private/Customizations/SlateChildSizeCustomization.cpp b/Engine/Source/Editor/UMGEditor/Private/Customizations/SlateChildSizeCustomization.cpp
new file mode 100644
--- /dev/null
+++ b/Engine/Source/Editor/UMGEditor/Private/Customizations/SlateChildSizeCustomization.cpp
@@ -0,0 +1,138 @@
+// Copyright 1998-2014 Epic Games, Inc. All Rights Reserved.
+
+#include "UMGEditorPrivatePCH.h"
+
+#include "PropertyEditing.h"
+#include "ObjectEditorUtils.h"
+#include "WidgetGraphSchema.h"
+#include "ScopedTransaction.h"
+#include "BlueprintEditorUtils.h"
+
+#define LOCTEXT_NAMESPACE "UMG"
+
+void FSlateChildSizeCustomization::CustomizeHeader(TSharedRef<IPropertyHandle> PropertyHandle, FDetailWidgetRow& HeaderRow, IPropertyTypeCustomizationUtils& CustomizationUtils)
+{
+	TSharedPtr<IPropertyHandle> ValueHandle = PropertyHandle->GetChildHandle(GET_MEMBER_NAME_CHECKED(FSlateChildSize, Value));
+	TSharedPtr<IPropertyHandle> RuleHandle = PropertyHandle->GetChildHandle(GET_MEMBER_NAME_CHECKED(FSlateChildSize, SizeRule));
+
+	const FMargin OuterPadding(2);
+	const FMargin ContentPadding(2);
+
+	if ( !( ValueHandle.IsValid() || RuleHandle.IsValid() ) )
+	{
+		return;
+	}
+
+	HeaderRow
+	.NameContent()
+	[
+		PropertyHandle->CreatePropertyNameWidget()
+	]
+	.ValueContent()
+	[
+		SNew(SHorizontalBox)
+
+		+ SHorizontalBox::Slot()
+		.AutoWidth()
+		.Padding(OuterPadding)
+		[
+			SNew(SCheckBox)
+			.Style(FEditorStyle::Get(), "ToggleButtonCheckbox")
+			.ToolTipText(LOCTEXT("Auto_ToolTip", "Auto"))
+			.Padding(ContentPadding)
+			.OnCheckStateChanged(this, &FSlateChildSizeCustomization::HandleCheckStateChanged, RuleHandle, ESlateSizeRule::Automatic)
+			.IsChecked(this, &FSlateChildSizeCustomization::GetCheckState, RuleHandle, ESlateSizeRule::Automatic)
+			[
+				SNew(STextBlock)
+				.Text(LOCTEXT("Auto", "Auto"))
+			]
+		]
+
+		+ SHorizontalBox::Slot()
+		.AutoWidth()
+		.Padding(OuterPadding)
+		[
+			SNew(SCheckBox)
+			.Style(FEditorStyle::Get(), "ToggleButtonCheckbox")
+			.ToolTipText(LOCTEXT("Fill_ToolTip", "Fill"))
+			.Padding(ContentPadding)
+			.OnCheckStateChanged(this, &FSlateChildSizeCustomization::HandleCheckStateChanged, RuleHandle, ESlateSizeRule::Fill)
+			.IsChecked(this, &FSlateChildSizeCustomization::GetCheckState, RuleHandle, ESlateSizeRule::Fill)
+			[
+				SNew(STextBlock)
+				.Text(LOCTEXT("Fill", "Fill"))
+			]
+		]
+
+		+ SHorizontalBox::Slot()
+		.AutoWidth()
+		.Padding(OuterPadding)
+		[
+			SNew(SBox)
+			.WidthOverride(40)
+			[
+				SNew( SNumericEntryBox<float> )
+				.LabelVAlign(VAlign_Center)
+				.Visibility(this, &FSlateChildSizeCustomization::GetValueVisiblity, RuleHandle)
+				//.Label()
+				//[
+				//	SNew(STextBlock)
+				//	.Text(LOCTEXT("Fill", "Fill"))
+				//]
+				.Value(this, &FSlateChildSizeCustomization::GetValue, ValueHandle)
+				.OnValueCommitted(this, &FSlateChildSizeCustomization::HandleValueComitted, ValueHandle)
+				.UndeterminedString( LOCTEXT("MultipleValues", "Multiple Values") )
+			]
+		]
+	];
+}
+
+void FSlateChildSizeCustomization::CustomizeChildren(TSharedRef<IPropertyHandle> PropertyHandle, IDetailChildrenBuilder& ChildBuilder, IPropertyTypeCustomizationUtils& CustomizationUtils)
+{
+	
+}
+
+void FSlateChildSizeCustomization::HandleCheckStateChanged(ESlateCheckBoxState::Type InCheckboxState, TSharedPtr<IPropertyHandle> PropertyHandle, ESlateSizeRule::Type ToRule)
+{
+	PropertyHandle->SetValue((uint8)ToRule);
+}
+
+ESlateCheckBoxState::Type FSlateChildSizeCustomization::GetCheckState(TSharedPtr<IPropertyHandle> PropertyHandle, ESlateSizeRule::Type ForRule) const
+{
+	uint8 Value;
+	if ( PropertyHandle->GetValue(Value) )
+	{
+		return Value == ForRule ? ESlateCheckBoxState::Checked : ESlateCheckBoxState::Unchecked;
+	}
+
+	return ESlateCheckBoxState::Unchecked;
+}
+
+TOptional<float> FSlateChildSizeCustomization::GetValue(TSharedPtr<IPropertyHandle> ValueHandle) const
+{
+	float Value;
+	if ( ValueHandle->GetValue(Value) )
+	{
+		return Value;
+	}
+
+	return TOptional<float>();
+}
+
+void FSlateChildSizeCustomization::HandleValueComitted(float NewValue, ETextCommit::Type CommitType, TSharedPtr<IPropertyHandle> ValueHandle)
+{
+	ValueHandle->SetValue(NewValue);
+}
+
+EVisibility FSlateChildSizeCustomization::GetValueVisiblity(TSharedPtr<IPropertyHandle> RuleHandle) const
+{
+	uint8 Value;
+	if ( RuleHandle->GetValue(Value) )
+	{
+		return Value == ESlateSizeRule::Fill ? EVisibility::Visible : EVisibility::Collapsed;
+	}
+
+	return EVisibility::Collapsed;
+}
+
+#undef LOCTEXT_NAMESPACE
diff --git a/Engine/Source/Editor/UMGEditor/Private/Customizations/SlateChildSizeCustomization.h b/Engine/Source/Editor/UMGEditor/Private/Customizations/SlateChildSizeCustomization.h
new file mode 100644
--- /dev/null
+++ b/Engine/Source/Editor/UMGEditor/Private/Customizations/SlateChildSizeCustomization.h
@@ -0,0 +1,31 @@
+// Copyright 1998-2014 Epic Games, Inc. All Rights Reserved.
+
+#pragma once
+
+class FSlateChildSizeCustomization : public IPropertyTypeCustomization
+{
+public:
+	/** Makes a new instance of this detail layout class for a specific detail view requesting it */
+	static TSharedRef<class IPropertyTypeCustomization> MakeInstance()
+	{
+		return MakeShareable(new FSlateChildSizeCustomization());
+	}
+
+	FSlateChildSizeCustomization()
+	{
+	}
+	
+	virtual void CustomizeHeader(TSharedRef<IPropertyHandle> PropertyHandle, class FDetailWidgetRow& HeaderRow, IPropertyTypeCustomizationUtils& CustomizationUtils) override;
+
+	virtual void CustomizeChildren(TSharedRef<IPropertyHandle> PropertyHandle, class IDetailChildrenBuilder& ChildBuilder, IPropertyTypeCustomizationUtils& CustomizationUtils) override;
+
+	void HandleCheckStateChanged(ESlateCheckBoxState::Type InCheckboxState, TSharedPtr<IPropertyHandle> PropertyHandle, ESlateSizeRule::Type ToRule);
+
+	ESlateCheckBoxState::Type GetCheckState(TSharedPtr<IPropertyHandle> PropertyHandle, ESlateSizeRule::Type ForRule) const;
+
+	TOptional<float> GetValue(TSharedPtr<IPropertyHandle> ValueHandle) const;
+
+	void HandleValueComitted(float NewValue, ETextCommit::Type CommitType, TSharedPtr<IPropertyHandle> ValueHandle);
+
+	EVisibility GetValueVisiblity(TSharedPtr<IPropertyHandle> RuleHandle) const;
+};
diff --git a/Engine/Source/Editor/UMGEditor/Private/Details/SWidgetDetailsView.cpp b/Engine/Source/Editor/UMGEditor/Private/Details/SWidgetDetailsView.cpp
--- a/Engine/Source/Editor/UMGEditor/Private/Details/SWidgetDetailsView.cpp
+++ b/Engine/Source/Editor/UMGEditor/Private/Details/SWidgetDetailsView.cpp
@@ -101,14 +101,13 @@ void SWidgetDetailsView::RegisterCustomizations()
 	FOnGetDetailCustomizationInstance LayoutDelegateDetails = FOnGetDetailCustomizationInstance::CreateStatic(&FBlueprintWidgetCustomization::MakeInstance, BlueprintEditor.Pin()->GetBlueprintObj());
 	PropertyView->RegisterInstancedCustomPropertyLayout(UWidget::StaticClass(), LayoutDelegateDetails);
 
-	FOnGetPropertyTypeCustomizationInstance CanvasSlotCustomization = FOnGetPropertyTypeCustomizationInstance::CreateStatic(&FCanvasSlotCustomization::MakeInstance, BlueprintEditor.Pin()->GetBlueprintObj());
-
 	static FName PropertyEditor("PropertyEditor");
 	FPropertyEditorModule& PropertyModule = FModuleManager::GetModuleChecked<FPropertyEditorModule>(PropertyEditor);
-	PropertyModule.RegisterCustomPropertyTypeLayout(TEXT("PanelSlot"), CanvasSlotCustomization);
 
+	PropertyModule.RegisterCustomPropertyTypeLayout(TEXT("PanelSlot"), FOnGetPropertyTypeCustomizationInstance::CreateStatic(&FCanvasSlotCustomization::MakeInstance, BlueprintEditor.Pin()->GetBlueprintObj()));
 	PropertyModule.RegisterCustomPropertyTypeLayout(TEXT("EHorizontalAlignment"), FOnGetPropertyTypeCustomizationInstance::CreateStatic(&FHorizontalAlignmentCustomization::MakeInstance));
 	PropertyModule.RegisterCustomPropertyTypeLayout(TEXT("EVerticalAlignment"), FOnGetPropertyTypeCustomizationInstance::CreateStatic(&FVerticalAlignmentCustomization::MakeInstance));
+	PropertyModule.RegisterCustomPropertyTypeLayout(TEXT("SlateChildSize"), FOnGetPropertyTypeCustomizationInstance::CreateStatic(&FSlateChildSizeCustomization::MakeInstance));
 }
 
 void SWidgetDetailsView::OnEditorSelectionChanging()
diff --git a/Engine/Source/Editor/UMGEditor/Private/UMGEditorPrivatePCH.h b/Engine/Source/Editor/UMGEditor/Private/UMGEditorPrivatePCH.h
--- a/Engine/Source/Editor/UMGEditor/Private/UMGEditorPrivatePCH.h
+++ b/Engine/Source/Editor/UMGEditor/Private/UMGEditorPrivatePCH.h
@@ -39,6 +39,7 @@
 #include "CanvasSlotCustomization.h"
 #include "HorizontalAlignmentCustomization.h"
 #include "VerticalAlignmentCustomization.h"
+#include "SlateChildSizeCustomization.h"
 
 #include "WidgetTemplate.h"
 #include "WidgetTemplateDragDropOp.h"
diff --git a/Engine/Source/Runtime/UMG/Public/Components/Widget.h b/Engine/Source/Runtime/UMG/Public/Components/Widget.h
--- a/Engine/Source/Runtime/UMG/Public/Components/Widget.h
+++ b/Engine/Source/Runtime/UMG/Public/Components/Widget.h
@@ -96,7 +96,7 @@ public:
 	//TODO UMG Cursor doesn't work yet, the underlying slate version needs it to be TOptional.
 
 	/** The cursor to show when the mouse is over the widget */
-	UPROPERTY(EditDefaultsOnly, Category=Behavior)
+	UPROPERTY(EditDefaultsOnly, Category=Behavior, AdvancedDisplay)
 	TEnumAsByte<EMouseCursor::Type> Cursor;
 
 	/** A bindable delegate for Cursor */
