UMG - Creating a function binding now takes you to the newly added function!

diff --git a/Engine/Source/Editor/UMGEditor/Private/Customizations/DetailCustomizations.cpp b/Engine/Source/Editor/UMGEditor/Private/Customizations/DetailCustomizations.cpp
--- a/Engine/Source/Editor/UMGEditor/Private/Customizations/DetailCustomizations.cpp
+++ b/Engine/Source/Editor/UMGEditor/Private/Customizations/DetailCustomizations.cpp
@@ -377,6 +377,15 @@ void FBlueprintWidgetCustomization::HandleCreateAndAddBinding(TSharedRef<IProper
 	SelectedFunction->EdGraph = FunctionGraph;
 
 	HandleAddBinding(PropertyHandle, SelectedFunction);
+
+	GotoFunction(FunctionGraph);
+}
+
+void FBlueprintWidgetCustomization::GotoFunction(UEdGraph* FunctionGraph)
+{
+	Editor.Pin()->SetCurrentMode(FWidgetBlueprintApplicationModes::GraphMode);
+
+	Editor.Pin()->OpenDocument(FunctionGraph, FDocumentTracker::OpenNewDocument);
 }
 
 #undef LOCTEXT_NAMESPACE
diff --git a/Engine/Source/Editor/UMGEditor/Private/Customizations/DetailCustomizations.h b/Engine/Source/Editor/UMGEditor/Private/Customizations/DetailCustomizations.h
--- a/Engine/Source/Editor/UMGEditor/Private/Customizations/DetailCustomizations.h
+++ b/Engine/Source/Editor/UMGEditor/Private/Customizations/DetailCustomizations.h
@@ -6,13 +6,14 @@ class FBlueprintWidgetCustomization : public IDetailCustomization
 {
 public:
 	/** Makes a new instance of this detail layout class for a specific detail view requesting it */
-	static TSharedRef<class IDetailCustomization> MakeInstance(UBlueprint* Blueprint)
+	static TSharedRef<class IDetailCustomization> MakeInstance(TSharedRef<FWidgetBlueprintEditor> InEditor, UBlueprint* InBlueprint)
 	{
-		return MakeShareable(new FBlueprintWidgetCustomization(Blueprint));
+		return MakeShareable(new FBlueprintWidgetCustomization(InEditor, InBlueprint));
 	}
 
-	FBlueprintWidgetCustomization(UBlueprint* InBlueprint)
-		: Blueprint(CastChecked<UWidgetBlueprint>(InBlueprint))
+	FBlueprintWidgetCustomization(TSharedRef<FWidgetBlueprintEditor> InEditor, UBlueprint* InBlueprint)
+		: Editor(InEditor)
+		, Blueprint(CastChecked<UWidgetBlueprint>(InBlueprint))
 	{
 	}
 	
@@ -38,9 +39,12 @@ private:
 	void HandleRemoveBinding(TSharedRef<IPropertyHandle> PropertyHandle);
 	void HandleAddBinding(TSharedRef<IPropertyHandle> PropertyHandle, TSharedPtr<FunctionInfo> SelectedFunction);
 	void HandleCreateAndAddBinding(TSharedRef<IPropertyHandle> PropertyHandle, UFunction* DelegateSignature);
+	void GotoFunction(UEdGraph* FunctionGraph);
 
 private:
 
+	TWeakPtr<FWidgetBlueprintEditor> Editor;
+
 	UWidgetBlueprint* Blueprint;
 
 	TArray< TSharedPtr<FunctionInfo> > BlueprintFunctionCache;
diff --git a/Engine/Source/Editor/UMGEditor/Private/Details/SWidgetDetailsView.cpp b/Engine/Source/Editor/UMGEditor/Private/Details/SWidgetDetailsView.cpp
--- a/Engine/Source/Editor/UMGEditor/Private/Details/SWidgetDetailsView.cpp
+++ b/Engine/Source/Editor/UMGEditor/Private/Details/SWidgetDetailsView.cpp
@@ -98,7 +98,7 @@ SWidgetDetailsView::~SWidgetDetailsView()
 
 void SWidgetDetailsView::RegisterCustomizations()
 {
-	FOnGetDetailCustomizationInstance LayoutDelegateDetails = FOnGetDetailCustomizationInstance::CreateStatic(&FBlueprintWidgetCustomization::MakeInstance, BlueprintEditor.Pin()->GetBlueprintObj());
+	FOnGetDetailCustomizationInstance LayoutDelegateDetails = FOnGetDetailCustomizationInstance::CreateStatic(&FBlueprintWidgetCustomization::MakeInstance, BlueprintEditor.Pin().ToSharedRef(), BlueprintEditor.Pin()->GetBlueprintObj());
 	PropertyView->RegisterInstancedCustomPropertyLayout(UWidget::StaticClass(), LayoutDelegateDetails);
 
 	static FName PropertyEditor("PropertyEditor");
diff --git a/Engine/Source/Runtime/UMG/Public/Blueprint/UserWidget.h b/Engine/Source/Runtime/UMG/Public/Blueprint/UserWidget.h
--- a/Engine/Source/Runtime/UMG/Public/Blueprint/UserWidget.h
+++ b/Engine/Source/Runtime/UMG/Public/Blueprint/UserWidget.h
@@ -192,6 +192,8 @@ public:
 	UFUNCTION(BlueprintImplementableEvent, Category="User Interface")
 	FSReply OnMotionDetected(FGeometry MyGeometry, FMotionEvent InMotionEvent);
 
+	//virtual bool OnVisualizeTooltip(const TSharedPtr<SWidget>& TooltipContent);
+
 	/** @returns The UObject wrapper for a given SWidget */
 	UWidget* GetWidgetHandle(TSharedRef<SWidget> InWidget);
 
