UMG - Animations are now saved and can be played at runtime
- Added two functions to user widgets, play and stop, to control playback.

diff --git a/Engine/Source/Editor/Sequencer/Private/Sequencer.cpp b/Engine/Source/Editor/Sequencer/Private/Sequencer.cpp
--- a/Engine/Source/Editor/Sequencer/Private/Sequencer.cpp
+++ b/Engine/Source/Editor/Sequencer/Private/Sequencer.cpp
@@ -538,7 +538,7 @@ void FSequencer::GetRuntimeObjects( TSharedRef<FMovieSceneInstance> MovieSceneIn
 	}*/
 }
 
-void FSequencer::UpdateViewports(AActor* ActorToViewThrough) const
+void FSequencer::UpdatePreviewViewports(UObject* ObjectToViewThrough) const
 {
 	if(!IsPerspectiveViewportPosessionEnabled())
 	{
@@ -549,7 +549,7 @@ void FSequencer::UpdateViewports(AActor* ActorToViewThrough) const
 	{
 		if(LevelVC && LevelVC->IsPerspective() && LevelVC->AllowMatineePreview())
 		{
-			LevelVC->SetMatineeActorLock(ActorToViewThrough);
+			LevelVC->SetMatineeActorLock(Cast<AActor>(ObjectToViewThrough));
 		}
 	}
 }
diff --git a/Engine/Source/Editor/Sequencer/Private/Sequencer.h b/Engine/Source/Editor/Sequencer/Private/Sequencer.h
--- a/Engine/Source/Editor/Sequencer/Private/Sequencer.h
+++ b/Engine/Source/Editor/Sequencer/Private/Sequencer.h
@@ -254,7 +254,7 @@ public:
 
 	/** IMovieScenePlayer interface */
 	virtual void GetRuntimeObjects( TSharedRef<FMovieSceneInstance> MovieSceneInstance, const FGuid& ObjectHandle, TArray< UObject* >& OutObjects ) const;
-	virtual void UpdateViewports(AActor* ActorToViewThrough) const override;
+	virtual void UpdatePreviewViewports(UObject* ObjectToViewThrough) const override;
 	virtual EMovieScenePlayerStatus::Type GetPlaybackStatus() const override;
 	virtual void AddMovieSceneInstance( class UMovieSceneSection& MovieSceneSection, TSharedRef<FMovieSceneInstance> InstanceToAdd ) override;
 	virtual void RemoveMovieSceneInstance( class UMovieSceneSection& MovieSceneSection, TSharedRef<FMovieSceneInstance> InstanceToRemove ) override;
diff --git a/Engine/Source/Editor/UMGEditor/Private/UMGSequencerObjectBindingManager.cpp b/Engine/Source/Editor/UMGEditor/Private/UMGSequencerObjectBindingManager.cpp
--- a/Engine/Source/Editor/UMGEditor/Private/UMGSequencerObjectBindingManager.cpp
+++ b/Engine/Source/Editor/UMGEditor/Private/UMGSequencerObjectBindingManager.cpp
@@ -5,10 +5,12 @@
 #include "WidgetBlueprintEditor.h"
 #include "ISequencer.h"
 
-FUMGSequencerObjectBindingManager::FUMGSequencerObjectBindingManager( FWidgetBlueprintEditor& InWidgetBlueprintEditor )
+FUMGSequencerObjectBindingManager::FUMGSequencerObjectBindingManager( FWidgetBlueprintEditor& InWidgetBlueprintEditor, UMovieScene& InMovieScene )
 	: WidgetBlueprintEditor( InWidgetBlueprintEditor )
+	, MovieScene( &InMovieScene )
 {
 	WidgetBlueprintEditor.GetOnWidgetPreviewUpdated().AddRaw( this, &FUMGSequencerObjectBindingManager::OnWidgetPreviewUpdated );
+
 }
 
 FUMGSequencerObjectBindingManager::~FUMGSequencerObjectBindingManager()
@@ -24,12 +26,24 @@ bool FUMGSequencerObjectBindingManager::CanPossessObject( UObject& Object ) cons
 {
 	UWidgetBlueprint* WidgetBlueprint = WidgetBlueprintEditor.GetWidgetBlueprintObj();
 
-	return !Object.IsIn( WidgetBlueprint );
+	// Only preview widgets in this blueprint can be possessed
+	UUserWidget* PreviewWidget = WidgetBlueprintEditor.GetPreview();
+
+	return Object.IsA<UVisual>() && Object.IsIn( PreviewWidget );
 }
 
 void FUMGSequencerObjectBindingManager::BindPossessableObject( const FGuid& PossessableGuid, UObject& PossessedObject )
 {
 	PreviewObjectToGuidMap.Add( &PossessedObject, PossessableGuid );
+
+	UWidgetBlueprint* WidgetBlueprint = WidgetBlueprintEditor.GetWidgetBlueprintObj();
+	FWidgetAnimation* WidgetAnimation = WidgetBlueprint->FindAnimationDataForMovieScene(*MovieScene);
+
+	FWidgetAnimationBinding NewBinding;
+	NewBinding.WidgetName = PossessedObject.GetFName();
+	NewBinding.AnimationGuid = PossessableGuid;
+
+	WidgetAnimation->AnimationBindings.Add( NewBinding );
 }
 
 void FUMGSequencerObjectBindingManager::UnbindPossessableObjects( const FGuid& PossessableGuid )
@@ -41,6 +55,12 @@ void FUMGSequencerObjectBindingManager::UnbindPossessableObjects( const FGuid& P
 			It.RemoveCurrent();
 		}
 	}
+
+	UWidgetBlueprint* WidgetBlueprint = WidgetBlueprintEditor.GetWidgetBlueprintObj();
+	FWidgetAnimation* WidgetAnimation = WidgetBlueprint->FindAnimationDataForMovieScene(*MovieScene);
+
+	WidgetAnimation->AnimationBindings.RemoveAll( [&]( const FWidgetAnimationBinding& Binding ) { return Binding.AnimationGuid == PossessableGuid; } );
+
 }
 
 void FUMGSequencerObjectBindingManager::GetRuntimeObjects( const TSharedRef<FMovieSceneInstance>& MovieSceneInstance, const FGuid& ObjectGuid, TArray<UObject*>& OutRuntimeObjects ) const
@@ -55,30 +75,34 @@ void FUMGSequencerObjectBindingManager::GetRuntimeObjects( const TSharedRef<FMov
 	}
 }
 
-void FUMGSequencerObjectBindingManager::OnWidgetPreviewUpdated()
+void FUMGSequencerObjectBindingManager::InitPreviewObjects()
 {
-	//
+	UWidgetBlueprint* WidgetBlueprint = WidgetBlueprintEditor.GetWidgetBlueprintObj();
+
+	// Populate preview object to guid map
 	UUserWidget* PreviewWidget = WidgetBlueprintEditor.GetPreview();
 
 	UWidgetTree* WidgetTree = PreviewWidget->WidgetTree;
 
-	TMap< TWeakObjectPtr<UObject>, FGuid> OldPreviewObjectToGuidMap = PreviewObjectToGuidMap;
-	PreviewObjectToGuidMap.Empty();
+	const FWidgetAnimation* WidgetAnimation = WidgetBlueprint->FindAnimationDataForMovieScene( *MovieScene );
+	check( WidgetAnimation );
 
-	for( auto It = OldPreviewObjectToGuidMap.CreateConstIterator(); It; ++It )
+	for( const FWidgetAnimationBinding& Binding : WidgetAnimation->AnimationBindings )
 	{
-		const TWeakObjectPtr<UObject> Object = It.Key();
-		if(Object.IsValid())
+		FName ObjectName = Binding.WidgetName;
+		UObject* FoundObject = FindObject<UObject>(WidgetTree, *ObjectName.ToString());
+		if(FoundObject)
 		{
-			FName ObjectName = Object->GetFName();
-			UObject* NewObject = FindObject<UObject>( WidgetTree, *ObjectName.ToString() );
-			if( NewObject )
-			{
-				check( NewObject->GetClass() == Object->GetClass() );
-				PreviewObjectToGuidMap.Add( NewObject, It.Value() );
-			}
+			PreviewObjectToGuidMap.Add(FoundObject, Binding.AnimationGuid);
 		}
 	}
+}
+
+void FUMGSequencerObjectBindingManager::OnWidgetPreviewUpdated()
+{
+	PreviewObjectToGuidMap.Empty();
+
+	InitPreviewObjects();
 
 	WidgetBlueprintEditor.GetSequencer()->UpdateRuntimeInstances();
 }
\ No newline at end of file
diff --git a/Engine/Source/Editor/UMGEditor/Private/UMGSequencerObjectBindingManager.h b/Engine/Source/Editor/UMGEditor/Private/UMGSequencerObjectBindingManager.h
--- a/Engine/Source/Editor/UMGEditor/Private/UMGSequencerObjectBindingManager.h
+++ b/Engine/Source/Editor/UMGEditor/Private/UMGSequencerObjectBindingManager.h
@@ -7,7 +7,7 @@
 class FUMGSequencerObjectBindingManager : public ISequencerObjectBindingManager
 {
 public:
-	FUMGSequencerObjectBindingManager( FWidgetBlueprintEditor& InWidgetBlueprintEditor );
+	FUMGSequencerObjectBindingManager( FWidgetBlueprintEditor& InWidgetBlueprintEditor, UMovieScene& InMovieScene );
 	~FUMGSequencerObjectBindingManager();
 
 	/** ISequencerObjectBindingManager interface */
@@ -20,11 +20,16 @@ public:
 	virtual void UnbindPossessableObjects( const FGuid& PossessableGuid ) override;
 	virtual void GetRuntimeObjects( const TSharedRef<FMovieSceneInstance>& MovieSceneInstance, const FGuid& ObjectGuid, TArray<UObject*>& OutRuntimeObjects ) const override;
 	
+	/** Rebuilds mappings to live preview objects */
+	void InitPreviewObjects();
+
 private:
 	void OnWidgetPreviewUpdated();
 private:
 	/** Mapping of preview objects to sequencer guids */
 	TMap< TWeakObjectPtr<UObject>, FGuid> PreviewObjectToGuidMap;
-	
+
+	TWeakObjectPtr<UMovieScene> MovieScene;
+
 	FWidgetBlueprintEditor& WidgetBlueprintEditor;
 };
diff --git a/Engine/Source/Editor/UMGEditor/Private/WidgetBlueprint.cpp b/Engine/Source/Editor/UMGEditor/Private/WidgetBlueprint.cpp
--- a/Engine/Source/Editor/UMGEditor/Private/WidgetBlueprint.cpp
+++ b/Engine/Source/Editor/UMGEditor/Private/WidgetBlueprint.cpp
@@ -56,6 +56,10 @@ bool UWidgetBlueprint::AllowsDynamicBinding() const
 	return true;
 }
 
+FWidgetAnimation* UWidgetBlueprint::FindAnimationDataForMovieScene( UMovieScene& MovieScene )
+{
+	return AnimationData.FindByPredicate( [&](const FWidgetAnimation& Data) { return Data.MovieScene == &MovieScene; } );
+}
 
 bool UWidgetBlueprint::ValidateGeneratedClass(const UClass* InClass)
 {
diff --git a/Engine/Source/Editor/UMGEditor/Private/WidgetBlueprintCompiler.cpp b/Engine/Source/Editor/UMGEditor/Private/WidgetBlueprintCompiler.cpp
--- a/Engine/Source/Editor/UMGEditor/Private/WidgetBlueprintCompiler.cpp
+++ b/Engine/Source/Editor/UMGEditor/Private/WidgetBlueprintCompiler.cpp
@@ -486,6 +486,9 @@ void FWidgetBlueprintCompiler::CreateClassVariablesFromBlueprint()
 			WidgetToMemberVariableMap.Add(Widget, WidgetProperty);
 		}
 	}
+
+	// Add movie scenes variables here
+
 }
 
 void FWidgetBlueprintCompiler::FinishCompilingClass(UClass* Class)
@@ -496,11 +499,16 @@ void FWidgetBlueprintCompiler::FinishCompilingClass(UClass* Class)
 
 	BPGClass->WidgetTree = DuplicateObject<UWidgetTree>(Blueprint->WidgetTree, BPGClass);
 
-	BPGClass->AnimationData.Reset();
-	for ( UMovieScene* Sequence : Blueprint->AnimationData )
+	BPGClass->AnimationData.Empty();
+	for ( const FWidgetAnimation& Animation : Blueprint->AnimationData )
 	{
-		UMovieScene* ClonedSequence = DuplicateObject<UMovieScene>(Sequence, BPGClass);
-		BPGClass->AnimationData.Add( ClonedSequence );
+		UMovieScene* ClonedSequence = DuplicateObject<UMovieScene>(Animation.MovieScene, BPGClass, *Animation.MovieScene->GetName());
+		FWidgetAnimation CompiledAnimation;
+
+		CompiledAnimation.MovieScene = ClonedSequence;
+		CompiledAnimation.AnimationBindings = Animation.AnimationBindings;
+
+		BPGClass->AnimationData.Add( CompiledAnimation );
 	}
 
 	BPGClass->Bindings.Reset();
diff --git a/Engine/Source/Editor/UMGEditor/Private/WidgetBlueprintEditor.cpp b/Engine/Source/Editor/UMGEditor/Private/WidgetBlueprintEditor.cpp
--- a/Engine/Source/Editor/UMGEditor/Private/WidgetBlueprintEditor.cpp
+++ b/Engine/Source/Editor/UMGEditor/Private/WidgetBlueprintEditor.cpp
@@ -20,7 +20,6 @@
 FWidgetBlueprintEditor::FWidgetBlueprintEditor()
 	: PreviewScene(FPreviewScene::ConstructionValues().AllowAudioPlayback(true).ShouldSimulatePhysics(true))
 	, PreviewBlueprint(NULL)
-	, DefaultMovieScene(NULL)
 {
 }
 
@@ -52,6 +51,8 @@ void FWidgetBlueprintEditor::InitWidgetBlueprintEditor(const EToolkitMode::Type
 
 	UpdatePreview(GetWidgetBlueprintObj(), true);
 
+	SequencerObjectBindingManager->InitPreviewObjects();
+
 	WidgetCommandList = MakeShareable(new FUICommandList);
 
 	WidgetCommandList->MapAction(FGenericCommands::Get().Delete,
@@ -333,10 +334,6 @@ static bool MigratePropertyValue(UObject* SourceObject, UObject* DestinationObje
 void FWidgetBlueprintEditor::AddReferencedObjects( FReferenceCollector& Collector )
 {
 	FBlueprintEditor::AddReferencedObjects( Collector );
-	if( DefaultMovieScene != nullptr )
-	{
-		Collector.AddReferencedObject( DefaultMovieScene );
-	}
 
 	UUserWidget* Preview = GetPreview();
 	Collector.AddReferencedObject( Preview );
@@ -384,11 +381,28 @@ TSharedPtr<ISequencer>& FWidgetBlueprintEditor::GetSequencer()
 	{
 		UWidgetBlueprint* Blueprint = GetWidgetBlueprintObj();
 
-		UMovieScene* MovieScene = Blueprint->AnimationData.Num() ? Blueprint->AnimationData[0] : GetDefaultMovieScene();
-		
-		TSharedRef<FUMGSequencerObjectBindingManager> ObjectBindingManager = MakeShareable( new FUMGSequencerObjectBindingManager( *this ) );
-		
-		Sequencer = FModuleManager::LoadModuleChecked< ISequencerModule >("Sequencer").CreateSequencer( MovieScene, ObjectBindingManager );
+		FWidgetAnimation* WidgetAnimation = nullptr;
+
+		if( Blueprint->AnimationData.Num() )
+		{
+			WidgetAnimation = &Blueprint->AnimationData[0];
+		}
+		else
+		{
+			UMovieScene* MovieScene = ConstructObject<UMovieScene>( UMovieScene::StaticClass(), Blueprint, MakeUniqueObjectName( Blueprint, UMovieScene::StaticClass(), "DefaultAnimationData"), RF_Transactional );
+			FWidgetAnimation NewAnimation;
+			NewAnimation.MovieScene = MovieScene;
+			int32 Index = Blueprint->AnimationData.Add( NewAnimation );
+
+			WidgetAnimation = &Blueprint->AnimationData[Index];
+		}
+
+		TSharedRef<FUMGSequencerObjectBindingManager> ObjectBindingManager = MakeShareable(new FUMGSequencerObjectBindingManager(*this, *WidgetAnimation->MovieScene) );
+		check( !SequencerObjectBindingManager.IsValid() );
+
+		SequencerObjectBindingManager = ObjectBindingManager;
+
+		Sequencer = FModuleManager::LoadModuleChecked< ISequencerModule >("Sequencer").CreateSequencer( WidgetAnimation->MovieScene, ObjectBindingManager );
 	}
 
 	return Sequencer;
@@ -460,16 +474,6 @@ void FWidgetBlueprintEditor::UpdatePreview(UBlueprint* InBlueprint, bool bInForc
 	OnWidgetPreviewUpdated.Broadcast();
 }
 
-UMovieScene* FWidgetBlueprintEditor::GetDefaultMovieScene()
-{
-	if( !DefaultMovieScene )
-	{
-		FName ObjectName = MakeUniqueObjectName( GetTransientPackage(), UMovieScene::StaticClass(), "DefaultAnimationData" );
-		DefaultMovieScene = ConstructObject<UMovieScene>( UMovieScene::StaticClass(), GetTransientPackage(), ObjectName );
-	}
-
-	return DefaultMovieScene;
-}
 
 FGraphAppearanceInfo FWidgetBlueprintEditor::GetGraphAppearance() const
 {
diff --git a/Engine/Source/Editor/UMGEditor/Private/WidgetBlueprintEditor.h b/Engine/Source/Editor/UMGEditor/Private/WidgetBlueprintEditor.h
--- a/Engine/Source/Editor/UMGEditor/Private/WidgetBlueprintEditor.h
+++ b/Engine/Source/Editor/UMGEditor/Private/WidgetBlueprintEditor.h
@@ -109,12 +109,12 @@ private:
 	/** Sequencer for creating and previewing widget animations */
 	TSharedPtr<ISequencer> Sequencer;
 
+	/** Manager for handling bindings to sequence animations */
+	TSharedPtr<class FUMGSequencerObjectBindingManager> SequencerObjectBindingManager;
+
 	/** The Blueprint associated with the current preview */
 	UWidgetBlueprint* PreviewBlueprint;
 
-	/** Default movie scene for new animations */
-	UMovieScene* DefaultMovieScene;
-
 	/** The currently selected preview widgets in the preview GUI */
 	TSet<FWidgetReference> SelectedWidgets;
 
diff --git a/Engine/Source/Editor/UMGEditor/Private/WidgetBlueprintEditorUtils.cpp b/Engine/Source/Editor/UMGEditor/Private/WidgetBlueprintEditorUtils.cpp
--- a/Engine/Source/Editor/UMGEditor/Private/WidgetBlueprintEditorUtils.cpp
+++ b/Engine/Source/Editor/UMGEditor/Private/WidgetBlueprintEditorUtils.cpp
@@ -146,6 +146,18 @@ bool FWidgetBlueprintEditorUtils::RenameWidget(TSharedRef<FWidgetBlueprintEditor
 			}
 		}
 
+		// Update widget blueprint names
+		for( FWidgetAnimation& WidgetAnimation : Blueprint->AnimationData )
+		{
+			for( FWidgetAnimationBinding& AnimBinding : WidgetAnimation.AnimationBindings )
+			{
+				if( AnimBinding.WidgetName == OldName )
+				{
+					AnimBinding.WidgetName = NewName;
+				}
+			}
+		}
+
 		// Validate child blueprints and adjust variable names to avoid a potential name collision
 		FBlueprintEditorUtils::ValidateBlueprintChildVariables(Blueprint, NewName);
 
diff --git a/Engine/Source/Editor/UMGEditor/Public/WidgetBlueprint.h b/Engine/Source/Editor/UMGEditor/Public/WidgetBlueprint.h
--- a/Engine/Source/Editor/UMGEditor/Public/WidgetBlueprint.h
+++ b/Engine/Source/Editor/UMGEditor/Public/WidgetBlueprint.h
@@ -54,11 +54,11 @@ public:
 	TArray< FDelegateEditorBinding > Bindings;
 
 	UPROPERTY()
-	TArray< UMovieScene* > AnimationData;
-	 
+	TArray<FWidgetAnimation> AnimationData;
+
+	/** UObject interface */
 	virtual void PostLoad() override;
 	virtual void PostLoadSubobjects(FObjectInstancingGraph* OuterInstanceGraph) override;
-
 	virtual UClass* RegenerateClass(UClass* ClassToRegenerate, UObject* PreviousCDO, TArray<UObject*>& ObjLoaded) override;
 	
 	// UBlueprint interface
@@ -72,5 +72,7 @@ public:
 	}
 	// End of UBlueprint interface
 
+	FWidgetAnimation* FindAnimationDataForMovieScene( UMovieScene& MovieScene );
+
 	static bool ValidateGeneratedClass(const UClass* InClass);
 };
\ No newline at end of file
diff --git a/Engine/Source/Runtime/MovieSceneCore/Classes/RuntimeMovieScenePlayer.h b/Engine/Source/Runtime/MovieSceneCore/Classes/RuntimeMovieScenePlayer.h
--- a/Engine/Source/Runtime/MovieSceneCore/Classes/RuntimeMovieScenePlayer.h
+++ b/Engine/Source/Runtime/MovieSceneCore/Classes/RuntimeMovieScenePlayer.h
@@ -14,9 +14,8 @@
  * RuntimeMovieScenePlayer is used to actually "play" a MovieScene asset at runtime.  It keeps track of playback
  * state and provides functions for manipulating a MovieScene while its playing.
  */
-typedef IMovieScenePlayer FMovieScenePlayerInterface;	// @todo sequencer uobject: Hack for header generator
 UCLASS( MinimalAPI )
-class URuntimeMovieScenePlayer : public UObject, public FMovieScenePlayerInterface, public IRuntimeMovieScenePlayerInterface
+class URuntimeMovieScenePlayer : public UObject, public IMovieScenePlayer, public IRuntimeMovieScenePlayerInterface
 {
 	GENERATED_UCLASS_BODY()
 
@@ -59,7 +58,7 @@ protected:
 	virtual void SpawnActorsForMovie( TSharedRef<FMovieSceneInstance> MovieSceneInstance );
 	virtual void DestroyActorsForMovie( TSharedRef<FMovieSceneInstance> MovieSceneInstance );
 	virtual void GetRuntimeObjects( TSharedRef<FMovieSceneInstance> MovieSceneInstance, const FGuid& ObjectHandle, TArray< UObject* >& OutObjects ) const override;
-	virtual void UpdateViewports(AActor* ActorToViewThrough) const override {}
+	virtual void UpdatePreviewViewports(UObject* ObjectToViewThrough) const override {}
 	virtual EMovieScenePlayerStatus::Type GetPlaybackStatus() const override;
 	virtual void AddMovieSceneInstance( class UMovieSceneSection& MovieSceneSection, TSharedRef<FMovieSceneInstance> InstanceToAdd ) override;
 	virtual void RemoveMovieSceneInstance( class UMovieSceneSection& MovieSceneSection, TSharedRef<FMovieSceneInstance> InstanceToRemove ) override;
diff --git a/Engine/Source/Runtime/MovieSceneCore/Public/IMovieScenePlayer.h b/Engine/Source/Runtime/MovieSceneCore/Public/IMovieScenePlayer.h
--- a/Engine/Source/Runtime/MovieSceneCore/Public/IMovieScenePlayer.h
+++ b/Engine/Source/Runtime/MovieSceneCore/Public/IMovieScenePlayer.h
@@ -37,9 +37,9 @@ public:
 	/**
 	 * Updates the perspective viewports with the actor to view through
 	 *
-	 * @param ActorToViewThrough The actor, probably a camera, that the viewports should lock to
+	 * @param ObjectToViewThrough The object, probably a camera, that the viewports should lock to
 	 */
-	virtual void UpdateViewports(AActor* ActorToViewThrough) const = 0;
+	virtual void UpdatePreviewViewports(UObject* ObjectToViewThrough) const = 0;
 
 	/**
 	 * Adds a MovieScene instance to the player.  MovieScene players need to know about each instance for actor spawning
diff --git a/Engine/Source/Runtime/MovieSceneCoreTypes/Private/Instances/MovieSceneDirectorTrackInstance.cpp b/Engine/Source/Runtime/MovieSceneCoreTypes/Private/Instances/MovieSceneDirectorTrackInstance.cpp
--- a/Engine/Source/Runtime/MovieSceneCoreTypes/Private/Instances/MovieSceneDirectorTrackInstance.cpp
+++ b/Engine/Source/Runtime/MovieSceneCoreTypes/Private/Instances/MovieSceneDirectorTrackInstance.cpp
@@ -39,6 +39,6 @@ void FMovieSceneDirectorTrackInstance::Update( float Position, float LastPositio
 		}
 	}
 
-	Player.UpdateViewports(PossessedCamera);
+	Player.UpdatePreviewViewports(PossessedCamera);
 }
 
diff --git a/Engine/Source/Runtime/UMG/Private/Animation/UMGSequencePlayer.cpp b/Engine/Source/Runtime/UMG/Private/Animation/UMGSequencePlayer.cpp
new file mode 100644
--- /dev/null
+++ b/Engine/Source/Runtime/UMG/Private/Animation/UMGSequencePlayer.cpp
@@ -0,0 +1,114 @@
+// Copyright 1998-2014 Epic Games, Inc. All Rights Reserved.
+
+#include "UMGPrivatePCH.h"
+#include "UMGSequencePlayer.h"
+#include "MovieSceneBindings.h"
+#include "MovieSceneInstance.h"
+
+UUMGSequencePlayer::UUMGSequencePlayer(const FPostConstructInitializeProperties& PCIP)
+	: Super(PCIP)
+{
+	PlayerStatus = EMovieScenePlayerStatus::Stopped;
+	TimeCursorPosition = 0.0f;
+	MovieScene = nullptr;
+}
+
+void UUMGSequencePlayer::InitSequencePlayer( const FWidgetAnimation& Animation, UUserWidget& UserWidget )
+{
+	MovieScene = Animation.MovieScene;
+
+	// Cache the time range of the sequence to determine when we stop
+	TimeRange = MovieScene->GetTimeRange();
+
+	RuntimeBindings = ConstructObject<UMovieSceneBindings>( UMovieSceneBindings::StaticClass(), this );
+	RuntimeBindings->SetRootMovieScene( MovieScene );
+
+	UWidgetTree* WidgetTree = UserWidget.WidgetTree;
+
+	TMap<FGuid, TArray<UObject*> > GuidToRuntimeObjectMap;
+	// Bind to Runtime Objects
+	for (const FWidgetAnimationBinding& Binding : Animation.AnimationBindings)
+	{
+		FName ObjectName = Binding.WidgetName;
+		UObject* FoundObject = FindObject<UObject>(WidgetTree, *ObjectName.ToString());
+		if (FoundObject)
+		{
+			TArray<UObject*>& Objects = GuidToRuntimeObjectMap.FindOrAdd( Binding.AnimationGuid );
+			Objects.Add( FoundObject );
+		
+		}
+	}
+
+	for( auto It = GuidToRuntimeObjectMap.CreateConstIterator(); It; ++It )
+	{
+		RuntimeBindings->AddBinding( It.Key(), It.Value() );
+	}
+
+}
+
+void UUMGSequencePlayer::Tick(float DeltaTime)
+{
+	if ( PlayerStatus == EMovieScenePlayerStatus::Playing )
+	{
+		double LastTimePosition = TimeCursorPosition;
+
+		TimeCursorPosition += DeltaTime;
+
+		if (RootMovieSceneInstance.IsValid())
+		{
+			RootMovieSceneInstance->Update(TimeCursorPosition, LastTimePosition, *this);
+		}
+
+		bool bIsFinished = false;
+
+		if (!TimeRange.Contains(TimeCursorPosition))
+		{
+			bIsFinished = true;
+		}
+
+		if (bIsFinished)
+		{
+
+			PlayerStatus = EMovieScenePlayerStatus::Stopped;
+
+			OnSequenceFinishedPlayingEvent.Broadcast( *this );
+
+		}
+	}
+}
+
+void UUMGSequencePlayer::Play()
+{
+	RootMovieSceneInstance = MakeShareable( new FMovieSceneInstance( *MovieScene ) );
+
+	RootMovieSceneInstance->RefreshInstance( *this );
+
+	// @todo UMG Support user specified start times
+	TimeCursorPosition = 0;
+
+	PlayerStatus = EMovieScenePlayerStatus::Playing;
+}
+
+void UUMGSequencePlayer::Pause()
+{
+	PlayerStatus = EMovieScenePlayerStatus::Stopped;
+}
+
+void UUMGSequencePlayer::Stop()
+{
+	PlayerStatus = EMovieScenePlayerStatus::Stopped;
+
+	OnSequenceFinishedPlayingEvent.Broadcast(*this);
+
+	TimeCursorPosition = 0;
+}
+
+void UUMGSequencePlayer::GetRuntimeObjects( TSharedRef<FMovieSceneInstance> MovieSceneInstance, const FGuid& ObjectHandle, TArray< UObject* >& OutObjects ) const
+{
+	OutObjects = RuntimeBindings->FindBoundObjects( ObjectHandle );
+}
+
+EMovieScenePlayerStatus::Type UUMGSequencePlayer::GetPlaybackStatus() const
+{
+	return PlayerStatus;
+}
\ No newline at end of file
diff --git a/Engine/Source/Runtime/UMG/Private/Slate/SObjectWidget.cpp b/Engine/Source/Runtime/UMG/Private/Slate/SObjectWidget.cpp
--- a/Engine/Source/Runtime/UMG/Private/Slate/SObjectWidget.cpp
+++ b/Engine/Source/Runtime/UMG/Private/Slate/SObjectWidget.cpp
@@ -25,7 +25,7 @@ void SObjectWidget::Tick(const FGeometry& AllottedGeometry, const double InCurre
 {
 	if ( WidgetObject && !WidgetObject->IsDesignTime() )
 	{
-		return WidgetObject->Tick(AllottedGeometry, InDeltaTime);
+		return WidgetObject->NativeTick(AllottedGeometry, InDeltaTime);
 	}
 }
 
diff --git a/Engine/Source/Runtime/UMG/Private/UserWidget.cpp b/Engine/Source/Runtime/UMG/Private/UserWidget.cpp
--- a/Engine/Source/Runtime/UMG/Private/UserWidget.cpp
+++ b/Engine/Source/Runtime/UMG/Private/UserWidget.cpp
@@ -2,6 +2,7 @@
 
 #include "UMGPrivatePCH.h"
 
+#include "UMGSequencePlayer.h"
 #include "SceneViewport.h"
 
 TSharedRef<FUMGDragDropOp> FUMGDragDropOp::New()
@@ -180,6 +181,59 @@ UWorld* UUserWidget::GetWorld() const
 	return NULL;
 }
 
+void UUserWidget::PlayAnimation(FName AnimationName)
+{
+	UWidgetBlueprintGeneratedClass* BPClass = Cast<UWidgetBlueprintGeneratedClass>(GetClass());
+	if (BPClass)
+	{
+		const FWidgetAnimation* Animation = BPClass->FindAnimation(AnimationName);
+		if( Animation )
+		{
+			// @todo UMG sequencer - Restart animations which have had Play called on them?
+			UUMGSequencePlayer** FoundPlayer = ActiveSequencePlayers.FindByPredicate( [&](const UUMGSequencePlayer* Player) { return Player->GetMovieScene() == Animation->MovieScene; } );
+
+			if( !FoundPlayer )
+			{
+				UUMGSequencePlayer* NewPlayer = ConstructObject<UUMGSequencePlayer>( UUMGSequencePlayer::StaticClass(), this );
+				ActiveSequencePlayers.Add( NewPlayer );
+
+				NewPlayer->OnSequenceFinishedPlaying().AddUObject( this, &UUserWidget::OnAnimationFinishedPlaying );
+
+				NewPlayer->InitSequencePlayer( *Animation, *this );
+
+				NewPlayer->Play();
+			}
+			else
+			{
+				(*FoundPlayer)->Play();
+			}
+		}
+	}
+}
+
+void UUserWidget::StopAnimation(FName AnimationName)
+{
+	UWidgetBlueprintGeneratedClass* BPClass = Cast<UWidgetBlueprintGeneratedClass>(GetClass());
+	if(BPClass)
+	{
+		const FWidgetAnimation* Animation = BPClass->FindAnimation(AnimationName);
+		if(Animation)
+		{
+			// @todo UMG sequencer - Restart animations which have had Play called on them?
+			UUMGSequencePlayer** FoundPlayer = ActiveSequencePlayers.FindByPredicate([&](const UUMGSequencePlayer* Player) { return Player->GetMovieScene() == Animation->MovieScene; });
+			if(FoundPlayer)
+			{
+				(*FoundPlayer)->Stop();
+			}
+		}
+	}
+}
+
+void UUserWidget::OnAnimationFinishedPlaying( UUMGSequencePlayer& Player )
+{
+	ActiveSequencePlayers.Remove( &Player );
+}
+
 UWidget* UUserWidget::GetWidgetHandle(TSharedRef<SWidget> InWidget)
 {
 	TWeakObjectPtr<UWidget> VisualWidget = WidgetToComponent.FindRef(InWidget);
@@ -246,6 +300,17 @@ UWidget* UUserWidget::GetHandleFromName(const FString& Name) const
 	return NULL;
 }
 
+void UUserWidget::NativeTick(const FGeometry& MyGeometry, float InDeltaTime )
+{
+	// Update active movie scenes
+	for( UUMGSequencePlayer* Player : ActiveSequencePlayers )
+	{
+		Player->Tick( InDeltaTime );
+	}
+
+	Tick( MyGeometry, InDeltaTime );
+}
+
 TSharedRef<SWidget> UUserWidget::MakeFullScreenWidget()
 {
 	if ( bAbsoluteLayout )
diff --git a/Engine/Source/Runtime/UMG/Private/WidgetBlueprintGeneratedClass.cpp b/Engine/Source/Runtime/UMG/Private/WidgetBlueprintGeneratedClass.cpp
--- a/Engine/Source/Runtime/UMG/Private/WidgetBlueprintGeneratedClass.cpp
+++ b/Engine/Source/Runtime/UMG/Private/WidgetBlueprintGeneratedClass.cpp
@@ -1,6 +1,7 @@
 // Copyright 1998-2014 Epic Games, Inc. All Rights Reserved.
 
 #include "UMGPrivatePCH.h"
+#include "MovieScene.h"
 
 #define LOCTEXT_NAMESPACE "UMG"
 
@@ -62,6 +63,8 @@ void UWidgetBlueprintGeneratedClass::InitializeWidget(UUserWidget* UserWidget) c
 		TArray<UWidget*> ClonedWidgets;
 		ClonedTree->GetAllWidgets(ClonedWidgets);
 
+		// @todo UMG Create Movie Scene vars here
+
 		for ( UWidget* Widget : ClonedWidgets )
 		{
 			// Not fatal if NULL, but shouldn't happen
@@ -119,4 +122,9 @@ void UWidgetBlueprintGeneratedClass::InitializeWidget(UUserWidget* UserWidget) c
 	}
 }
 
-#undef LOCTEXT_NAMESPACE
\ No newline at end of file
+const FWidgetAnimation* UWidgetBlueprintGeneratedClass::FindAnimation( FName AnimationName ) const
+{
+	return AnimationData.FindByPredicate( [&](const FWidgetAnimation& Animation) { return Animation.MovieScene->GetFName() == AnimationName; } );
+}
+
+#undef LOCTEXT_NAMESPACE
diff --git a/Engine/Source/Runtime/UMG/Public/Animation/UMGSequencePlayer.h b/Engine/Source/Runtime/UMG/Public/Animation/UMGSequencePlayer.h
new file mode 100644
--- /dev/null
+++ b/Engine/Source/Runtime/UMG/Public/Animation/UMGSequencePlayer.h
@@ -0,0 +1,64 @@
+// Copyright 1998-2014 Epic Games, Inc. All Rights Reserved.
+
+#pragma once
+
+#include "IMovieScenePlayer.h"
+#include "UMGSequencePlayer.generated.h"
+
+UCLASS(transient)
+class UMG_API UUMGSequencePlayer : public UObject, public IMovieScenePlayer
+{
+	GENERATED_UCLASS_BODY()
+
+public:
+	void InitSequencePlayer( const FWidgetAnimation& Animation, UUserWidget& UserWidget );
+
+	/** Updates the running movie */
+	void Tick( float DeltaTime );
+
+	/** Begins playing or restarts an animation */
+	void Play();
+
+	/** Stops a running animation and resets time */
+	void Stop();
+
+	/** Pauses a running animation */
+	void Pause();
+
+	const UMovieScene* GetMovieScene() const { return MovieScene; }
+
+	/** IMovieScenePlayer interface */
+	virtual void GetRuntimeObjects( TSharedRef<FMovieSceneInstance> MovieSceneInstance, const FGuid& ObjectHandle, TArray< UObject* >& OutObjects ) const override;
+	virtual void UpdatePreviewViewports(UObject* ObjectToViewThrough) const override {}
+	virtual void AddMovieSceneInstance( class UMovieSceneSection& MovieSceneSection, TSharedRef<FMovieSceneInstance> InstanceToAdd ) override {}
+	virtual void RemoveMovieSceneInstance( class UMovieSceneSection& MovieSceneSection, TSharedRef<FMovieSceneInstance> InstanceToRemove ) override {}
+	virtual TSharedRef<FMovieSceneInstance> GetRootMovieSceneInstance() const override { return RootMovieSceneInstance.ToSharedRef(); }
+	virtual EMovieScenePlayerStatus::Type GetPlaybackStatus() const override;
+
+	DECLARE_EVENT_OneParam(UUMGSequencePlayer, FOnSequenceFinishedPlaying, UUMGSequencePlayer&);
+	FOnSequenceFinishedPlaying& OnSequenceFinishedPlaying() { return OnSequenceFinishedPlayingEvent; }
+
+private:
+	/** Sequence being played */
+	UPROPERTY()
+	UMovieScene* MovieScene;
+
+	/** Bindings to actual live objects */
+	UPROPERTY()
+	UMovieSceneBindings* RuntimeBindings;
+
+	/** The root movie scene instance to update when playing. */
+	TSharedPtr<class FMovieSceneInstance> RootMovieSceneInstance;
+
+	/** Time range of the animation */
+	TRange<float> TimeRange;
+
+	/** The current time cursor position within the sequence (in seconds) */
+	double TimeCursorPosition;
+
+	/** Status of the player (e.g play, stopped) */
+	EMovieScenePlayerStatus::Type PlayerStatus;
+
+	/** Delegate to call when a sequence has finished playing */
+	FOnSequenceFinishedPlaying OnSequenceFinishedPlayingEvent;
+};
diff --git a/Engine/Source/Runtime/UMG/Public/Blueprint/UserWidget.h b/Engine/Source/Runtime/UMG/Public/Blueprint/UserWidget.h
--- a/Engine/Source/Runtime/UMG/Public/Blueprint/UserWidget.h
+++ b/Engine/Source/Runtime/UMG/Public/Blueprint/UserWidget.h
@@ -84,6 +84,8 @@ public:
 	int32 MaxLayer;
 };
 
+class UUMGSequencePlayer;
+
 //TODO UMG If you want to host a widget that's full screen there may need to be a SWindow equivalent that you spawn it into.
 
 DECLARE_DYNAMIC_MULTICAST_DELEGATE(FOnConstructEvent);
@@ -194,6 +196,25 @@ public:
 
 	//virtual bool OnVisualizeTooltip(const TSharedPtr<SWidget>& TooltipContent);
 
+	/**
+ 	 * Plays an animation in this widget
+	 * 
+	 * @param The name of the animation to play
+	 */
+	UFUNCTION(BlueprintCallable, Category="User Interface|Animation")
+	void PlayAnimation(FName AnimationName);
+
+	/**
+ 	 * Stops an already running animation in this widget
+	 * 
+	 * @param The name of the animation to stop
+	 */
+	UFUNCTION(BlueprintCallable, Category="User Interface|Animation")
+	void StopAnimation(FName AnimationName);
+
+	/** Called when a sequence player is finished playing an animation */
+	void OnAnimationFinishedPlaying(UUMGSequencePlayer& Player );
+
 	/** @returns The UObject wrapper for a given SWidget */
 	UWidget* GetWidgetHandle(TSharedRef<SWidget> InWidget);
 
@@ -209,6 +230,8 @@ public:
 	/** @returns The uobject widget corresponding to a given name */
 	UWidget* GetHandleFromName(const FString& Name) const;
 
+	/** Ticks this widget and forwards to the underlying widget to tick */
+	void NativeTick(const FGeometry& MyGeometry, float InDeltaTime );
 public:
 	/** Called when the visibility changes. */
 	UPROPERTY(BlueprintAssignable)
@@ -277,6 +300,10 @@ public:
 	UPROPERTY()
 	class UWidgetTree* WidgetTree;
 
+	/** All the sequence players currently playing */
+	UPROPERTY(transient)
+	TArray<UUMGSequencePlayer*> ActiveSequencePlayers;
+
 protected:
 	virtual TSharedRef<SWidget> RebuildWidget() override;
 
diff --git a/Engine/Source/Runtime/UMG/Public/Blueprint/WidgetBlueprintGeneratedClass.h b/Engine/Source/Runtime/UMG/Public/Blueprint/WidgetBlueprintGeneratedClass.h
--- a/Engine/Source/Runtime/UMG/Public/Blueprint/WidgetBlueprintGeneratedClass.h
+++ b/Engine/Source/Runtime/UMG/Public/Blueprint/WidgetBlueprintGeneratedClass.h
@@ -4,6 +4,8 @@
 
 #include "WidgetBlueprintGeneratedClass.generated.h"
 
+class UMovieScene;
+
 USTRUCT()
 struct FDelegateRuntimeBinding
 {
@@ -25,7 +27,36 @@ struct FDelegateRuntimeBinding
 	}*/
 };
 
-class UMovieScene;
+/** A single object bound to a umg sequence */
+USTRUCT()
+struct FWidgetAnimationBinding
+{
+	GENERATED_USTRUCT_BODY()
+
+	UPROPERTY()
+	FName WidgetName;
+
+	UPROPERTY()
+	FGuid AnimationGuid;
+
+	bool operator==( const FWidgetAnimationBinding& Other ) const
+	{
+		return WidgetName == Other.WidgetName && AnimationGuid == Other.AnimationGuid;
+	}
+};
+
+/** Runtime animation bindings for a single movie scene */
+USTRUCT()
+struct FWidgetAnimation
+{
+	GENERATED_USTRUCT_BODY()
+
+	UPROPERTY()
+	class UMovieScene* MovieScene;
+
+	UPROPERTY()
+	TArray<FWidgetAnimationBinding> AnimationBindings;
+};
 
 /**
  * The widget blueprint generated class allows us to create blueprintable widgets for UMG at runtime.
@@ -46,7 +77,7 @@ class UMG_API UWidgetBlueprintGeneratedClass : public UBlueprintGeneratedClass
 	TArray< FDelegateRuntimeBinding > Bindings;
 
 	UPROPERTY()
-	TArray< UMovieScene* > AnimationData;
+	TArray<FWidgetAnimation> AnimationData;
 
 	/** This is transient data calculated at link time. */
 	TArray<UStructProperty*> WidgetNodeProperties;
@@ -61,4 +92,12 @@ class UMG_API UWidgetBlueprintGeneratedClass : public UBlueprintGeneratedClass
 	 * binding and wiring nessesary to have the user's widget perform as desired.
 	 */
 	void InitializeWidget(class UUserWidget* UserWidget) const;
+
+	/** 
+	 * Finds an animation by name of the Sequence
+	 * 
+	 * @param AnimationName	Name of the animation
+	 * @return Found animation or nullptr if one is not found
+	 */
+	const FWidgetAnimation* FindAnimation( FName AnimationName ) const;
 };
\ No newline at end of file
