UMG - Now allowing direct binding to properties in the designer.

diff --git a/Engine/Source/Editor/BlueprintGraph/Classes/K2Node_Variable.h b/Engine/Source/Editor/BlueprintGraph/Classes/K2Node_Variable.h
--- a/Engine/Source/Editor/BlueprintGraph/Classes/K2Node_Variable.h
+++ b/Engine/Source/Editor/BlueprintGraph/Classes/K2Node_Variable.h
@@ -127,7 +127,7 @@ public:
 	 * @param  IconColorOut	A color out, further discerning the variable's type.
 	 * @return A icon representing the specified variable's type.
 	 */
-	static FName GetVariableIconAndColor(UStruct* VarScope, FName VarName, FLinearColor& IconColorOut);
+	static FName GetVariableIconAndColor(const UStruct* VarScope, FName VarName, FLinearColor& IconColorOut);
 
 	/**
 	 * Utility method intended to serve as a choke point for various slate 
diff --git a/Engine/Source/Editor/BlueprintGraph/Private/K2Node_Variable.cpp b/Engine/Source/Editor/BlueprintGraph/Private/K2Node_Variable.cpp
--- a/Engine/Source/Editor/BlueprintGraph/Private/K2Node_Variable.cpp
+++ b/Engine/Source/Editor/BlueprintGraph/Private/K2Node_Variable.cpp
@@ -372,7 +372,7 @@ void UK2Node_Variable::GetNodeAttributes( TArray<TKeyValuePair<FString, FString>
 	OutNodeAttributes.Add( TKeyValuePair<FString, FString>( TEXT( "Name" ), VariableName ));
 }
 
-FName UK2Node_Variable::GetVariableIconAndColor(UStruct* VarScope, FName VarName, FLinearColor& IconColorOut)
+FName UK2Node_Variable::GetVariableIconAndColor(const UStruct* VarScope, FName VarName, FLinearColor& IconColorOut)
 {
 	FName IconBrush = TEXT("Kismet.AllClasses.VariableIcon");
 
diff --git a/Engine/Source/Editor/Kismet/Private/BlueprintEditor.cpp b/Engine/Source/Editor/Kismet/Private/BlueprintEditor.cpp
--- a/Engine/Source/Editor/Kismet/Private/BlueprintEditor.cpp
+++ b/Engine/Source/Editor/Kismet/Private/BlueprintEditor.cpp
@@ -253,7 +253,7 @@ const FSlateBrush* FBlueprintEditor::GetGlyphForGraph(const UEdGraph* Graph, boo
 	return ReturnValue;
 }
 
-FSlateBrush const* FBlueprintEditor::GetVarIconAndColor(UStruct* VarScope, FName VarName, FSlateColor& IconColorOut)
+FSlateBrush const* FBlueprintEditor::GetVarIconAndColor(const UStruct* VarScope, FName VarName, FSlateColor& IconColorOut)
 {
 	FLinearColor ColorOut;
 	const FSlateBrush* IconBrush = FEditorStyle::GetBrush(UK2Node_Variable::GetVariableIconAndColor(VarScope, VarName, ColorOut));
diff --git a/Engine/Source/Editor/Kismet/Public/BlueprintEditor.h b/Engine/Source/Editor/Kismet/Public/BlueprintEditor.h
--- a/Engine/Source/Editor/Kismet/Public/BlueprintEditor.h
+++ b/Engine/Source/Editor/Kismet/Public/BlueprintEditor.h
@@ -296,7 +296,7 @@ public:
 	 * @param IconColorOut		The resulting color for the glyph
 	 * @return					The resulting glyph brush
 	 */
-	static FSlateBrush const* GetVarIconAndColor(UStruct* VarScope, FName VarName, FSlateColor& IconColorOut);
+	static FSlateBrush const* GetVarIconAndColor(const UStruct* VarScope, FName VarName, FSlateColor& IconColorOut);
 
 	/** Overridable function for determining if the current mode can script */
 	virtual bool IsInAScriptingMode() const;
diff --git a/Engine/Source/Editor/UMGEditor/Private/Customizations/DetailCustomizations.cpp b/Engine/Source/Editor/UMGEditor/Private/Customizations/DetailCustomizations.cpp
--- a/Engine/Source/Editor/UMGEditor/Private/Customizations/DetailCustomizations.cpp
+++ b/Engine/Source/Editor/UMGEditor/Private/Customizations/DetailCustomizations.cpp
@@ -67,9 +67,26 @@ void FBlueprintWidgetCustomization::CustomizeDetails( IDetailLayoutBuilder& Deta
 				.ContentPadding(0)
 				.ButtonContent()
 				[
-					SNew(STextBlock)
-					.Text(this, &FBlueprintWidgetCustomization::GetCurrentBindingText, ConstPropertyHandle)
-					.Font(IDetailLayoutBuilder::GetDetailFont())
+					SNew(SHorizontalBox)
+
+					+ SHorizontalBox::Slot()
+					.AutoWidth()
+					.VAlign(VAlign_Center)
+					[
+						SNew(SImage)
+						.Image(this, &FBlueprintWidgetCustomization::GetCurrentBindingImage, ConstPropertyHandle)
+						.ColorAndOpacity(FLinearColor(0.25f, 0.25f, 0.25f))
+					]
+
+					+ SHorizontalBox::Slot()
+					.AutoWidth()
+					.VAlign(VAlign_Center)
+					.Padding(4,1,0,0)
+					[
+						SNew(STextBlock)
+						.Text(this, &FBlueprintWidgetCustomization::GetCurrentBindingText, ConstPropertyHandle)
+						.Font(IDetailLayoutBuilder::GetDetailFont())
+					]
 				];
 
 				const bool bShowChildren = true;
@@ -167,12 +184,13 @@ void FBlueprintWidgetCustomization::CustomizeDetails( IDetailLayoutBuilder& Deta
 	}
 }
 
-void FBlueprintWidgetCustomization::RefreshBlueprintFunctionCache(const UFunction* DelegateSignature)
+void FBlueprintWidgetCustomization::RefreshBlueprintMemberCache(const UFunction* DelegateSignature)
 {
 	const UWidgetGraphSchema* Schema = GetDefault<UWidgetGraphSchema>();
 	const FSlateFontInfo DetailFontInfo = IDetailLayoutBuilder::GetDetailFont();
 
 	BlueprintFunctionCache.Reset();
+	BlueprintPropertyCache.Reset();
 
 	// Get the current skeleton class, think header for the blueprint.
 	UBlueprintGeneratedClass* SkeletonClass = Cast<UBlueprintGeneratedClass>(Blueprint->SkeletonGeneratedClass);
@@ -212,27 +230,49 @@ void FBlueprintWidgetCustomization::RefreshBlueprintFunctionCache(const UFunctio
 			BlueprintFunctionCache.Add(NewFuncAction);
 		}
 	}
+
+	// Grab functions implemented by the blueprint
+	for ( TFieldIterator<UProperty> PropIt(SkeletonClass, EFieldIteratorFlags::ExcludeSuper); PropIt; ++PropIt )
+	{
+		UProperty* Prop = *PropIt;
+		
+		if ( DelegateSignature->GetReturnProperty()->SameType(Prop) )
+		{
+			if ( Prop->HasAnyPropertyFlags(UP::BlueprintReadWrite) )
+			{
+				BlueprintPropertyCache.Add(Prop);
+			}
+		}
+	}
 }
 
 TSharedRef<SWidget> FBlueprintWidgetCustomization::OnGenerateDelegateMenu(TSharedRef<IPropertyHandle> PropertyHandle, UFunction* DelegateSignature)
 {
-	RefreshBlueprintFunctionCache(DelegateSignature);
+	RefreshBlueprintMemberCache(DelegateSignature);
 
 	const bool bInShouldCloseWindowAfterMenuSelection = true;
 	FMenuBuilder MenuBuilder(bInShouldCloseWindowAfterMenuSelection, NULL);
 
-	//TODO UMG Enable or disable Remove Binding if needed.
+	static FName PropertyIcon(TEXT("Kismet.Tabs.Variables"));
+	static FName FunctionIcon(TEXT("GraphEditor.Function_16x"));
+
+	//FLinearColor ColorOut;
+	//const UClass* VarClass = PropertyHandle->GetPropertyClass();
+	//BrushOut = FBlueprintEditor::GetVarIconAndColor(VarClass, PropertyHandle->GetProperty()->GetFName(), ColorOut);
 
-	MenuBuilder.BeginSection("ClearBinding");
+	if ( CanRemoveBinding(PropertyHandle) )
 	{
-		MenuBuilder.AddMenuEntry(
-			LOCTEXT("RemoveBinding", "Remove Binding"),
-			LOCTEXT("RemoveBindingToolTip", "Removes the current binding"),
-			FSlateIcon(),
-			FUIAction(FExecuteAction::CreateSP(this, &FBlueprintWidgetCustomization::HandleRemoveBinding, PropertyHandle))
-			);
+		MenuBuilder.BeginSection("RemoveBinding");
+		{
+			MenuBuilder.AddMenuEntry(
+				LOCTEXT("RemoveBinding", "Remove Binding"),
+				LOCTEXT("RemoveBindingToolTip", "Removes the current binding"),
+				FSlateIcon(),
+				FUIAction(FExecuteAction::CreateSP(this, &FBlueprintWidgetCustomization::HandleRemoveBinding, PropertyHandle))
+				);
+		}
+		MenuBuilder.EndSection(); //RemoveBinding
 	}
-	MenuBuilder.EndSection(); //ClearBinding
 
 	MenuBuilder.BeginSection("CreateBinding");
 	{
@@ -245,15 +285,29 @@ TSharedRef<SWidget> FBlueprintWidgetCustomization::OnGenerateDelegateMenu(TShare
 	}
 	MenuBuilder.EndSection(); //CreateBinding
 
-	MenuBuilder.BeginSection("Functions");
+	MenuBuilder.BeginSection("Functions", LOCTEXT("Functions", "Functions"));
 	{
 		for ( TSharedPtr<FunctionInfo>& Function : BlueprintFunctionCache )
 		{
 			MenuBuilder.AddMenuEntry(
 				Function->DisplayName,
 				FText::FromString(Function->Tooltip),
-				FSlateIcon(),
-				FUIAction(FExecuteAction::CreateSP(this, &FBlueprintWidgetCustomization::HandleAddBinding, PropertyHandle, Function))
+				FSlateIcon(FEditorStyle::GetStyleSetName(), FunctionIcon),
+				FUIAction(FExecuteAction::CreateSP(this, &FBlueprintWidgetCustomization::HandleAddFunctionBinding, PropertyHandle, Function))
+				);
+		}
+	}
+	MenuBuilder.EndSection(); //Functions
+
+	MenuBuilder.BeginSection("Properties", LOCTEXT("Properties", "Properties"));
+	{
+		for ( UProperty* ExistingProperty : BlueprintPropertyCache )
+		{
+			MenuBuilder.AddMenuEntry(
+				ExistingProperty->GetDisplayNameText(),
+				ExistingProperty->GetToolTipText(),
+				FSlateIcon(FEditorStyle::GetStyleSetName(), PropertyIcon),
+				FUIAction(FExecuteAction::CreateSP(this, &FBlueprintWidgetCustomization::HandleAddPropertyBinding, PropertyHandle, ExistingProperty))
 				);
 		}
 	}
@@ -273,6 +327,46 @@ TSharedRef<SWidget> FBlueprintWidgetCustomization::OnGenerateDelegateMenu(TShare
 		];
 }
 
+const FSlateBrush* FBlueprintWidgetCustomization::GetCurrentBindingImage(TSharedRef<IPropertyHandle> PropertyHandle) const
+{
+	static FName PropertyIcon(TEXT("Kismet.Tabs.Variables"));
+	static FName FunctionIcon(TEXT("GraphEditor.Function_16x"));
+
+	TArray<UObject*> OuterObjects;
+	PropertyHandle->GetOuterObjects(OuterObjects);
+
+	//TODO UMG O(N) Isn't good for this, needs to be map, but map isn't serialized, need cached runtime map for fast lookups.
+
+	FName PropertyName = PropertyHandle->GetProperty()->GetFName();
+	for ( int32 ObjectIndex = 0; ObjectIndex < OuterObjects.Num(); ObjectIndex++ )
+	{
+		// Ignore null outer objects
+		if ( OuterObjects[ObjectIndex] == NULL )
+		{
+			continue;
+		}
+
+		//TODO UMG handle multiple things selected
+
+		for ( const FDelegateEditorBinding& Binding : Blueprint->Bindings )
+		{
+			if ( Binding.ObjectName == OuterObjects[ObjectIndex]->GetName() && Binding.PropertyName == PropertyName )
+			{
+				if ( Binding.Kind == EBindingKind::Function )
+				{
+					return FEditorStyle::GetBrush(FunctionIcon);
+				}
+				else // Property
+				{
+					return FEditorStyle::GetBrush(PropertyIcon);
+				}
+			}
+		}
+	}
+
+	return NULL;
+}
+
 FText FBlueprintWidgetCustomization::GetCurrentBindingText(TSharedRef<IPropertyHandle> PropertyHandle) const
 {
 	TArray<UObject*> OuterObjects;
@@ -283,7 +377,7 @@ FText FBlueprintWidgetCustomization::GetCurrentBindingText(TSharedRef<IPropertyH
 	FName PropertyName = PropertyHandle->GetProperty()->GetFName();
 	for ( int32 ObjectIndex = 0; ObjectIndex < OuterObjects.Num(); ObjectIndex++ )
 	{
-		// Ignore null outter objects
+		// Ignore null outer objects
 		if ( OuterObjects[ObjectIndex] == NULL )
 		{
 			continue;
@@ -295,9 +389,16 @@ FText FBlueprintWidgetCustomization::GetCurrentBindingText(TSharedRef<IPropertyH
 		{
 			if ( Binding.ObjectName == OuterObjects[ObjectIndex]->GetName() && Binding.PropertyName == PropertyName )
 			{
-				FName FoundName = Blueprint->GetFieldNameFromClassByGuid<UFunction>(Blueprint->GeneratedClass, Binding.MemberGuid);
-				return FText::FromName(FoundName);
-				//return FText::FromName(Binding.FunctionName);
+				if ( Binding.Kind == EBindingKind::Function )
+				{
+					FName FoundName = Blueprint->GetFieldNameFromClassByGuid<UFunction>(Blueprint->GeneratedClass, Binding.MemberGuid);
+					return FText::FromString(FName::NameToDisplayString(FoundName.ToString(), false));
+				}
+				else // Property
+				{
+					FName FoundName = Blueprint->GetFieldNameFromClassByGuid<UProperty>(Blueprint->GeneratedClass, Binding.MemberGuid);
+					return FText::FromString( FName::NameToDisplayString(FoundName.ToString(), false) );
+				}
 			}
 		}
 
@@ -309,6 +410,26 @@ FText FBlueprintWidgetCustomization::GetCurrentBindingText(TSharedRef<IPropertyH
 	return LOCTEXT("Bind", "Bind");
 }
 
+bool FBlueprintWidgetCustomization::CanRemoveBinding(TSharedRef<IPropertyHandle> PropertyHandle)
+{
+	FName PropertyName = PropertyHandle->GetProperty()->GetFName();
+
+	TArray<UObject*> OuterObjects;
+	PropertyHandle->GetOuterObjects(OuterObjects);
+	for ( UObject* SelectedObject : OuterObjects )
+	{
+		for ( const FDelegateEditorBinding& Binding : Blueprint->Bindings )
+		{
+			if ( Binding.ObjectName == SelectedObject->GetName() && Binding.PropertyName == PropertyName )
+			{
+				return true;
+			}
+		}
+	}
+
+	return false;
+}
+
 void FBlueprintWidgetCustomization::HandleRemoveBinding(TSharedRef<IPropertyHandle> PropertyHandle)
 {
 	const FScopedTransaction Transaction(LOCTEXT("BindDelegate", "Remove Binding"));
@@ -329,7 +450,7 @@ void FBlueprintWidgetCustomization::HandleRemoveBinding(TSharedRef<IPropertyHand
 	FBlueprintEditorUtils::MarkBlueprintAsStructurallyModified(Blueprint);
 }
 
-void FBlueprintWidgetCustomization::HandleAddBinding(TSharedRef<IPropertyHandle> PropertyHandle, TSharedPtr<FunctionInfo> SelectedFunction)
+void FBlueprintWidgetCustomization::HandleAddFunctionBinding(TSharedRef<IPropertyHandle> PropertyHandle, TSharedPtr<FunctionInfo> SelectedFunction)
 {
 	const FScopedTransaction Transaction(LOCTEXT("BindDelegate", "Set Binding"));
 
@@ -344,6 +465,37 @@ void FBlueprintWidgetCustomization::HandleAddBinding(TSharedRef<IPropertyHandle>
 		Binding.PropertyName = PropertyHandle->GetProperty()->GetFName();
 		Binding.FunctionName = SelectedFunction->FuncName;
 		Binding.MemberGuid = SelectedFunction->EdGraph->GraphGuid;
+		Binding.Kind = EBindingKind::Function;
+
+		Blueprint->Bindings.Remove(Binding);
+		Blueprint->Bindings.AddUnique(Binding);
+	}
+
+	FBlueprintEditorUtils::MarkBlueprintAsStructurallyModified(Blueprint);
+}
+
+void FBlueprintWidgetCustomization::HandleAddPropertyBinding(TSharedRef<IPropertyHandle> PropertyHandle, UProperty* SelectedProperty)
+{
+	const FScopedTransaction Transaction(LOCTEXT("BindDelegate", "Set Binding"));
+
+	// Get the current skeleton class, think header for the blueprint.
+	UBlueprintGeneratedClass* SkeletonClass = Cast<UBlueprintGeneratedClass>(Blueprint->SkeletonGeneratedClass);
+
+	Blueprint->Modify();
+
+	FGuid MemberGuid;
+	UBlueprint::GetGuidFromClassByFieldName<UProperty>(SkeletonClass, SelectedProperty->GetFName(), MemberGuid);
+
+	TArray<UObject*> OuterObjects;
+	PropertyHandle->GetOuterObjects(OuterObjects);
+	for ( UObject* SelectedObject : OuterObjects )
+	{
+		FDelegateEditorBinding Binding;
+		Binding.ObjectName = SelectedObject->GetName();
+		Binding.PropertyName = PropertyHandle->GetProperty()->GetFName();
+		Binding.SourceProperty = SelectedProperty->GetFName();
+		Binding.MemberGuid = MemberGuid;
+		Binding.Kind = EBindingKind::Property;
 
 		Blueprint->Bindings.Remove(Binding);
 		Blueprint->Bindings.AddUnique(Binding);
@@ -376,7 +528,7 @@ void FBlueprintWidgetCustomization::HandleCreateAndAddBinding(TSharedRef<IProper
 	SelectedFunction->FuncName = FunctionGraph->GetFName();
 	SelectedFunction->EdGraph = FunctionGraph;
 
-	HandleAddBinding(PropertyHandle, SelectedFunction);
+	HandleAddFunctionBinding(PropertyHandle, SelectedFunction);
 
 	GotoFunction(FunctionGraph);
 }
diff --git a/Engine/Source/Editor/UMGEditor/Private/Customizations/DetailCustomizations.h b/Engine/Source/Editor/UMGEditor/Private/Customizations/DetailCustomizations.h
--- a/Engine/Source/Editor/UMGEditor/Private/Customizations/DetailCustomizations.h
+++ b/Engine/Source/Editor/UMGEditor/Private/Customizations/DetailCustomizations.h
@@ -30,14 +30,19 @@ private:
 		UEdGraph* EdGraph;
 	};
 
-	void RefreshBlueprintFunctionCache(const UFunction* DelegateSignature);
+	void RefreshBlueprintMemberCache(const UFunction* DelegateSignature);
 
 	TSharedRef<SWidget> OnGenerateDelegateMenu(TSharedRef<IPropertyHandle> PropertyHandle, UFunction* DelegateSignature);
 
+	const FSlateBrush* GetCurrentBindingImage(TSharedRef<IPropertyHandle> PropertyHandle) const;
 	FText GetCurrentBindingText(TSharedRef<IPropertyHandle> PropertyHandle) const;
 
+	bool CanRemoveBinding(TSharedRef<IPropertyHandle> PropertyHandle);
 	void HandleRemoveBinding(TSharedRef<IPropertyHandle> PropertyHandle);
-	void HandleAddBinding(TSharedRef<IPropertyHandle> PropertyHandle, TSharedPtr<FunctionInfo> SelectedFunction);
+
+	void HandleAddFunctionBinding(TSharedRef<IPropertyHandle> PropertyHandle, TSharedPtr<FunctionInfo> SelectedFunction);
+	void HandleAddPropertyBinding(TSharedRef<IPropertyHandle> PropertyHandle, UProperty* SelectedProperty);
+
 	void HandleCreateAndAddBinding(TSharedRef<IPropertyHandle> PropertyHandle, UFunction* DelegateSignature);
 	void GotoFunction(UEdGraph* FunctionGraph);
 
@@ -48,4 +53,5 @@ private:
 	UWidgetBlueprint* Blueprint;
 
 	TArray< TSharedPtr<FunctionInfo> > BlueprintFunctionCache;
+	TArray< UProperty* > BlueprintPropertyCache;
 };
diff --git a/Engine/Source/Editor/UMGEditor/Private/WidgetBlueprint.cpp b/Engine/Source/Editor/UMGEditor/Private/WidgetBlueprint.cpp
--- a/Engine/Source/Editor/UMGEditor/Private/WidgetBlueprint.cpp
+++ b/Engine/Source/Editor/UMGEditor/Private/WidgetBlueprint.cpp
@@ -10,7 +10,15 @@ FDelegateRuntimeBinding FDelegateEditorBinding::ToRuntimeBinding(UWidgetBlueprin
 	FDelegateRuntimeBinding Binding;
 	Binding.ObjectName = ObjectName;
 	Binding.PropertyName = PropertyName;
-	Binding.FunctionName = Blueprint->GetFieldNameFromClassByGuid<UFunction>(Blueprint->SkeletonGeneratedClass, MemberGuid);
+	if ( Kind == EBindingKind::Function )
+	{
+		Binding.FunctionName = Blueprint->GetFieldNameFromClassByGuid<UFunction>(Blueprint->SkeletonGeneratedClass, MemberGuid);
+	}
+	else
+	{
+		Binding.FunctionName = FunctionName;// Blueprint->GetFieldNameFromClassByGuid<UProperty>(Blueprint->SkeletonGeneratedClass, MemberGuid);
+	}
+	Binding.Kind = Kind;
 
 	return Binding;
 }
diff --git a/Engine/Source/Editor/UMGEditor/Private/WidgetBlueprintCompiler.cpp b/Engine/Source/Editor/UMGEditor/Private/WidgetBlueprintCompiler.cpp
--- a/Engine/Source/Editor/UMGEditor/Private/WidgetBlueprintCompiler.cpp
+++ b/Engine/Source/Editor/UMGEditor/Private/WidgetBlueprintCompiler.cpp
@@ -8,6 +8,9 @@
 #include "KismetReinstanceUtilities.h"
 #include "MovieScene.h"
 
+#include "BlueprintPropertyNodeSpawner.h"
+#include "K2Node_FunctionEntry.h"
+
 template<>
 FString FNetNameMapping::MakeBaseName<UWidgetGraphNode_Base>(const UWidgetGraphNode_Base* Net)
 {
@@ -68,7 +71,7 @@ void FWidgetBlueprintCompiler::MergeUbergraphPagesIn(UEdGraph* Ubergraph)
 	// Build the raw node list
 	TArray<UWidgetGraphNode_Base*> WidgetNodeList;
 	ConsolidatedEventGraph->GetNodesOfClass<UWidgetGraphNode_Base>(/*out*/ WidgetNodeList);
-	
+
 	AllocateNodeIndexCounter = 0;
 
 	// Process the remaining nodes
@@ -78,6 +81,79 @@ void FWidgetBlueprintCompiler::MergeUbergraphPagesIn(UEdGraph* Ubergraph)
 	}
 }
 
+void FWidgetBlueprintCompiler::CreateFunctionList()
+{
+	Super::CreateFunctionList();
+
+	for ( FDelegateEditorBinding& EditorBinding : WidgetBlueprint()->Bindings )
+	{
+		FName PropertyName = EditorBinding.SourceProperty;
+
+		UProperty* Property = FindField<UProperty>(Blueprint->SkeletonGeneratedClass, PropertyName);
+		if ( Property )
+		{
+			// Create the function graph.
+			FString FunctionName = FString(TEXT("__Get")) + PropertyName.ToString();
+			UEdGraph* FunctionGraph = FBlueprintEditorUtils::CreateNewGraph(Blueprint, FBlueprintEditorUtils::FindUniqueKismetName(Blueprint, FunctionName), UEdGraph::StaticClass(), UEdGraphSchema_K2::StaticClass());
+
+			// Update the function binding to match the generated graph name
+			EditorBinding.FunctionName = FunctionGraph->GetFName();
+
+			const UEdGraphSchema_K2* K2Schema = Cast<const UEdGraphSchema_K2>(FunctionGraph->GetSchema());
+
+			Schema->CreateDefaultNodesForGraph(*FunctionGraph);
+
+			// We need to flag the entry node to make sure that the compiled function is callable from Kismet2
+			int32 ExtraFunctionFlags = ( FUNC_BlueprintCallable | FUNC_BlueprintEvent | FUNC_Public );
+			if ( BPTYPE_FunctionLibrary == Blueprint->BlueprintType )
+			{
+				ExtraFunctionFlags |= FUNC_Static;
+			}
+			K2Schema->AddExtraFunctionFlags(FunctionGraph, ExtraFunctionFlags);
+			K2Schema->MarkFunctionEntryAsEditable(FunctionGraph, true);
+
+			// Create a function entry node
+			FGraphNodeCreator<UK2Node_FunctionEntry> FunctionEntryCreator(*FunctionGraph);
+			UK2Node_FunctionEntry* EntryNode = FunctionEntryCreator.CreateNode();
+			EntryNode->SignatureClass = NULL;
+			EntryNode->SignatureName = FunctionGraph->GetFName();
+			FunctionEntryCreator.Finalize();
+
+			FGraphNodeCreator<UK2Node_FunctionResult> FunctionReturnCreator(*FunctionGraph);
+			UK2Node_FunctionResult* ReturnNode = FunctionReturnCreator.CreateNode();
+			ReturnNode->SignatureClass = NULL;
+			ReturnNode->SignatureName = FunctionGraph->GetFName();
+			ReturnNode->NodePosX = EntryNode->NodePosX + EntryNode->NodeWidth + 256;
+			ReturnNode->NodePosY = EntryNode->NodePosY;
+			FunctionReturnCreator.Finalize();
+
+			UProperty* Property = FindField<UProperty>(Blueprint->SkeletonGeneratedClass, PropertyName);
+
+			FEdGraphPinType PinType;
+			K2Schema->ConvertPropertyToPinType(Property, /*out*/ PinType);
+
+			UEdGraphPin* ReturnPin = ReturnNode->CreateUserDefinedPin(TEXT("ReturnValue"), PinType);
+
+			// Auto-connect the pins for entry and exit, so that by default the signature is properly generated
+			UEdGraphPin* EntryNodeExec = K2Schema->FindExecutionPin(*EntryNode, EGPD_Output);
+			UEdGraphPin* ResultNodeExec = K2Schema->FindExecutionPin(*ReturnNode, EGPD_Input);
+			EntryNodeExec->MakeLinkTo(ResultNodeExec);
+
+			FGraphNodeCreator<UK2Node_VariableGet> MemberGetCreator(*FunctionGraph);
+			UK2Node_VariableGet* VarNode = MemberGetCreator.CreateNode();
+			VarNode->VariableReference.SetSelfMember(PropertyName);
+			MemberGetCreator.Finalize();
+
+			ReturnPin->MakeLinkTo(VarNode->GetValuePin());
+
+			//Blueprint->FunctionGraphs.Add(FunctionGraph);
+
+			ProcessOneFunctionGraph(FunctionGraph);
+			//FEdGraphUtilities::MergeChildrenGraphsIn(Ubergraph, FunctionGraph, /*bRequireSchemaMatch=*/ true);
+		}
+	}
+}
+
 void FWidgetBlueprintCompiler::ProcessWidgetNode(UWidgetGraphNode_Base* VisualWidgetNode)
 {
 	const UWidgetGraphSchema* Schema = GetDefault<UWidgetGraphSchema>();
diff --git a/Engine/Source/Editor/UMGEditor/Private/WidgetBlueprintCompiler.h b/Engine/Source/Editor/UMGEditor/Private/WidgetBlueprintCompiler.h
--- a/Engine/Source/Editor/UMGEditor/Private/WidgetBlueprintCompiler.h
+++ b/Engine/Source/Editor/UMGEditor/Private/WidgetBlueprintCompiler.h
@@ -28,6 +28,7 @@ protected:
 	// FKismetCompilerContext
 	//virtual UEdGraphSchema_K2* CreateSchema() override;
 	virtual void MergeUbergraphPagesIn(UEdGraph* Ubergraph) override;
+	virtual void CreateFunctionList() override;
 	virtual void SpawnNewClass(const FString& NewClassName) override;
 	virtual void CleanAndSanitizeClass(UBlueprintGeneratedClass* ClassToClean, UObject*& OldCDO) override;
 	virtual void SaveSubObjectsFromCleanAndSanitizeClass(FSubobjectCollection& SubObjectsToSave, UBlueprintGeneratedClass* ClassToClean, UObject*& OldCDO) override;
diff --git a/Engine/Source/Editor/UMGEditor/Public/WidgetBlueprint.h b/Engine/Source/Editor/UMGEditor/Public/WidgetBlueprint.h
--- a/Engine/Source/Editor/UMGEditor/Public/WidgetBlueprint.h
+++ b/Engine/Source/Editor/UMGEditor/Public/WidgetBlueprint.h
@@ -24,9 +24,15 @@ struct UMGEDITOR_API FDelegateEditorBinding
 	UPROPERTY()
 	FName FunctionName;
 
+	UPROPERTY()
+	FName SourceProperty;
+
 	UPROPERTY()
 	FGuid MemberGuid;
 
+	UPROPERTY()
+	TEnumAsByte<EBindingKind::Type> Kind;
+
 	bool operator==( const FDelegateEditorBinding& Other ) const
 	{
 		// NOTE: We intentionally only compare object name and property name, the function is irrelevant since
diff --git a/Engine/Source/Editor/UnrealEd/Public/Kismet2/BlueprintEditorUtils.h b/Engine/Source/Editor/UnrealEd/Public/Kismet2/BlueprintEditorUtils.h
--- a/Engine/Source/Editor/UnrealEd/Public/Kismet2/BlueprintEditorUtils.h
+++ b/Engine/Source/Editor/UnrealEd/Public/Kismet2/BlueprintEditorUtils.h
@@ -160,12 +160,13 @@ public:
 	 */
 	static class UEdGraph* CreateNewGraph(UObject* ParentScope, const FName& GraphName, TSubclassOf<class UEdGraph> GraphClass, TSubclassOf<class UEdGraphSchema> SchemaClass);
 
-	/** 
-	 * Adds a function graph to this blueprint.  If bIsUserCreated is true, the entry/exit nodes will be editable. SignatureFromObject is used to find signature for entry/exit nodes if using an existing signature.
+	/**
+	 * Creates a function graph, but does not add it to the blueprint.  If bIsUserCreated is true, the entry/exit nodes will be editable. 
+	 * SignatureFromObject is used to find signature for entry/exit nodes if using an existing signature.
 	 * The template argument SignatureType should be UClass or UFunction.
 	 */
 	template <typename SignatureType>
-	static void AddFunctionGraph(UBlueprint* Blueprint, class UEdGraph* Graph, bool bIsUserCreated, SignatureType* SignatureFromObject)
+	static void CreateFunctionGraph(UBlueprint* Blueprint, class UEdGraph* Graph, bool bIsUserCreated, SignatureType* SignatureFromObject)
 	{
 		// Give the schema a chance to fill out any required nodes (like the entry node or results node)
 		const UEdGraphSchema* Schema = Graph->GetSchema();
@@ -189,6 +190,17 @@ public:
 				K2Schema->MarkFunctionEntryAsEditable(Graph, true);
 			}
 		}
+	}
+
+	/** 
+	 * Adds a function graph to this blueprint.  If bIsUserCreated is true, the entry/exit nodes will be editable. 
+	 * SignatureFromObject is used to find signature for entry/exit nodes if using an existing signature.
+	 * The template argument SignatureType should be UClass or UFunction.
+	 */
+	template <typename SignatureType>
+	static void AddFunctionGraph(UBlueprint* Blueprint, class UEdGraph* Graph, bool bIsUserCreated, SignatureType* SignatureFromObject)
+	{
+		CreateFunctionGraph(Blueprint, Graph, bIsUserCreated, SignatureFromObject);
 
 		Blueprint->FunctionGraphs.Add(Graph);
 
diff --git a/Engine/Source/Runtime/UMG/Private/WidgetBlueprintGeneratedClass.cpp b/Engine/Source/Runtime/UMG/Private/WidgetBlueprintGeneratedClass.cpp
--- a/Engine/Source/Runtime/UMG/Private/WidgetBlueprintGeneratedClass.cpp
+++ b/Engine/Source/Runtime/UMG/Private/WidgetBlueprintGeneratedClass.cpp
@@ -91,8 +91,6 @@ void UWidgetBlueprintGeneratedClass::InitializeWidget(UUserWidget* UserWidget) c
 				//TODO UMG Terrible performance, improve with Maps.
 				if ( Binding.ObjectName == VariableName )
 				{
-					UFunction* BoundFunction = UserWidget->FindFunction(Binding.FunctionName);
-
 					FString DelegateName = Binding.PropertyName.ToString() + "Delegate";
 
 					for ( TFieldIterator<UProperty> It(Widget->GetClass()); It; ++It )
@@ -102,7 +100,31 @@ void UWidgetBlueprintGeneratedClass::InitializeWidget(UUserWidget* UserWidget) c
 							if ( DelegateProp->GetName() == DelegateName || DelegateProp->GetFName() == Binding.PropertyName )
 							{
 								FScriptDelegate* ScriptDelegate = DelegateProp->GetPropertyValuePtr_InContainer(Widget);
-								ScriptDelegate->BindUFunction(UserWidget, Binding.FunctionName);
+
+								if ( Binding.Kind == EBindingKind::Function )
+								{
+									ScriptDelegate->BindUFunction(UserWidget, Binding.FunctionName);
+								}
+								else if ( Binding.Kind == EBindingKind::Property )
+								{
+									//FString FunctionNameStr = FString(TEXT("__Get")) + Binding.FunctionName.ToString();// +TEXT("_0");
+									//FName FunctionName = FName(*FunctionNameStr, 1);
+									UFunction* Fun = UserWidget->FindFunction(Binding.FunctionName);
+
+									TArray<FName> names;
+									for ( TFieldIterator<UFunction> It(WidgetBlueprintClass); It; ++It )
+									{
+										FName FunName = It->GetFName();
+										names.Add(FunName);
+									}
+
+									ScriptDelegate->BindUFunction(UserWidget, Binding.FunctionName);
+								}
+								else
+								{
+									check(false);
+								}
+								
 								break;
 							}
 						}
diff --git a/Engine/Source/Runtime/UMG/Public/Blueprint/WidgetBlueprintGeneratedClass.h b/Engine/Source/Runtime/UMG/Public/Blueprint/WidgetBlueprintGeneratedClass.h
--- a/Engine/Source/Runtime/UMG/Public/Blueprint/WidgetBlueprintGeneratedClass.h
+++ b/Engine/Source/Runtime/UMG/Public/Blueprint/WidgetBlueprintGeneratedClass.h
@@ -6,6 +6,16 @@
 
 class UMovieScene;
 
+UENUM()
+namespace EBindingKind
+{
+	enum Type
+	{
+		Function,
+		Property,
+	};
+}
+
 USTRUCT()
 struct FDelegateRuntimeBinding
 {
@@ -19,6 +29,9 @@ struct FDelegateRuntimeBinding
 
 	UPROPERTY()
 	FName FunctionName;
+
+	UPROPERTY()
+	TEnumAsByte<EBindingKind::Type> Kind;
 	
 	/*UFunction* UK2Node_CallFunction::GetTargetFunction() const
 	{
