non unity fix

diff --git a/Engine/Source/Editor/UMGEditor/Private/UMGSequencerObjectBindingManager.cpp b/Engine/Source/Editor/UMGEditor/Private/UMGSequencerObjectBindingManager.cpp
--- a/Engine/Source/Editor/UMGEditor/Private/UMGSequencerObjectBindingManager.cpp
+++ b/Engine/Source/Editor/UMGEditor/Private/UMGSequencerObjectBindingManager.cpp
@@ -4,6 +4,7 @@
 #include "UMGSequencerObjectBindingManager.h"
 #include "WidgetBlueprintEditor.h"
 #include "ISequencer.h"
+#include "MovieScene.h"
 
 FUMGSequencerObjectBindingManager::FUMGSequencerObjectBindingManager( FWidgetBlueprintEditor& InWidgetBlueprintEditor, UMovieScene& InMovieScene )
 	: WidgetBlueprintEditor( InWidgetBlueprintEditor )
