UMG - There is now a goto binding button on the binding customization.

diff --git a/Engine/Source/Editor/UMGEditor/Private/Customizations/DetailCustomizations.cpp b/Engine/Source/Editor/UMGEditor/Private/Customizations/DetailCustomizations.cpp
--- a/Engine/Source/Editor/UMGEditor/Private/Customizations/DetailCustomizations.cpp
+++ b/Engine/Source/Editor/UMGEditor/Private/Customizations/DetailCustomizations.cpp
@@ -115,6 +115,21 @@ void FBlueprintWidgetCustomization::CustomizeDetails( IDetailLayoutBuilder& Deta
 					[
 						DelegateButton
 					]
+
+					+ SHorizontalBox::Slot()
+					.AutoWidth()
+					[
+						SNew(SButton)
+						.ButtonStyle( FEditorStyle::Get(), "HoverHintOnly" )
+						.Visibility( this, &FBlueprintWidgetCustomization::GetGotoBindingVisibility, ConstPropertyHandle )
+						.OnClicked( this, &FBlueprintWidgetCustomization::HandleGotoBindingClicked, ConstPropertyHandle )
+						.VAlign(VAlign_Center)
+						.ToolTipText(LOCTEXT("GotoFunction", "Goto Function"))
+						[
+							SNew(SImage)
+							.Image( FEditorStyle::GetBrush("PropertyWindow.Button_Browse") )
+						]
+					]
 				];
 			}
 		}
@@ -177,6 +192,21 @@ void FBlueprintWidgetCustomization::CustomizeDetails( IDetailLayoutBuilder& Deta
 					[
 						DelegateButton
 					]
+
+					+ SHorizontalBox::Slot()
+					.AutoWidth()
+					[
+						SNew(SButton)
+						.ButtonStyle( FEditorStyle::Get(), "HoverHintOnly" )
+						.Visibility( this, &FBlueprintWidgetCustomization::GetGotoBindingVisibility, DelegatePropertyHandle )
+						.OnClicked( this, &FBlueprintWidgetCustomization::HandleGotoBindingClicked, DelegatePropertyHandle )
+						.VAlign(VAlign_Center)
+						.ToolTipText(LOCTEXT("GotoFunction", "Goto Function"))
+						[
+							SNew(SImage)
+							.Image( FEditorStyle::GetBrush("PropertyWindow.Button_Browse") )
+						]
+					]
 				];
 			}
 		}
@@ -533,6 +563,85 @@ void FBlueprintWidgetCustomization::HandleCreateAndAddBinding(TSharedRef<IProper
 	GotoFunction(FunctionGraph);
 }
 
+EVisibility FBlueprintWidgetCustomization::GetGotoBindingVisibility(TSharedRef<IPropertyHandle> PropertyHandle) const
+{
+	TArray<UObject*> OuterObjects;
+	PropertyHandle->GetOuterObjects(OuterObjects);
+
+	//TODO UMG O(N) Isn't good for this, needs to be map, but map isn't serialized, need cached runtime map for fast lookups.
+
+	FName PropertyName = PropertyHandle->GetProperty()->GetFName();
+	for ( int32 ObjectIndex = 0; ObjectIndex < OuterObjects.Num(); ObjectIndex++ )
+	{
+		// Ignore null outer objects
+		if ( OuterObjects[ObjectIndex] == NULL )
+		{
+			continue;
+		}
+
+		//TODO UMG handle multiple things selected
+
+		for ( const FDelegateEditorBinding& Binding : Blueprint->Bindings )
+		{
+			if ( Binding.ObjectName == OuterObjects[ObjectIndex]->GetName() && Binding.PropertyName == PropertyName )
+			{
+				if ( Binding.Kind == EBindingKind::Function )
+				{
+					return EVisibility::Visible;
+				}
+			}
+		}
+	}
+
+	return EVisibility::Collapsed;
+}
+
+FReply FBlueprintWidgetCustomization::HandleGotoBindingClicked(TSharedRef<IPropertyHandle> PropertyHandle)
+{
+	TArray<UObject*> OuterObjects;
+	PropertyHandle->GetOuterObjects(OuterObjects);
+
+	//TODO UMG O(N) Isn't good for this, needs to be map, but map isn't serialized, need cached runtime map for fast lookups.
+
+	FName PropertyName = PropertyHandle->GetProperty()->GetFName();
+	for ( int32 ObjectIndex = 0; ObjectIndex < OuterObjects.Num(); ObjectIndex++ )
+	{
+		// Ignore null outer objects
+		if ( OuterObjects[ObjectIndex] == NULL )
+		{
+			continue;
+		}
+
+		//TODO UMG handle multiple things selected
+
+		for ( const FDelegateEditorBinding& Binding : Blueprint->Bindings )
+		{
+			if ( Binding.ObjectName == OuterObjects[ObjectIndex]->GetName() && Binding.PropertyName == PropertyName )
+			{
+				if ( Binding.Kind == EBindingKind::Function )
+				{
+					TArray<UEdGraph*> AllGraphs;
+					Blueprint->GetAllGraphs(AllGraphs);
+
+					for ( UEdGraph* Graph : AllGraphs )
+					{
+						if ( Graph->GraphGuid == Binding.MemberGuid )
+						{
+							GotoFunction(Graph);
+						}
+					}
+
+
+					// Either way return
+					return FReply::Handled();
+				}
+			}
+		}
+	}
+
+	return FReply::Unhandled();
+}
+
 void FBlueprintWidgetCustomization::GotoFunction(UEdGraph* FunctionGraph)
 {
 	Editor.Pin()->SetCurrentMode(FWidgetBlueprintApplicationModes::GraphMode);
diff --git a/Engine/Source/Editor/UMGEditor/Private/Customizations/DetailCustomizations.h b/Engine/Source/Editor/UMGEditor/Private/Customizations/DetailCustomizations.h
--- a/Engine/Source/Editor/UMGEditor/Private/Customizations/DetailCustomizations.h
+++ b/Engine/Source/Editor/UMGEditor/Private/Customizations/DetailCustomizations.h
@@ -46,6 +46,10 @@ private:
 	void HandleCreateAndAddBinding(TSharedRef<IPropertyHandle> PropertyHandle, UFunction* DelegateSignature);
 	void GotoFunction(UEdGraph* FunctionGraph);
 
+	EVisibility GetGotoBindingVisibility(TSharedRef<IPropertyHandle> PropertyHandle) const;
+
+	FReply HandleGotoBindingClicked(TSharedRef<IPropertyHandle> PropertyHandle);
+
 private:
 
 	TWeakPtr<FWidgetBlueprintEditor> Editor;
