UMG - Fixing a copy paste issue.  Newly copied objects are now given a new name, instead of the old name, and the old object getting renamed.

diff --git a/Engine/Source/Editor/UMGEditor/Private/WidgetBlueprintEditorUtils.cpp b/Engine/Source/Editor/UMGEditor/Private/WidgetBlueprintEditorUtils.cpp
--- a/Engine/Source/Editor/UMGEditor/Private/WidgetBlueprintEditorUtils.cpp
+++ b/Engine/Source/Editor/UMGEditor/Private/WidgetBlueprintEditorUtils.cpp
@@ -350,8 +350,6 @@ void FWidgetBlueprintEditorUtils::PasteWidgets(UWidgetBlueprint* BP, FWidgetRefe
 		{
 			RootPasteWidgets.Add(NewWidget);
 		}
-
-		NewWidget->SetFlags(RF_Transactional);
 	}
 
 	// If there isn't a root widget and we're copying multiple root widgets, then we need to add a container root
@@ -409,14 +407,27 @@ void FWidgetBlueprintEditorUtils::PasteWidgets(UWidgetBlueprint* BP, FWidgetRefe
 
 void FWidgetBlueprintEditorUtils::ImportWidgetsFromText(UWidgetBlueprint* BP, const FString& TextToImport, /*out*/ TSet<UWidget*>& ImportedWidgetSet)
 {
+	// We create our own transient package here so that we can deserialize the data in isolation and ensure unreferenced
+	// objects not part of the deserialization set are unresolved.
+	UPackage* TempPackage = new( NULL, TEXT("/Engine/UMG/Editor/Transient") )UPackage(FPostConstructInitializeProperties());
+	TempPackage->AddToRoot();
+
 	// Turn the text buffer into objects
 	FWidgetObjectTextFactory Factory;
-	Factory.ProcessBuffer(BP->WidgetTree, RF_Transactional, TextToImport);
+	Factory.ProcessBuffer(TempPackage, RF_Transactional, TextToImport);
 
 	for ( auto& Entry : Factory.NewObjectMap )
 	{
-		ImportedWidgetSet.Add(Entry.Value);
+		UWidget* Widget = Entry.Value;
+
+		ImportedWidgetSet.Add(Widget);
+
+		Widget->SetFlags(RF_Transactional);
+		Widget->Rename(nullptr, BP->WidgetTree);
 	}
+
+	// Remove the temp package from the root now that it has served its purpose.
+	TempPackage->RemoveFromRoot();
 }
 
 void FWidgetBlueprintEditorUtils::ExportPropertiesToText(UObject* Object, TMap<FName, FString>& ExportedProperties)
