UMG - Can now zoomy / pan the designer.  The autosaver no longer saves packages with the RF_Transient flag.  Now using a transient flag to temporarily deserialize objects when pasting to prevent the renaming of existing widgets.

diff --git a/Engine/Source/Editor/UMGEditor/Private/Designer/SDesignSurface.cpp b/Engine/Source/Editor/UMGEditor/Private/Designer/SDesignSurface.cpp
--- a/Engine/Source/Editor/UMGEditor/Private/Designer/SDesignSurface.cpp
+++ b/Engine/Source/Editor/UMGEditor/Private/Designer/SDesignSurface.cpp
@@ -112,7 +112,8 @@ void SDesignSurface::Construct(const FArguments& InArgs)
 	PreviousZoomLevel = ZoomLevels->GetDefaultZoomLevel();
 	PostChangedZoom();
 	SnapGridSize = 16.0f;
-	bAllowContinousZoomInterpolation = false;
+	AllowContinousZoomInterpolation = InArgs._AllowContinousZoomInterpolation;
+	bIsPanning = false;
 
 	ViewOffset = FVector2D::ZeroVector;
 
@@ -138,6 +139,51 @@ int32 SDesignSurface::OnPaint(const FGeometry& AllottedGeometry, const FSlateRec
 	return LayerId;
 }
 
+FReply SDesignSurface::OnMouseButtonDown(const FGeometry& MyGeometry, const FPointerEvent& MouseEvent)
+{
+	if ( MouseEvent.GetEffectingButton() == EKeys::RightMouseButton )
+	{
+		bIsPanning = false;
+	}
+
+	return FReply::Unhandled();
+}
+
+FReply SDesignSurface::OnMouseButtonUp(const FGeometry& MyGeometry, const FPointerEvent& MouseEvent)
+{
+	if ( MouseEvent.GetEffectingButton() == EKeys::RightMouseButton )
+	{
+		bIsPanning = false;
+	}
+
+	return FReply::Unhandled();
+}
+
+FReply SDesignSurface::OnMouseMove(const FGeometry& MyGeometry, const FPointerEvent& MouseEvent)
+{
+	const bool bIsRightMouseButtonDown = MouseEvent.IsMouseButtonDown(EKeys::RightMouseButton);
+	const bool bIsLeftMouseButtonDown = MouseEvent.IsMouseButtonDown(EKeys::LeftMouseButton);
+	const FModifierKeysState ModifierKeysState = FSlateApplication::Get().GetModifierKeys();
+
+	if ( HasMouseCapture() )
+	{
+		const FVector2D CursorDelta = MouseEvent.GetCursorDelta();
+
+		//const bool bShouldZoom = bIsRightMouseButtonDown && ( bIsLeftMouseButtonDown || ModifierKeysState.IsAltDown() || FSlateApplication::Get().IsUsingTrackpad() );
+		if ( bIsRightMouseButtonDown )
+		{
+			FReply ReplyState = FReply::Handled();
+
+			bIsPanning = true;
+			ViewOffset -= CursorDelta / GetZoomAmount();
+
+			return ReplyState;
+		}
+	}
+
+	return FReply::Unhandled();
+}
+
 FReply SDesignSurface::OnMouseWheel(const FGeometry& MyGeometry, const FPointerEvent& MouseEvent)
 {
 	// We want to zoom into this point; i.e. keep it the same fraction offset into the panel
@@ -155,7 +201,7 @@ inline float FancyMod(float Value, float Size)
 
 float SDesignSurface::GetZoomAmount() const
 {
-	if ( bAllowContinousZoomInterpolation )
+	if ( AllowContinousZoomInterpolation.Get() )
 	{
 		return FMath::Lerp(ZoomLevels->GetZoomAmount(PreviousZoomLevel), ZoomLevels->GetZoomAmount(ZoomLevel), ZoomLevelGraphFade.GetLerp());
 	}
@@ -194,22 +240,22 @@ void SDesignSurface::ChangeZoomLevel(int32 ZoomLevelDelta, const FVector2D& Widg
 	if ( OldZoomLevel != ZoomLevel )
 	{
 		PostChangedZoom();
-	}
 
-	// Note: This happens even when maxed out at a stop; so the user sees the animation and knows that they're at max zoom in/out
-	ZoomLevelFade.Play();
+		// Note: This happens even when maxed out at a stop; so the user sees the animation and knows that they're at max zoom in/out
+		ZoomLevelFade.Play();
 
-	// Re-center the screen so that it feels like zooming around the cursor.
-	{
-		FSlateRect GraphBounds = ComputeSensibleGraphBounds();
+		// Re-center the screen so that it feels like zooming around the cursor.
+		{
+			FSlateRect GraphBounds = ComputeSensibleGraphBounds();
 
-		// Make sure we are not zooming into/out into emptiness; otherwise the user will get lost..
-		const FVector2D ClampedPointToMaintainGraphSpace(
-			FMath::Clamp(PointToMaintainGraphSpace.X, GraphBounds.Left, GraphBounds.Right),
-			FMath::Clamp(PointToMaintainGraphSpace.Y, GraphBounds.Top, GraphBounds.Bottom)
-			);
+			// Make sure we are not zooming into/out into emptiness; otherwise the user will get lost..
+			const FVector2D ClampedPointToMaintainGraphSpace(
+				FMath::Clamp(PointToMaintainGraphSpace.X, GraphBounds.Left, GraphBounds.Right),
+				FMath::Clamp(PointToMaintainGraphSpace.Y, GraphBounds.Top, GraphBounds.Bottom)
+				);
 
-		this->ViewOffset = ClampedPointToMaintainGraphSpace - WidgetSpaceZoomOrigin / GetZoomAmount();
+			this->ViewOffset = ClampedPointToMaintainGraphSpace - WidgetSpaceZoomOrigin / GetZoomAmount();
+		}
 	}
 }
 
@@ -217,8 +263,8 @@ FSlateRect SDesignSurface::ComputeSensibleGraphBounds() const
 {
 	float Left = 0.0f;
 	float Top = 0.0f;
-	float Right = 0.0f;
-	float Bottom = 0.0f;
+	float Right = 1920.f;
+	float Bottom = 1080.0f;
 
 	// Find the bounds of the node positions
 	//for ( int32 ChildIndex = 0; ChildIndex < Children.Num(); ++ChildIndex )
diff --git a/Engine/Source/Editor/UMGEditor/Private/Designer/SDesignSurface.h b/Engine/Source/Editor/UMGEditor/Private/Designer/SDesignSurface.h
--- a/Engine/Source/Editor/UMGEditor/Private/Designer/SDesignSurface.h
+++ b/Engine/Source/Editor/UMGEditor/Private/Designer/SDesignSurface.h
@@ -9,14 +9,22 @@ class SDesignSurface : public SCompoundWidget
 {
 public:
 
-	SLATE_BEGIN_ARGS( SDesignSurface ) {}
+	SLATE_BEGIN_ARGS( SDesignSurface )
+		: _AllowContinousZoomInterpolation(false)
+	{ }
+
 		/** Slot for this designers content (optional) */
 		SLATE_DEFAULT_SLOT(FArguments, Content)
+
+		SLATE_ATTRIBUTE(bool, AllowContinousZoomInterpolation)
 	SLATE_END_ARGS()
 
 	void Construct(const FArguments& InArgs);
 
 	// SWidget interface
+	virtual FReply OnMouseButtonDown(const FGeometry& MyGeometry, const FPointerEvent& MouseEvent) override;
+	virtual FReply OnMouseButtonUp(const FGeometry& MyGeometry, const FPointerEvent& MouseEvent) override;
+	virtual FReply OnMouseMove(const FGeometry& MyGeometry, const FPointerEvent& MouseEvent) override;
 	virtual FReply OnMouseWheel(const FGeometry& MyGeometry, const FPointerEvent& MouseEvent) override;
 	virtual int32 OnPaint(const FGeometry& AllottedGeometry, const FSlateRect& MyClippingRect, FSlateWindowElementList& OutDrawElements, int32 LayerId, const FWidgetStyle& InWidgetStyle, bool bParentEnabled) const override;
 	virtual bool SupportsKeyboardFocus() const override { return true; }
@@ -53,8 +61,11 @@ protected:
 	/** How zoomed in/out we are. e.g. 0.25f results in quarter-sized nodes. */
 	int32 ZoomLevel;
 
+	/** Are we panning the view at the moment? */
+	bool bIsPanning;
+
 	/** Allow continuous zoom interpolation? */
-	bool bAllowContinousZoomInterpolation;
+	TAttribute<bool> AllowContinousZoomInterpolation;
 
 	/** Fade on zoom for graph */
 	FCurveSequence ZoomLevelGraphFade;
diff --git a/Engine/Source/Editor/UMGEditor/Private/Designer/SDesignerView.cpp b/Engine/Source/Editor/UMGEditor/Private/Designer/SDesignerView.cpp
--- a/Engine/Source/Editor/UMGEditor/Private/Designer/SDesignerView.cpp
+++ b/Engine/Source/Editor/UMGEditor/Private/Designer/SDesignerView.cpp
@@ -8,6 +8,7 @@
 #include "BlueprintEditorUtils.h"
 
 #include "WidgetTemplateDragDropOp.h"
+#include "SZoomPan.h"
 
 #define LOCTEXT_NAMESPACE "UMG"
 
@@ -126,6 +127,7 @@ void SDesignerView::Construct(const FArguments& InArgs, TSharedPtr<FWidgetBluepr
 	Blueprint->OnChanged().AddSP(this, &SDesignerView::OnBlueprintChanged);
 
 	SDesignSurface::Construct(SDesignSurface::FArguments()
+		.AllowContinousZoomInterpolation(false)
 		.Content()
 		[
 			SNew(SOverlay)
@@ -135,8 +137,10 @@ void SDesignerView::Construct(const FArguments& InArgs, TSharedPtr<FWidgetBluepr
 			.HAlign(HAlign_Fill)
 			.VAlign(VAlign_Fill)
 			[
-				SAssignNew(PreviewSurface, SBorder)
+				SAssignNew(PreviewSurface, SZoomPan)
 				.Visibility(EVisibility::HitTestInvisible)
+				.ZoomAmount(this, &SDesignerView::GetZoomAmount)
+				.ViewOffset(this, &SDesignerView::GetViewOffset)
 			]
 
 			// A layer in the overlay where we put all the user intractable widgets, like the reorder widgets.
@@ -343,6 +347,8 @@ FWidgetReference SDesignerView::GetWidgetAtCursor(const FGeometry& MyGeometry, c
 
 FReply SDesignerView::OnMouseButtonDown(const FGeometry& MyGeometry, const FPointerEvent& MouseEvent)
 {
+	SDesignSurface::OnMouseButtonDown(MyGeometry, MouseEvent);
+
 	CurrentHandle = HitTestDragHandles(MyGeometry, MouseEvent);
 
 	if ( CurrentHandle != DH_NONE )
@@ -389,9 +395,14 @@ FReply SDesignerView::OnMouseButtonUp(const FGeometry& MyGeometry, const FPointe
 	}
 	else if ( MouseEvent.GetEffectingButton() == EKeys::RightMouseButton )
 	{
-		ShowContextMenu(MyGeometry, MouseEvent);
+		if ( !bIsPanning )
+		{
+			ShowContextMenu(MyGeometry, MouseEvent);
+		}
 	}
 
+	SDesignSurface::OnMouseButtonUp(MyGeometry, MouseEvent);
+
 	return FReply::Handled().ReleaseMouseCapture();
 }
 
@@ -402,6 +413,12 @@ FReply SDesignerView::OnMouseMove(const FGeometry& MyGeometry, const FPointerEve
 		return FReply::Unhandled();
 	}
 
+	FReply SurfaceHandled = SDesignSurface::OnMouseMove(MyGeometry, MouseEvent);
+	if ( SurfaceHandled.IsEventHandled() )
+	{
+		return SurfaceHandled;
+	}
+
 	// Update the resizing based on the mouse movement
 	if ( CurrentHandle != DH_NONE )
 	{
@@ -575,6 +592,31 @@ FVector2D SDesignerView::GetSelectionDesignerWidgetsLocation() const
 	return FVector2D(0, 0);
 }
 
+//void SDesignerView::OnArrangeChildren(const FGeometry& AllottedGeometry, FArrangedChildren& ArrangedChildren) const
+//{
+//	const TSlotlessChildren<SNode>& ChildrenToArrange = ArrangedChildren.Accepts(EVisibility::Hidden) ? Children : VisibleChildren;
+//	// First pass nodes
+//	for ( int32 ChildIndex = 0; ChildIndex < ChildrenToArrange.Num(); ++ChildIndex )
+//	{
+//		const TSharedRef<SNode>& SomeChild = ChildrenToArrange[ChildIndex];
+//		if ( !SomeChild->RequiresSecondPassLayout() )
+//		{
+//			ArrangedChildren.AddWidget(AllottedGeometry.MakeChild(SomeChild, SomeChild->GetPosition() - ViewOffset, SomeChild->GetDesiredSize(), GetZoomAmount()));
+//		}
+//	}
+//
+//	// Second pass nodes
+//	for ( int32 ChildIndex = 0; ChildIndex < ChildrenToArrange.Num(); ++ChildIndex )
+//	{
+//		const TSharedRef<SNode>& SomeChild = ChildrenToArrange[ChildIndex];
+//		if ( SomeChild->RequiresSecondPassLayout() )
+//		{
+//			SomeChild->PerformSecondPassLayout(NodeToWidgetLookup);
+//			ArrangedChildren.AddWidget(AllottedGeometry.MakeChild(SomeChild, SomeChild->GetPosition() - ViewOffset, SomeChild->GetDesiredSize(), GetZoomAmount()));
+//		}
+//	}
+//}
+
 int32 SDesignerView::OnPaint(const FGeometry& AllottedGeometry, const FSlateRect& MyClippingRect, FSlateWindowElementList& OutDrawElements, int32 LayerId, const FWidgetStyle& InWidgetStyle, bool bParentEnabled) const
 {
 	SDesignSurface::OnPaint(AllottedGeometry, MyClippingRect, OutDrawElements, LayerId, InWidgetStyle, bParentEnabled);
diff --git a/Engine/Source/Editor/UMGEditor/Private/Designer/SDesignerView.h b/Engine/Source/Editor/UMGEditor/Private/Designer/SDesignerView.h
--- a/Engine/Source/Editor/UMGEditor/Private/Designer/SDesignerView.h
+++ b/Engine/Source/Editor/UMGEditor/Private/Designer/SDesignerView.h
@@ -34,6 +34,8 @@ public:
 
 	virtual FCursorReply OnCursorQuery(const FGeometry& MyGeometry, const FPointerEvent& CursorEvent) const override;
 
+	//virtual void OnArrangeChildren(const FGeometry& AllottedGeometry, FArrangedChildren& ArrangedChildren) const;
+
 	virtual int32 OnPaint(const FGeometry& AllottedGeometry, const FSlateRect& MyClippingRect, FSlateWindowElementList& OutDrawElements, int32 LayerId, const FWidgetStyle& InWidgetStyle, bool bParentEnabled) const override;
 	virtual void Tick(const FGeometry& AllottedGeometry, const double InCurrentTime, const float InDeltaTime) override;
 
@@ -52,6 +54,8 @@ private:
 	/** Displays the context menu when you right click */
 	void ShowContextMenu(const FGeometry& MyGeometry, const FPointerEvent& MouseEvent);
 
+	FVector2D GetZoomVector() const;
+
 	void OnEditorSelectionChanged();
 
 	/** @returns Gets the widget under the cursor based on a mouse pointer event. */
@@ -123,7 +127,7 @@ private:
 	UWidget* DropPreviewWidget;
 	UPanelWidget* DropPreviewParent;
 
-	TSharedPtr<SBorder> PreviewSurface;
+	TSharedPtr<class SZoomPan> PreviewSurface;
 	TSharedPtr<SCanvas> ExtensionWidgetCanvas;
 
 	DragHandle CurrentHandle;
diff --git a/Engine/Source/Editor/UMGEditor/Private/Designer/SZoomPan.cpp b/Engine/Source/Editor/UMGEditor/Private/Designer/SZoomPan.cpp
new file mode 100644
--- /dev/null
+++ b/Engine/Source/Editor/UMGEditor/Private/Designer/SZoomPan.cpp
@@ -0,0 +1,49 @@
+// Copyright 1998-2014 Epic Games, Inc. All Rights Reserved.
+
+#include "UMGEditorPrivatePCH.h"
+
+#include "SZoomPan.h"
+
+#define LOCTEXT_NAMESPACE "UMG"
+
+/////////////////////////////////////////////////////
+// SZoomPan
+
+void SZoomPan::Construct(const FArguments& InArgs)
+{
+	ViewOffset = InArgs._ViewOffset;
+	ZoomAmount = InArgs._ZoomAmount;
+
+	ChildSlot
+	[
+		InArgs._Content.Widget
+	];
+}
+
+void SZoomPan::OnArrangeChildren(const FGeometry& AllottedGeometry, FArrangedChildren& ArrangedChildren) const
+{
+	const EVisibility ChildVisibility = ChildSlot.Widget->GetVisibility();
+	if ( ArrangedChildren.Accepts(ChildVisibility) )
+	{
+		const FMargin SlotPadding(ChildSlot.SlotPadding.Get());
+		AlignmentArrangeResult XResult = AlignChild<Orient_Horizontal>(AllottedGeometry.Size.X, ChildSlot, SlotPadding, 1);
+		AlignmentArrangeResult YResult = AlignChild<Orient_Vertical>(AllottedGeometry.Size.Y, ChildSlot, SlotPadding, 1);
+
+		ArrangedChildren.AddWidget( ChildVisibility, AllottedGeometry.MakeChild(
+				ChildSlot.Widget,
+				FVector2D(XResult.Offset, YResult.Offset) - ViewOffset.Get(),
+				FVector2D(XResult.Size, YResult.Size),
+				ZoomAmount.Get()
+		) );
+	}
+}
+
+void SZoomPan::SetContent(const TSharedRef< SWidget >& InContent)
+{
+	ChildSlot
+	[
+		InContent
+	];
+}
+
+#undef LOCTEXT_NAMESPACE
diff --git a/Engine/Source/Editor/UMGEditor/Private/Designer/SZoomPan.h b/Engine/Source/Editor/UMGEditor/Private/Designer/SZoomPan.h
new file mode 100644
--- /dev/null
+++ b/Engine/Source/Editor/UMGEditor/Private/Designer/SZoomPan.h
@@ -0,0 +1,37 @@
+// Copyright 1998-2014 Epic Games, Inc. All Rights Reserved.
+
+#pragma once
+
+#include "SCompoundWidget.h"
+
+class SZoomPan : public SCompoundWidget
+{
+public:
+
+	SLATE_BEGIN_ARGS(SZoomPan) {}
+		/** Slot for this designers content (optional) */
+		SLATE_DEFAULT_SLOT(FArguments, Content)
+		
+		SLATE_ATTRIBUTE(FVector2D, ViewOffset)
+		
+		SLATE_ATTRIBUTE(float, ZoomAmount)
+	SLATE_END_ARGS()
+
+	void Construct(const FArguments& InArgs);
+	
+	void OnArrangeChildren(const FGeometry& AllottedGeometry, FArrangedChildren& ArrangedChildren) const;
+
+	/**
+	 * Sets the content for this border
+	 *
+	 * @param	InContent	The widget to use as content for the border
+	 */
+	void SetContent( const TSharedRef< SWidget >& InContent );
+
+protected:
+	/** The position within the panel at which the user is looking */
+	TAttribute<FVector2D> ViewOffset;
+
+	/** How zoomed in/out we are. e.g. 0.25f results in quarter-sized widgets. */
+	TAttribute<float> ZoomAmount;
+};
diff --git a/Engine/Source/Editor/UMGEditor/Private/WidgetBlueprintEditorUtils.cpp b/Engine/Source/Editor/UMGEditor/Private/WidgetBlueprintEditorUtils.cpp
--- a/Engine/Source/Editor/UMGEditor/Private/WidgetBlueprintEditorUtils.cpp
+++ b/Engine/Source/Editor/UMGEditor/Private/WidgetBlueprintEditorUtils.cpp
@@ -409,7 +409,7 @@ void FWidgetBlueprintEditorUtils::ImportWidgetsFromText(UWidgetBlueprint* BP, co
 {
 	// We create our own transient package here so that we can deserialize the data in isolation and ensure unreferenced
 	// objects not part of the deserialization set are unresolved.
-	UPackage* TempPackage = new( NULL, TEXT("/Engine/UMG/Editor/Transient") )UPackage(FPostConstructInitializeProperties());
+	UPackage* TempPackage = ConstructObject<UPackage>(UPackage::StaticClass(), nullptr, TEXT("/Engine/UMG/Editor/Transient"), RF_Transient);
 	TempPackage->AddToRoot();
 
 	// Turn the text buffer into objects
diff --git a/Engine/Source/Editor/UnrealEd/Private/PackageAutoSaver.cpp b/Engine/Source/Editor/UnrealEd/Private/PackageAutoSaver.cpp
--- a/Engine/Source/Editor/UnrealEd/Private/PackageAutoSaver.cpp
+++ b/Engine/Source/Editor/UnrealEd/Private/PackageAutoSaver.cpp
@@ -316,34 +316,37 @@ void FPackageAutoSaver::UpdateDirtyListsForPackage(UPackage* Pkg)
 {
 	const UPackage* TransientPackage = GetTransientPackage();
 
-	if(Pkg != TransientPackage)
+	// Don't auto-save the transient package or packages with the transient flag.
+	if ( Pkg == TransientPackage || Pkg->HasAnyFlags(RF_Transient) )
 	{
-		if(Pkg->IsDirty())
-		{
-			// Always add the package to the user list
-			DirtyPackagesForUserSave.FindOrAdd(Pkg);
+		return;
+	}
 
-			// Only add the package to the auto-save list if we're not auto-saving
-			// Note: Packages get dirtied again after they're auto-saved, so this would add them back again, which we don't want
-			if(!IsAutoSaving())
-			{
-				DirtyPackagesForAutoSave.Add(Pkg);
-			}
-		}
-		else
+	if ( Pkg->IsDirty() )
+	{
+		// Always add the package to the user list
+		DirtyPackagesForUserSave.FindOrAdd(Pkg);
+
+		// Only add the package to the auto-save list if we're not auto-saving
+		// Note: Packages get dirtied again after they're auto-saved, so this would add them back again, which we don't want
+		if ( !IsAutoSaving() )
 		{
-			// Always remove the package from the auto-save list
-			DirtyPackagesForAutoSave.Remove(Pkg);
+			DirtyPackagesForAutoSave.Add(Pkg);
+		}
+	}
+	else
+	{
+		// Always remove the package from the auto-save list
+		DirtyPackagesForAutoSave.Remove(Pkg);
 
-			// Only remove the package from the user list if we're not auto-saving
-			// Note: Packages call this even when auto-saving, so this would remove them from the user list, which we don't want as they're still dirty
-			if(!IsAutoSaving())
+		// Only remove the package from the user list if we're not auto-saving
+		// Note: Packages call this even when auto-saving, so this would remove them from the user list, which we don't want as they're still dirty
+		if ( !IsAutoSaving() )
+		{
+			if ( DirtyPackagesForUserSave.Remove(Pkg) )
 			{
-				if(DirtyPackagesForUserSave.Remove(Pkg))
-				{
-					// Update the restore information too
-					UpdateRestoreFile(PackageAutoSaverJson::IsRestoreEnabled());
-				}
+				// Update the restore information too
+				UpdateRestoreFile(PackageAutoSaverJson::IsRestoreEnabled());
 			}
 		}
 	}
