UMG - Add filtering to the widget palette.
+ Create a filter handler class to package up simple filtering behavior for tree views.  It didn't end up as clean I would have liked, but it avoids a lot of duplicated code.
+ Use the new filter handler in both the hierarchy and the palette.

diff --git a/Engine/Source/Editor/UMGEditor/Private/Hierarchy/SHierarchyView.cpp b/Engine/Source/Editor/UMGEditor/Private/Hierarchy/SHierarchyView.cpp
--- a/Engine/Source/Editor/UMGEditor/Private/Hierarchy/SHierarchyView.cpp
+++ b/Engine/Source/Editor/UMGEditor/Private/Hierarchy/SHierarchyView.cpp
@@ -21,21 +21,27 @@ void SHierarchyView::Construct(const FArguments& InArgs, TSharedPtr<FWidgetBluep
 {
 	BlueprintEditor = InBlueprintEditor;
 	bRefreshRequested = false;
-	bIsFilterActive = false;
 
 	SearchBoxWidgetFilter = MakeShareable(new WidgetTextFilter(WidgetTextFilter::FItemToStringArray::CreateSP(this, &SHierarchyView::TransformWidgetToString)));
 
 	UWidgetBlueprint* Blueprint = GetBlueprint();
 	Blueprint->OnChanged().AddSP(this, &SHierarchyView::OnBlueprintChanged);
 
+	FilterHandler = MakeShareable(new TreeFilterHandler<UWidget*>());
+	FilterHandler->SetFilter(SearchBoxWidgetFilter.Get());
+	FilterHandler->SetRootItems(&RootWidgets, &TreeRootWidgets);
+	FilterHandler->SetGetChildrenDelegate(TreeFilterHandler<UWidget*>::FOnGetChildren::CreateRaw(this, &SHierarchyView::WidgetHierarchy_OnGetChildren));
+
 	SAssignNew(WidgetTreeView, STreeView< UWidget* >)
 	.ItemHeight(20.0f)
 	.SelectionMode(ESelectionMode::Single)
-	.OnGetChildren(this, &SHierarchyView::WidgetHierarchy_OnGetChildren)
+	.OnGetChildren(FilterHandler.ToSharedRef(), &TreeFilterHandler<UWidget*>::OnGetFilteredChildren)
 	.OnGenerateRow(this, &SHierarchyView::WidgetHierarchy_OnGenerateRow)
 	.OnSelectionChanged(this, &SHierarchyView::WidgetHierarchy_OnSelectionChanged)
 	.OnContextMenuOpening(this, &SHierarchyView::WidgetHierarchy_OnContextMenuOpening)
-	.TreeItemsSource(&RootWidgets);
+	.TreeItemsSource(&TreeRootWidgets);
+
+	FilterHandler->SetTreeView(WidgetTreeView.Get());
 
 	ChildSlot
 	[
@@ -105,7 +111,7 @@ FReply SHierarchyView::OnKeyDown(const FGeometry& MyGeometry, const FKeyboardEve
 	return FReply::Unhandled();
 }
 
-void SHierarchyView::TransformWidgetToString(const UWidget* Widget, OUT TArray< FString >& Array)
+void SHierarchyView::TransformWidgetToString(UWidget* Widget, OUT TArray< FString >& Array)
 {
 	Array.Add( Widget->GetLabel() );
 }
@@ -113,6 +119,7 @@ void SHierarchyView::TransformWidgetToString(const UWidget* Widget, OUT TArray<
 void SHierarchyView::OnSearchChanged(const FText& InFilterText)
 {
 	bRefreshRequested = true;
+	FilterHandler->SetIsEnabled(!InFilterText.IsEmpty());
 	SearchBoxWidgetFilter->SetRawFilterText(InFilterText);
 }
 
@@ -210,11 +217,6 @@ void SHierarchyView::WidgetHierarchy_OnGetChildren(UWidget* InParent, TArray< UW
 			UWidget* Child = Widget->GetChildAt(i);
 			if ( Child )
 			{
-				if ( bIsFilterActive && !WidgetsPassingFilter.Contains(Child) )
-				{
-					continue;
-				}
-
 				OutChildren.Add(Child);
 			}
 		}
@@ -242,12 +244,11 @@ void SHierarchyView::WidgetHierarchy_OnSelectionChanged(UWidget* SelectedItem, E
 FReply SHierarchyView::HandleDeleteSelected()
 {
 	TSet<FWidgetReference> SelectedWidgets = BlueprintEditor.Pin()->GetSelectedWidgets();
-	//TArray<UWidget*> SelectedWidgets = WidgetTreeView->GetSelectedItems();
 
-	// Remove the selected items from the widget cache
-	for ( FWidgetReference& Item : SelectedWidgets )
+	// Remove the selected items from the filter cache
+	for (FWidgetReference& Item : SelectedWidgets)
 	{
-		CachedExpandedWidgets.Remove(Item.GetTemplate());
+		FilterHandler->RemoveCachedItem(Item.GetTemplate());
 	}
 
 	FWidgetBlueprintEditorUtils::DeleteWidgets(GetBlueprint(), SelectedWidgets);
@@ -255,90 +256,17 @@ FReply SHierarchyView::HandleDeleteSelected()
 	return FReply::Handled();
 }
 
-bool SHierarchyView::FilterWidgetHierarchy(UWidget* CurrentWidget)
-{
-	bool bAnyChildrenPass = false;
-
-	// Iterate over children and check to see if any of them pass the current filter
-	UPanelWidget* Widget = Cast<UPanelWidget>(CurrentWidget);
-	if ( Widget )
-	{
-		for ( int32 i = 0; i < Widget->GetChildrenCount(); i++ )
-		{
-			UWidget* Child = Widget->GetChildAt(i);
-			if ( Child )
-			{
-				bAnyChildrenPass |= FilterWidgetHierarchy(Child);
-			}
-		}
-	}
-
-	// Check to see if I pass the filter
-	const bool bWidgetPass = SearchBoxWidgetFilter->PassesFilter(CurrentWidget);
-
-	// If this particular widget passes the filter, expand every widget leading up to it.
-	if ( bWidgetPass )
-	{
-		ExpandPathToWidget(CurrentWidget);
-	}
-
-	// If either the widget or the children pass, add the current widget
-	if ( bWidgetPass || bAnyChildrenPass )
-	{
-		WidgetsPassingFilter.Add(CurrentWidget);
-		return true;
-	}
-
-	return false;
-}
-
 void SHierarchyView::RefreshTree()
 {
-	RootWidgets.Reset();
-
+	RootWidgets.Empty();
 	UWidgetBlueprint* Blueprint = GetBlueprint();
-
-	if ( Blueprint->WidgetTree->RootWidget )
+	if (Blueprint->WidgetTree->RootWidget)
 	{
-		bool bRootPassed = true;
-
-		const bool bWillFilterBeActive = !SearchBoxWidgetFilter->GetRawFilterText().IsEmpty();
-
-		// Save the expansion state when the filter becomes active
-		if ( !bIsFilterActive && bWillFilterBeActive )
-		{
-			CachedExpandedWidgets.Empty();
-			WidgetTreeView->GetExpandedItems(CachedExpandedWidgets);
-		}
-		// Restore the expansion state when the filter is removed
-		else if ( bIsFilterActive && !bWillFilterBeActive )
-		{
-			WidgetTreeView->ClearExpandedItems();
-			for ( UWidget* ExpandedWidget : CachedExpandedWidgets )
-			{
-				WidgetTreeView->SetItemExpansion(ExpandedWidget, true);
-			}
-			CachedExpandedWidgets.Empty();
-		}
-		
-		bIsFilterActive = bWillFilterBeActive;
-
-		WidgetsPassingFilter.Empty();
-		if ( bIsFilterActive )
-		{
-			bRootPassed = FilterWidgetHierarchy(Blueprint->WidgetTree->RootWidget);
-		}
-
-		if ( bRootPassed )
-		{
-			RootWidgets.Add(Blueprint->WidgetTree->RootWidget);
-		}
+		RootWidgets.Add(Blueprint->WidgetTree->RootWidget);
 	}
-
-	WidgetTreeView->RequestTreeRefresh();
+	FilterHandler->RefreshAndFilterTree();
 }
 
-
 //@TODO UMG Drop widgets onto the tree, when nothing is present, if there is a root node present, what happens then, let the root node attempt to place it?
 
 #undef LOCTEXT_NAMESPACE
diff --git a/Engine/Source/Editor/UMGEditor/Private/Hierarchy/SHierarchyView.h b/Engine/Source/Editor/UMGEditor/Private/Hierarchy/SHierarchyView.h
--- a/Engine/Source/Editor/UMGEditor/Private/Hierarchy/SHierarchyView.h
+++ b/Engine/Source/Editor/UMGEditor/Private/Hierarchy/SHierarchyView.h
@@ -4,6 +4,7 @@
 
 #include "SCompoundWidget.h"
 #include "BlueprintEditor.h"
+#include "TreeFilterHandler.h"
 
 //TODO rename SUMGEditorHierarchy
 
@@ -14,7 +15,7 @@
 class SHierarchyView : public SCompoundWidget
 {
 public:
-	typedef TTextFilter< const UWidget* > WidgetTextFilter;
+	typedef TTextFilter< UWidget* > WidgetTextFilter;
 
 public:
 	SLATE_BEGIN_ARGS( SHierarchyView ){}
@@ -72,37 +73,28 @@ private:
 	FText GetSearchText() const;
 
 	/** Transforms the widget into a searchable string */
-	void TransformWidgetToString(const UWidget* Widget, OUT TArray< FString >& Array);
-
-	/** Builds a filterable set by recursing through the widget hierarchy */
-	bool FilterWidgetHierarchy(UWidget* CurrentWidget);
+	void TransformWidgetToString(UWidget* Widget, OUT TArray< FString >& Array);
 
 private:
 
 	/** Cached pointer to the blueprint editor that owns this tree. */
 	TWeakPtr<class FWidgetBlueprintEditor> BlueprintEditor;
 
-	/** The tree item source.  We only ever have one root widget, but the tree expects an array. */
+	/** Handles filtering the hierarchy based on an IFilter. */
+	TSharedPtr<TreeFilterHandler<UWidget*>> FilterHandler;
+
+	/** The source root widgets for the tree. */
 	TArray< UWidget* > RootWidgets;
 
+	/** The root widgets which are actually displayed by the TreeView which will be managed by the TreeFilterHandler. */
+	TArray< UWidget* > TreeRootWidgets;
+
 	/** The widget hierarchy slate treeview widget */
 	TSharedPtr< STreeView< UWidget* > > WidgetTreeView;
 
 	/** The filter used by the search box */
 	TSharedPtr<WidgetTextFilter> SearchBoxWidgetFilter;
 
-	/** The set of widgets currently passing the filter.  Used to quickly filter children when the tree is expanded lazily. */
-	TSet< UWidget* > WidgetsPassingFilter;
-
-	/**
-	  * A cache of the expanded items.  We store this off before a user filters the list so that we can 
-	  * restore the expansion state after the filter is cleared.
-	  */
-	TSet< UWidget* > CachedExpandedWidgets;
-
-	/** Caches a value to know if a filter is currently active. */
-	bool bIsFilterActive;
-
 	/** Has a full refresh of the tree been requested?  This happens when the user is filtering the tree */
 	bool bRefreshRequested;
 };
diff --git a/Engine/Source/Editor/UMGEditor/Private/Palette/SPaletteView.cpp b/Engine/Source/Editor/UMGEditor/Private/Palette/SPaletteView.cpp
--- a/Engine/Source/Editor/UMGEditor/Private/Palette/SPaletteView.cpp
+++ b/Engine/Source/Editor/UMGEditor/Private/Palette/SPaletteView.cpp
@@ -25,6 +25,10 @@ class SPaletteViewItem : public SCompoundWidget
 public:
 
 	SLATE_BEGIN_ARGS(SPaletteViewItem) {}
+
+		/** The current text to highlight */
+		SLATE_ATTRIBUTE(FText, HighlightText)
+
 	SLATE_END_ARGS()
 
 	/**
@@ -56,6 +60,7 @@ public:
 			[
 				SNew(STextBlock)
 				.Text(Template->Name)
+				.HighlightText(InArgs._HighlightText)
 			]
 		];
 	}
@@ -72,6 +77,11 @@ public:
 		return Template->Name;
 	}
 
+	virtual FString GetFilterString() const
+	{
+		return Template->Name.ToString();
+	}
+
 	virtual TSharedRef<ITableRow> BuildRow(const TSharedRef<STableViewBase>& OwnerTable) override
 	{
 		return SNew(STableRow< TSharedPtr<FWidgetViewModel> >, OwnerTable)
@@ -79,6 +89,7 @@ public:
 			.OnDragDetected(this, &FWidgetTemplateViewModel::OnDraggingWidgetTemplateItem)
 			[
 				SNew(SPaletteViewItem, Template)
+				.HighlightText(OwnerView, &SPaletteView::GetSearchText)
 			];
 	}
 
@@ -87,6 +98,7 @@ public:
 		return FReply::Handled().BeginDragDrop(FWidgetTemplateDragDropOp::New(Template));
 	}
 
+	SPaletteView* OwnerView;
 	TSharedPtr<FWidgetTemplate> Template;
 };
 
@@ -97,6 +109,13 @@ public:
 	{
 		return GroupName;
 	}
+
+	virtual FString GetFilterString() const
+	{
+		// Headers should never be included in filtering to avoid showing a header with all of
+		// it's widgets filtered out, so return an empty filter string.
+		return TEXT("");
+	}
 	
 	virtual TSharedRef<ITableRow> BuildRow(const TSharedRef<STableViewBase>& OwnerTable) override
 	{
@@ -128,15 +147,25 @@ void SPaletteView::Construct(const FArguments& InArgs, TSharedPtr<FBlueprintEdit
 
 	UBlueprint* BP = InBlueprintEditor->GetBlueprintObj();
 
+	WidgetFilter = MakeShareable(new WidgetViewModelTextFilter(
+		WidgetViewModelTextFilter::FItemToStringArray::CreateSP(this, &SPaletteView::TransformWidgetViewModelToString)));
+
 	BuildWidgetList();
 
 	//FCoreDelegates::OnObjectPropertyChanged.AddRaw(this, &SPaletteView::OnObjectPropertyChanged);
 
+	FilterHandler = MakeShareable(new PaletteFilterHandler());
+	FilterHandler->SetFilter(WidgetFilter.Get());
+	FilterHandler->SetRootItems(&WidgetViewModels, &TreeWidgetViewModels);
+	FilterHandler->SetGetChildrenDelegate(PaletteFilterHandler::FOnGetChildren::CreateRaw(this, &SPaletteView::OnGetChildren));
+
 	SAssignNew(WidgetTemplatesView, STreeView< TSharedPtr<FWidgetViewModel> >)
 	.SelectionMode(ESelectionMode::Single)
 	.OnGenerateRow(this, &SPaletteView::OnGenerateWidgetTemplateItem)
-	.OnGetChildren(this, &SPaletteView::OnGetChildren)
-	.TreeItemsSource(&WidgetViewModels);
+	.OnGetChildren(FilterHandler.ToSharedRef(), &PaletteFilterHandler::OnGetFilteredChildren)
+	.TreeItemsSource(&TreeWidgetViewModels);
+
+	FilterHandler->SetTreeView(WidgetTemplatesView.Get());
 
 	ChildSlot
 	[
@@ -162,20 +191,36 @@ void SPaletteView::Construct(const FArguments& InArgs, TSharedPtr<FBlueprintEdit
 	];
 
 	LoadItemExpanssion();
+	bRefreshRequested = true;
 }
 
 SPaletteView::~SPaletteView()
 {
 	//FCoreDelegates::OnObjectPropertyChanged.RemoveAll(this);
+
+	// If the filter is enabled, disable it before saving the expanded items since
+	// filtering expands all items by default.
+	if (FilterHandler->GetIsEnabled())
+	{
+		FilterHandler->SetIsEnabled(false);
+		FilterHandler->RefreshAndFilterTree();
+	}
 	SaveItemExpansion();
 }
 
 void SPaletteView::OnSearchChanged(const FText& InFilterText)
 {
 	bRefreshRequested = true;
+	FilterHandler->SetIsEnabled(!InFilterText.IsEmpty());
+	WidgetFilter->SetRawFilterText(InFilterText);
 	SearchText = InFilterText;
 }
 
+FText SPaletteView::GetSearchText() const
+{
+	return SearchText;
+}
+
 void SPaletteView::LoadItemExpanssion()
 {
 	// Restore the expansion state of the widget groups.
@@ -227,7 +272,7 @@ void SPaletteView::BuildWidgetList()
 		{
 			TSharedPtr<FWidgetTemplateViewModel> TemplateViewModel = MakeShareable(new FWidgetTemplateViewModel());
 			TemplateViewModel->Template = Template;
-
+			TemplateViewModel->OwnerView = this;
 			Header->Children.Add(TemplateViewModel);
 		}
 
@@ -288,4 +333,18 @@ TSharedRef<ITableRow> SPaletteView::OnGenerateWidgetTemplateItem(TSharedPtr<FWid
 	return Item->BuildRow(OwnerTable);
 }
 
+void SPaletteView::Tick(const FGeometry& AllottedGeometry, const double InCurrentTime, const float InDeltaTime)
+{
+	if (bRefreshRequested)
+	{
+		bRefreshRequested = false;
+		FilterHandler->RefreshAndFilterTree();
+	}
+}
+
+void SPaletteView::TransformWidgetViewModelToString(TSharedPtr<FWidgetViewModel> WidgetViewModel, OUT TArray< FString >& Array)
+{
+	Array.Add(WidgetViewModel->GetFilterString());
+}
+
 #undef LOCTEXT_NAMESPACE
diff --git a/Engine/Source/Editor/UMGEditor/Private/Palette/SPaletteView.h b/Engine/Source/Editor/UMGEditor/Private/Palette/SPaletteView.h
--- a/Engine/Source/Editor/UMGEditor/Private/Palette/SPaletteView.h
+++ b/Engine/Source/Editor/UMGEditor/Private/Palette/SPaletteView.h
@@ -4,6 +4,7 @@
 
 #include "SCompoundWidget.h"
 #include "BlueprintEditor.h"
+#include "TreeFilterHandler.h"
 
 #include "WidgetTemplate.h"
 
@@ -13,6 +14,9 @@ class FWidgetViewModel : public TSharedFromThis<FWidgetViewModel>
 public:
 	virtual FText GetName() const = 0;
 
+	/** Get the string which should be used for filtering the item. */
+	virtual FString GetFilterString() const = 0;
+
 	virtual TSharedRef<ITableRow> BuildRow(const TSharedRef<STableViewBase>& OwnerTable) = 0;
 
 	virtual void GetChildren(TArray< TSharedPtr<FWidgetViewModel> >& OutChildren)
@@ -23,12 +27,20 @@ public:
 /**  */
 class SPaletteView : public SCompoundWidget
 {
+public:
+	typedef TTextFilter<TSharedPtr<FWidgetViewModel>> WidgetViewModelTextFilter;
+
 public:
 	SLATE_BEGIN_ARGS( SPaletteView ){}
 	SLATE_END_ARGS()
 
 	void Construct(const FArguments& InArgs, TSharedPtr<FBlueprintEditor> InBlueprintEditor);
 	virtual ~SPaletteView();
+	
+	virtual void Tick(const FGeometry& AllottedGeometry, const double InCurrentTime, const float InDeltaTime);
+
+	/** Gets the text currently displayed in the search box. */
+	FText GetSearchText() const;
 
 private:
 	UWidgetBlueprint* GetBlueprint() const;
@@ -49,16 +61,31 @@ private:
 
 	void AddWidgetTemplate(TSharedPtr<FWidgetTemplate> Template);
 
+	/** Transforms the widget view model into a searchable string. */
+	void TransformWidgetViewModelToString(TSharedPtr<FWidgetViewModel> WidgetViewModel, OUT TArray< FString >& Array);
+
 	TWeakPtr<class FBlueprintEditor> BlueprintEditor;
 
+	/** Handles filtering the palette based on an IFilter. */
+	typedef TreeFilterHandler<TSharedPtr<FWidgetViewModel>> PaletteFilterHandler;
+	TSharedPtr<PaletteFilterHandler> FilterHandler;
+
 	typedef TArray< TSharedPtr<FWidgetTemplate> > WidgetTemplateArray;
 	TMap< FString, WidgetTemplateArray > WidgetTemplateCategories;
 
 	typedef TArray< TSharedPtr<FWidgetViewModel> > ViewModelsArray;
+	
+	/** The source root view models for the tree. */
 	ViewModelsArray WidgetViewModels;
 
+	/** The root view models which are actually displayed by the TreeView which will be managed by the TreeFilterHandler. */
+	ViewModelsArray TreeWidgetViewModels;
+
 	TSharedPtr< STreeView< TSharedPtr<FWidgetViewModel> > > WidgetTemplatesView;
 
+	/** The filter instance which is used by the TreeFilterHandler to filter the TreeView. */
+	TSharedPtr<WidgetViewModelTextFilter> WidgetFilter;
+
 	bool bRefreshRequested;
 	FText SearchText;
 };
diff --git a/Engine/Source/Editor/UMGEditor/Private/TreeFilterHandler.h b/Engine/Source/Editor/UMGEditor/Private/TreeFilterHandler.h
new file mode 100644
--- /dev/null
+++ b/Engine/Source/Editor/UMGEditor/Private/TreeFilterHandler.h
@@ -0,0 +1,204 @@
+// Copyright 1998-2014 Epic Games, Inc. All Rights Reserved.
+
+#pragma once
+
+/** Handles filtering and expanding a TreeView using an IFilter implementation. */
+template<typename ItemType>
+class TreeFilterHandler
+{
+public:
+	typedef typename TSlateDelegates< ItemType >::FOnGetChildren FOnGetChildren;
+
+public:
+
+	TreeFilterHandler()
+	{
+		TreeView = nullptr;
+		Filter = nullptr;
+		RootItems = nullptr;
+		TreeRootItems = nullptr;
+		bIsEnabled = false;
+		bWasEnabled = false;
+	}
+
+	/** Sets the TreeView to be filtered. */
+	void SetTreeView(STreeView<ItemType>* TreeView)
+	{
+		this->TreeView = TreeView;
+	}
+
+	/** Sets the filter which will be used to filter the items in the TreeView. */
+	void SetFilter(IFilter<ItemType>* Filter)
+	{
+		this->Filter = Filter;
+	}
+
+	/** 
+	 * Sets the root data arrays for the tree.
+	 * @param RootItems		An array of source root items to be displayed by the tree.
+	 * @param TreeRootItems	The array of root items which are set as the ItemsSource of the tree. This array will be modified to reflect any filtering.
+	 */
+	void SetRootItems(TArray<ItemType>* RootItems, TArray<ItemType>* TreeRootItems)
+	{
+		this->RootItems = RootItems;
+		this->TreeRootItems = TreeRootItems;
+	}
+
+	/**
+	 * Sets the delegate which should be used for traversing the children of the nodes in
+	 * the tree.
+	 */
+	void SetGetChildrenDelegate(FOnGetChildren GetChildrenDelegate)
+	{
+		this->GetChildrenDelegate = GetChildrenDelegate;
+	}
+
+	/** Gets whether or not filtering is enabled for the TreeView. */
+	bool GetIsEnabled()
+	{
+		return bIsEnabled;
+	}
+
+	/** 
+	 * Sets whether or not filtering the tree is enabled. Calling this DOES NOT refresh the tree,
+	 * this must be done by calling RefreshAndFilterTree.
+	 */
+	void SetIsEnabled(bool bIsEnabled)
+	{
+		this->bIsEnabled = bIsEnabled;
+	}
+
+	/** Removes any cached references to an item. */
+	void RemoveCachedItem(ItemType ItemToRemove)
+	{
+		PrefilterExpandedItems.Remove(ItemToRemove);
+	}
+
+	/**
+	 * Refreshes the associated TreeView, filtering items if the filter is enabled. When filtering,
+	 * the expanded state of the items is saved, and when the filter is cleared the pre-filter expand
+	 * state is restored.
+	 */
+	void RefreshAndFilterTree()
+	{
+		if (TreeView && Filter && RootItems && TreeRootItems && GetChildrenDelegate.IsBound())
+		{
+			TreeRootItems->Empty();
+			ItemsPassingFilter.Empty();
+			if (bIsEnabled)
+			{
+				if (!bWasEnabled)
+				{
+					PrefilterExpandedItems.Empty();
+					TreeView->GetExpandedItems(PrefilterExpandedItems);
+				}
+				for (auto RootItem : *RootItems)
+				{
+					if (ExpandAndCacheMatches(RootItem))
+					{
+						TreeRootItems->Add(RootItem);
+					}
+				}
+			}
+			else
+			{
+				if (bWasEnabled)
+				{
+					TreeView->ClearExpandedItems();
+					for (auto PrefilterExpandedItem : PrefilterExpandedItems)
+					{
+						TreeView->SetItemExpansion(PrefilterExpandedItem, true);
+					}
+					PrefilterExpandedItems.Empty();
+				}
+				for (auto RootItem : *RootItems)
+				{
+					TreeRootItems->Add(RootItem);
+				}
+			}
+
+			TreeView->RequestTreeRefresh();
+			bWasEnabled = bIsEnabled;
+		}
+	}
+
+	/**
+	 * Gets the children for a node in the tree which pass the filter if filtering is enabled, or
+	 * gets all children if the filter is not enabled. In order for filtering to work correctly this
+	 * function must be used for the TreeView's OnGetChildren delegate.
+	 */
+	void OnGetFilteredChildren(ItemType Parent, TArray<ItemType>& FilteredChildren)
+	{
+		if (bIsEnabled)
+		{
+			TArray<ItemType> AllChildren;
+			GetChildrenDelegate.Execute(Parent, AllChildren);
+			for (auto Child : AllChildren)
+			{
+				if (ItemsPassingFilter.Contains(Child))
+				{
+					FilteredChildren.Add(Child);
+				}
+			}
+		}
+		else
+		{
+			GetChildrenDelegate.Execute(Parent, FilteredChildren);
+		}
+	}
+
+private:
+	/**
+	 * Recursively checks for nodes in the tree which match the filter using depth first
+	 * traversal. Any nodes which match are cached, and all nodes leading to matches are expanded. 
+	 */
+	bool ExpandAndCacheMatches(ItemType Parent)
+	{
+		bool bChildMatches = false;
+		TArray<ItemType> AllChildren;
+		GetChildrenDelegate.Execute(Parent, AllChildren);
+		for (auto Child : AllChildren)
+		{
+			bChildMatches |= ExpandAndCacheMatches(Child);
+		}
+
+		bool bParentMatches = Filter->PassesFilter(Parent);
+
+		if (bChildMatches || bParentMatches)
+		{
+			TreeView->SetItemExpansion(Parent, true);
+			ItemsPassingFilter.Add(Parent);
+			return true;
+		}
+
+		return false;
+	}
+
+private:
+	/** The tree view to be filtered. */
+	STreeView<ItemType>* TreeView;
+
+	/** The filter to use for filtering. */
+	IFilter<ItemType>* Filter;
+
+	/** The source collection of root items for the tree. */
+	TArray<ItemType>* RootItems;
+
+	/** The collection of root items which is being displayed by the tree. */
+	TArray<ItemType>* TreeRootItems;
+
+	/** A delegate to get the children for a node in the tree. */
+	FOnGetChildren GetChildrenDelegate;
+
+	/** A cache of the items which passed the filter. */
+	TSet<ItemType> ItemsPassingFilter;
+
+	/** A set of the items which were expanded in the tree before any filtering. */
+	TSet<ItemType> PrefilterExpandedItems;
+
+	/** Whether or not the collection will be filtered on refresh. */
+	bool bIsEnabled;
+
+	/** Whether or not the tree was filtered the last time it was refreshed. */
+	bool bWasEnabled;
+};
\ No newline at end of file
