UMG - Removing dead code.  Cleaning up lots of prototype things that are not useful right now.  User widget blueprints now get placed into their own category in the palette.

diff --git a/Engine/Source/Editor/UMGEditor/Classes/WidgetGraphNode_Base.h b/Engine/Source/Editor/UMGEditor/Classes/WidgetGraphNode_Base.h
deleted file mode 100644
diff --git a/Engine/Source/Editor/UMGEditor/Classes/WidgetGraphNode_Text.h b/Engine/Source/Editor/UMGEditor/Classes/WidgetGraphNode_Text.h
deleted file mode 100644
diff --git a/Engine/Source/Editor/UMGEditor/Private/Palette/SPaletteView.cpp b/Engine/Source/Editor/UMGEditor/Private/Palette/SPaletteView.cpp
--- a/Engine/Source/Editor/UMGEditor/Private/Palette/SPaletteView.cpp
+++ b/Engine/Source/Editor/UMGEditor/Private/Palette/SPaletteView.cpp
@@ -15,8 +15,6 @@
 
 #include "WidgetTemplate.h"
 #include "WidgetTemplateClass.h"
-#include "WidgetTemplateButton.h"
-#include "WidgetTemplateCheckBox.h"
 
 #define LOCTEXT_NAMESPACE "UMG"
 
@@ -310,10 +308,10 @@ void SPaletteView::BuildClassWidgetList()
 
 void SPaletteView::BuildSpecialWidgetList()
 {
-	AddWidgetTemplate(MakeShareable(new FWidgetTemplateButton()));
-	AddWidgetTemplate(MakeShareable(new FWidgetTemplateCheckBox()));
+	//AddWidgetTemplate(MakeShareable(new FWidgetTemplateButton()));
+	//AddWidgetTemplate(MakeShareable(new FWidgetTemplateCheckBox()));
 
-	//TODO Make this pluggable.
+	//TODO UMG Make this pluggable.
 }
 
 void SPaletteView::AddWidgetTemplate(TSharedPtr<FWidgetTemplate> Template)
diff --git a/Engine/Source/Editor/UMGEditor/Private/Templates/WidgetTemplateButton.cpp b/Engine/Source/Editor/UMGEditor/Private/Templates/WidgetTemplateButton.cpp
deleted file mode 100644
diff --git a/Engine/Source/Editor/UMGEditor/Private/Templates/WidgetTemplateButton.h b/Engine/Source/Editor/UMGEditor/Private/Templates/WidgetTemplateButton.h
deleted file mode 100644
diff --git a/Engine/Source/Editor/UMGEditor/Private/Templates/WidgetTemplateCheckBox.cpp b/Engine/Source/Editor/UMGEditor/Private/Templates/WidgetTemplateCheckBox.cpp
deleted file mode 100644
diff --git a/Engine/Source/Editor/UMGEditor/Private/Templates/WidgetTemplateCheckBox.h b/Engine/Source/Editor/UMGEditor/Private/Templates/WidgetTemplateCheckBox.h
deleted file mode 100644
diff --git a/Engine/Source/Editor/UMGEditor/Private/Templates/WidgetTemplateClass.cpp b/Engine/Source/Editor/UMGEditor/Private/Templates/WidgetTemplateClass.cpp
--- a/Engine/Source/Editor/UMGEditor/Private/Templates/WidgetTemplateClass.cpp
+++ b/Engine/Source/Editor/UMGEditor/Private/Templates/WidgetTemplateClass.cpp
@@ -13,14 +13,26 @@ FWidgetTemplateClass::FWidgetTemplateClass(TSubclassOf<UWidget> InWidgetClass)
 
 FText FWidgetTemplateClass::GetCategory()
 {
-	const FString& Category = WidgetClass->GetMetaData("Category");
-
 	if ( Category.IsEmpty() )
 	{
-		return LOCTEXT("Misc", "Misc");
+		const FString& MetadatCategory = WidgetClass->GetMetaData("Category");
+
+		if ( MetadatCategory.IsEmpty() )
+		{
+			if ( WidgetClass->IsChildOf(UUserWidget::StaticClass()) )
+			{
+				return LOCTEXT("UserControls", "User Controls");
+			}
+
+			return LOCTEXT("Misc", "Misc");
+		}
+		else
+		{
+			Category = FText::FromString(MetadatCategory);
+		}
 	}
 
-	return FText::FromString(Category);
+	return Category;
 }
 
 UWidget* FWidgetTemplateClass::Create(UWidgetTree* Tree)
diff --git a/Engine/Source/Editor/UMGEditor/Private/Templates/WidgetTemplateClass.h b/Engine/Source/Editor/UMGEditor/Private/Templates/WidgetTemplateClass.h
--- a/Engine/Source/Editor/UMGEditor/Private/Templates/WidgetTemplateClass.h
+++ b/Engine/Source/Editor/UMGEditor/Private/Templates/WidgetTemplateClass.h
@@ -4,17 +4,28 @@
 
 #include "WidgetTemplate.h"
 
+/**
+ * A template that can spawn any widget derived from the UWidget class.
+ */
 class FWidgetTemplateClass : public FWidgetTemplate
 {
 public:
+	/** Constructor */
 	FWidgetTemplateClass(TSubclassOf<UWidget> InWidgetClass);
 
-	virtual FText GetCategory() override;
+	/** Gets the category for the widget */
+	virtual FText GetCategory() const override;
 
+	/** Creates an instance of the widget for the tree. */
 	virtual UWidget* Create(UWidgetTree* Tree) override;
 
+	/** The icon coming from the default object of the class */
 	virtual const FSlateBrush* GetIcon() const override;
 
 protected:
+	/** The widget class that will be created by this template */
 	TSubclassOf<UWidget> WidgetClass;
+
+	/** The cached category for the widget. */
+	mutable FText Category;
 };
diff --git a/Engine/Source/Editor/UMGEditor/Private/WidgetBlueprintCompiler.cpp b/Engine/Source/Editor/UMGEditor/Private/WidgetBlueprintCompiler.cpp
--- a/Engine/Source/Editor/UMGEditor/Private/WidgetBlueprintCompiler.cpp
+++ b/Engine/Source/Editor/UMGEditor/Private/WidgetBlueprintCompiler.cpp
@@ -11,47 +11,6 @@
 #include "BlueprintPropertyNodeSpawner.h"
 #include "K2Node_FunctionEntry.h"
 
-template<>
-FString FNetNameMapping::MakeBaseName<UWidgetGraphNode_Base>(const UWidgetGraphNode_Base* Net)
-{
-	return FString::Printf(TEXT("%s_%s"), *Net->GetDescriptiveCompiledName(), *Net->NodeGuid.ToString());
-}
-
-///-------------------------------------------------------------
-
-FWidgetBlueprintCompiler::FEvaluationHandlerRecord::FEvaluationHandlerRecord()
-	: NodeVariableProperty(NULL)
-	, EvaluationHandlerProperty(NULL)
-	, HandlerFunctionName(NAME_None)
-{
-}
-
-bool FWidgetBlueprintCompiler::FEvaluationHandlerRecord::IsValid() const
-{
-	return ( NodeVariableProperty != NULL ) && ( EvaluationHandlerProperty != NULL );
-}
-
-void FWidgetBlueprintCompiler::FEvaluationHandlerRecord::PatchFunctionNameInto(UObject* TargetObject) const
-{
-	EvaluationHandlerProperty->ContainerPtrToValuePtr<FExposedValueHandler>(
-		NodeVariableProperty->ContainerPtrToValuePtr<void>(TargetObject)
-		)->BoundFunction = HandlerFunctionName;
-}
-
-void FWidgetBlueprintCompiler::FEvaluationHandlerRecord::RegisterPin(UEdGraphPin* SourcePin, UProperty* AssociatedProperty, int32 AssociatedPropertyArrayIndex)
-{
-	FAnimNodeSinglePropertyHandler& Handler = ServicedProperties.FindOrAdd(AssociatedProperty->GetFName());
-
-	if ( AssociatedPropertyArrayIndex != INDEX_NONE )
-	{
-		Handler.ArrayPins.Add(AssociatedPropertyArrayIndex, SourcePin);
-	}
-	else
-	{
-		check(Handler.SinglePin == NULL);
-		Handler.SinglePin = SourcePin;
-	}
-}
 
 ///-------------------------------------------------------------
 
@@ -67,18 +26,6 @@ FWidgetBlueprintCompiler::~FWidgetBlueprintCompiler()
 void FWidgetBlueprintCompiler::MergeUbergraphPagesIn(UEdGraph* Ubergraph)
 {
 	Super::MergeUbergraphPagesIn(Ubergraph);
-
-	// Build the raw node list
-	TArray<UWidgetGraphNode_Base*> WidgetNodeList;
-	ConsolidatedEventGraph->GetNodesOfClass<UWidgetGraphNode_Base>(/*out*/ WidgetNodeList);
-
-	AllocateNodeIndexCounter = 0;
-
-	// Process the remaining nodes
-	for ( UWidgetGraphNode_Base* VisualWidgetNode : WidgetNodeList )
-	{
-		ProcessWidgetNode(VisualWidgetNode);
-	}
 }
 
 void FWidgetBlueprintCompiler::CreateFunctionList()
@@ -154,289 +101,6 @@ void FWidgetBlueprintCompiler::CreateFunctionList()
 	}
 }
 
-void FWidgetBlueprintCompiler::ProcessWidgetNode(UWidgetGraphNode_Base* VisualWidgetNode)
-{
-	const UWidgetGraphSchema* Schema = GetDefault<UWidgetGraphSchema>();
-
-	// Early out if this node has already been processed
-	//if ( AllocatedAnimNodes.Contains(VisualWidgetNode) )
-	//{
-	//	return;
-	//}
-
-	// Make sure the visual node has a runtime node template
-	UScriptStruct* NodeType = VisualWidgetNode->GetFNodeType();
-	if ( NodeType == NULL )
-	{
-		MessageLog.Error(TEXT("@@ has no widget value node member"), VisualWidgetNode);
-		return;
-	}
-
-	// Give the visual node a chance to do validation
-	{
-		const int32 PreValidationErrorCount = MessageLog.NumErrors;
-		//VisualWidgetNode->ValidateAnimNodeDuringCompilation(VisualWidgetNode->TargetSkeleton, MessageLog);
-		//VisualWidgetNode->BakeDataDuringCompilation(MessageLog);
-		if ( MessageLog.NumErrors != PreValidationErrorCount )
-		{
-			return;
-		}
-	}
-
-	// Create a property for the node
-	const FString NodeVariableName = ClassScopeNetNameMap.MakeValidName(VisualWidgetNode);
-
-	FEdGraphPinType NodeVariableType;
-	NodeVariableType.PinCategory = Schema->PC_Struct;
-	NodeVariableType.PinSubCategoryObject = NodeType;
-
-	UStructProperty* NewProperty = Cast<UStructProperty>(CreateVariable(FName(*NodeVariableName), NodeVariableType));
-
-	if ( NewProperty == NULL )
-	{
-		MessageLog.Error(TEXT("Failed to create node property for @@"), VisualWidgetNode);
-	}
-
-	// Register this node with the compile-time data structures
-	const int32 AllocatedIndex = AllocateNodeIndexCounter++;
-	AllocatedAnimNodes.Add(VisualWidgetNode, NewProperty);
-	AllocatedNodePropertiesToNodes.Add(NewProperty, VisualWidgetNode);
-	AllocatedAnimNodeIndices.Add(VisualWidgetNode, AllocatedIndex);
-	AllocatedPropertiesByIndex.Add(AllocatedIndex, NewProperty);
-
-	UWidgetGraphNode_Base* TrueSourceObject = MessageLog.FindSourceObjectTypeChecked<UWidgetGraphNode_Base>(VisualWidgetNode);
-	SourceNodeToProcessedNodeMap.Add(TrueSourceObject, VisualWidgetNode);
-
-
-	// TODO UMG add debug data.
-	// Register the slightly more permanent debug information
-//	NewAnimBlueprintClass->GetAnimBlueprintDebugData().NodePropertyToIndexMap.Add(TrueSourceObject, AllocatedIndex);
-//	NewAnimBlueprintClass->GetDebugData().RegisterClassPropertyAssociation(TrueSourceObject, NewProperty);
-
-	// Record pose pins for later patchup and gather pins that have an associated evaluation handler
-	TMap<FName, FEvaluationHandlerRecord> EvaluationHandlers;
-
-	for ( auto SourcePinIt = VisualWidgetNode->Pins.CreateIterator(); SourcePinIt; ++SourcePinIt )
-	{
-		UEdGraphPin* SourcePin = *SourcePinIt;
-		bool bConsumed = false;
-
-		// Does this pin have an associated evaluation handler?
-		UProperty* SourcePinProperty;
-		int32 SourceArrayIndex;
-		VisualWidgetNode->GetPinAssociatedProperty(NodeType, SourcePin, /*out*/ SourcePinProperty, /*out*/ SourceArrayIndex);
-
-		if ( SourcePinProperty != NULL )
-		{
-			if ( SourcePin->LinkedTo.Num() == 0 )
-			{
-				// Literal that can be pushed into the CDO instead of re-evaluated every frame
-				new (ValidAnimNodePinConstants)FEffectiveConstantRecord(NewProperty, SourcePin, SourcePinProperty, SourceArrayIndex);
-				bConsumed = true;
-			}
-			else
-			{
-				// Dynamic value that needs to be wired up and evaluated each frame
-				FString EvaluationHandlerStr = SourcePinProperty->GetMetaData(Schema->NAME_OnEvaluate);
-				FName EvaluationHandlerName(*EvaluationHandlerStr);
-				if ( EvaluationHandlerName == NAME_None )
-				{
-					EvaluationHandlerName = Schema->DefaultEvaluationHandlerName;
-				}
-
-				EvaluationHandlers.FindOrAdd(EvaluationHandlerName).RegisterPin(SourcePin, SourcePinProperty, SourceArrayIndex);
-				bConsumed = true;
-			}
-		}
-
-		if ( !bConsumed && ( SourcePin->Direction == EGPD_Input ) )
-		{
-			//@TODO: ANIMREFACTOR: It's probably OK to have certain pins ignored eventually, but this is very helpful during development
-			MessageLog.Note(TEXT("@@ was visible but ignored"), SourcePin);
-		}
-	}
-
-	// Match the associated property to each evaluation handler
-	for ( TFieldIterator<UProperty> NodePropIt(NodeType); NodePropIt; ++NodePropIt )
-	{
-		if ( UStructProperty* StructProp = Cast<UStructProperty>(*NodePropIt) )
-		{
-			if ( StructProp->Struct == FExposedWidgetValueHandler::StaticStruct() )
-			{
-				// Register this property to the list of pins that need to be updated
-				// (it's OK if there isn't an entry for this handler, it means that the values are static and don't need to be calculated every frame)
-				FName EvaluationHandlerName(StructProp->GetFName());
-				if ( FEvaluationHandlerRecord* pRecord = EvaluationHandlers.Find(EvaluationHandlerName) )
-				{
-					pRecord->NodeVariableProperty = NewProperty;
-					pRecord->EvaluationHandlerProperty = StructProp;
-				}
-			}
-		}
-	}
-
-	// Generate a new event to update the value of these properties
-	for ( auto HandlerIt = EvaluationHandlers.CreateIterator(); HandlerIt; ++HandlerIt )
-	{
-		FName EvaluationHandlerName = HandlerIt.Key();
-		FEvaluationHandlerRecord& Record = HandlerIt.Value();
-
-		if ( Record.IsValid() )
-		{
-			CreateEvaluationHandler(VisualWidgetNode, Record);
-			ValidEvaluationHandlerList.Add(Record);
-		}
-		else
-		{
-			MessageLog.Error(*FString::Printf(TEXT("A property on @@ references a non-existent %s property named %s"), *( Schema->NAME_OnEvaluate.ToString() ), *( EvaluationHandlerName.ToString() )), VisualWidgetNode );
-		}
-	}
-}
-
-void FWidgetBlueprintCompiler::CreateEvaluationHandler(UWidgetGraphNode_Base* VisualAnimNode, FEvaluationHandlerRecord& Record)
-{
-	// Shouldn't create a handler if there is nothing to work with
-	check(Record.ServicedProperties.Num() > 0);
-	check(Record.NodeVariableProperty != NULL);
-
-	//@TODO: Want to name these better
-	const FString FunctionName = FString::Printf(TEXT("%s_%s_%s"), *Record.EvaluationHandlerProperty->GetName(), *VisualAnimNode->GetOuter()->GetName(), *VisualAnimNode->GetName());
-	Record.HandlerFunctionName = FName(*FunctionName);
-
-	// Add a custom event in the graph
-	UK2Node_CustomEvent* EntryNode = SpawnIntermediateNode<UK2Node_CustomEvent>(VisualAnimNode, ConsolidatedEventGraph);
-	EntryNode->bInternalEvent = true;
-	EntryNode->CustomFunctionName = Record.HandlerFunctionName;
-	EntryNode->AllocateDefaultPins();
-
-	// The ExecChain is the current exec output pin in the linear chain
-	UEdGraphPin* ExecChain = Schema->FindExecutionPin(*EntryNode, EGPD_Output);
-
-	// Create a struct member write node to store the parameters into the animation node
-	UK2Node_StructMemberSet* AssignmentNode = SpawnIntermediateNode<UK2Node_StructMemberSet>(VisualAnimNode, ConsolidatedEventGraph);
-	AssignmentNode->VariableReference.SetSelfMember(Record.NodeVariableProperty->GetFName());
-	AssignmentNode->StructType = Record.NodeVariableProperty->Struct;
-	AssignmentNode->AllocateDefaultPins();
-
-	// Wire up the variable node execution wires
-	UEdGraphPin* ExecVariablesIn = Schema->FindExecutionPin(*AssignmentNode, EGPD_Input);
-	ExecChain->MakeLinkTo(ExecVariablesIn);
-	ExecChain = Schema->FindExecutionPin(*AssignmentNode, EGPD_Output);
-
-	// Run thru each property
-	TSet<FName> PropertiesBeingSet;
-
-	//UK2Node_FunctionResult* ReturnNode = SpawnIntermediateNode<UK2Node_FunctionResult>(VisualAnimNode, ConsolidatedEventGraph);
-	//ReturnNode->AllocateDefaultPins();
-
-	//// Wire up the variable node execution wires
-	//UEdGraphPin* ExecReturnIn = Schema->FindExecutionPin(*ReturnNode, EGPD_Input);
-	//ExecChain->MakeLinkTo(ExecReturnIn);
-	//ExecChain = Schema->FindExecutionPin(*ReturnNode, EGPD_Output);
-
-	for ( auto TargetPinIt = AssignmentNode->Pins.CreateIterator(); TargetPinIt; ++TargetPinIt )
-	{
-		UEdGraphPin* TargetPin = *TargetPinIt;
-		FString PropertyNameStr = TargetPin->PinName;
-		FName PropertyName(*PropertyNameStr);
-
-		// Does it get serviced by this handler?
-		if ( FAnimNodeSinglePropertyHandler* SourceInfo = Record.ServicedProperties.Find(PropertyName) )
-		{
-			if ( TargetPin->PinType.bIsArray )
-			{
-				// Grab the array that we need to set members for
-				UK2Node_StructMemberGet* FetchArrayNode = SpawnIntermediateNode<UK2Node_StructMemberGet>(VisualAnimNode, ConsolidatedEventGraph);
-				FetchArrayNode->VariableReference.SetSelfMember(Record.NodeVariableProperty->GetFName());
-				FetchArrayNode->StructType = Record.NodeVariableProperty->Struct;
-				FetchArrayNode->AllocatePinsForSingleMemberGet(PropertyName);
-
-				UEdGraphPin* ArrayVariableNode = FetchArrayNode->FindPin(PropertyNameStr);
-
-				if ( SourceInfo->ArrayPins.Num() > 0 )
-				{
-					// Set each element in the array
-					for ( auto SourcePinIt = SourceInfo->ArrayPins.CreateIterator(); SourcePinIt; ++SourcePinIt )
-					{
-						int32 ArrayIndex = SourcePinIt.Key();
-						UEdGraphPin* SourcePin = SourcePinIt.Value();
-
-						// Create an array element set node
-						UK2Node_CallArrayFunction* ArrayNode = SpawnIntermediateNode<UK2Node_CallArrayFunction>(VisualAnimNode, ConsolidatedEventGraph);
-						ArrayNode->FunctionReference.SetExternalMember(GET_FUNCTION_NAME_CHECKED(UKismetArrayLibrary, Array_Set), UKismetArrayLibrary::StaticClass());
-						ArrayNode->AllocateDefaultPins();
-
-						// Connect the execution chain
-						ExecChain->MakeLinkTo(ArrayNode->GetExecPin());
-						ExecChain = ArrayNode->GetThenPin();
-
-						// Connect the input array
-						UEdGraphPin* TargetArrayPin = ArrayNode->FindPinChecked(TEXT("TargetArray"));
-						TargetArrayPin->MakeLinkTo(ArrayVariableNode);
-						ArrayNode->PinConnectionListChanged(TargetArrayPin);
-
-						// Set the array index
-						UEdGraphPin* TargetIndexPin = ArrayNode->FindPinChecked(TEXT("Index"));
-						TargetIndexPin->DefaultValue = FString::FromInt(ArrayIndex);
-
-						// Wire up the data input
-						UEdGraphPin* TargetItemPin = ArrayNode->FindPinChecked(TEXT("Item"));
-						TargetItemPin->CopyPersistentDataFromOldPin(*SourcePin);
-						MessageLog.NotifyIntermediateObjectCreation(TargetItemPin, SourcePin);
-						SourcePin->BreakAllPinLinks();
-					}
-				}
-			}
-			else
-			{
-				// Single property
-				if ( SourceInfo->SinglePin != NULL )
-				{
-					UEdGraphPin* SourcePin = SourceInfo->SinglePin;
-
-					PropertiesBeingSet.Add(FName(*SourcePin->PinName));
-					TargetPin->CopyPersistentDataFromOldPin(*SourcePin);
-					MessageLog.NotifyIntermediateObjectCreation(TargetPin, SourcePin);
-					SourcePin->BreakAllPinLinks();
-
-
-					//// First, add this pin to the user-defined pins
-					//TSharedPtr<FUserPinInfo> NewPinInfo = MakeShareable(new FUserPinInfo());
-					//NewPinInfo->PinName = TEXT("return");
-					//NewPinInfo->PinType = SourcePin->PinType;
-
-					//TargetPin = ReturnNode->CreatePinFromUserDefinition(NewPinInfo);
-					//TargetPin->CopyPersistentDataFromOldPin(*SourcePin);
-					//MessageLog.NotifyIntermediateObjectCreation(TargetPin, SourcePin);
-
-					//SourcePin->BreakAllPinLinks();
-				}
-			}
-		}
-	}
-
-	// Remove any unused pins from the assignment node to avoid smashing constant values
-	for ( int32 PinIndex = 0; PinIndex < AssignmentNode->ShowPinForProperties.Num(); ++PinIndex )
-	{
-		FOptionalPinFromProperty& TestProperty = AssignmentNode->ShowPinForProperties[PinIndex];
-		TestProperty.bShowPin = PropertiesBeingSet.Contains(TestProperty.PropertyName);
-	}
-	AssignmentNode->ReconstructNode();
-
-}
-
-void FWidgetBlueprintCompiler::CopyTermDefaultsToDefaultObject(UObject* DefaultObject)
-{
-	Super::CopyTermDefaultsToDefaultObject(DefaultObject);
-
-	// And patch evaluation function entry names
-	for ( auto EvalLinkIt = ValidEvaluationHandlerList.CreateIterator(); EvalLinkIt; ++EvalLinkIt )
-	{
-		FEvaluationHandlerRecord& Record = *EvalLinkIt;
-		Record.PatchFunctionNameInto(DefaultObject);
-	}
-}
-
 void FWidgetBlueprintCompiler::ValidateWidgetNames()
 {
 	UWidgetBlueprint* Blueprint = WidgetBlueprint();
diff --git a/Engine/Source/Editor/UMGEditor/Private/WidgetBlueprintCompiler.h b/Engine/Source/Editor/UMGEditor/Private/WidgetBlueprintCompiler.h
--- a/Engine/Source/Editor/UMGEditor/Private/WidgetBlueprintCompiler.h
+++ b/Engine/Source/Editor/UMGEditor/Private/WidgetBlueprintCompiler.h
@@ -36,114 +36,14 @@ protected:
 	virtual void CreateClassVariablesFromBlueprint() override;
 	virtual void FinishCompilingClass(UClass* Class) override;
 	virtual bool ValidateGeneratedClass(UBlueprintGeneratedClass* Class) override;
-	virtual void CopyTermDefaultsToDefaultObject(UObject* DefaultObject) override;
 	// End FKismetCompilerContext
 
 protected:
-	// Wireup record for a single anim node property (which might be an array)
-	struct FAnimNodeSinglePropertyHandler
-	{
-		TMap<int32, UEdGraphPin*> ArrayPins;
-
-		UEdGraphPin* SinglePin;
-
-		FAnimNodeSinglePropertyHandler()
-			: SinglePin(NULL)
-		{
-		}
-	};
-
-	// Record for a property that was exposed as a pin, but wasn't wired up (just a literal)
-	struct FEffectiveConstantRecord
-	{
-	public:
-		// The node variable that the handler is in
-		class UStructProperty* NodeVariableProperty;
-
-		// The property within the struct to set
-		class UProperty* ConstantProperty;
-
-		// The array index if ConstantProperty is an array property, or INDEX_NONE otherwise
-		int32 ArrayIndex;
-
-		// The pin to pull the DefaultValue/DefaultObject from
-		UEdGraphPin* LiteralSourcePin;
-
-		FEffectiveConstantRecord()
-			: NodeVariableProperty(NULL)
-			, ConstantProperty(NULL)
-			, ArrayIndex(INDEX_NONE)
-			, LiteralSourcePin(NULL)
-		{
-		}
-
-		FEffectiveConstantRecord(UStructProperty* ContainingNodeProperty, UEdGraphPin* SourcePin, UProperty* SourcePinProperty, int32 SourceArrayIndex)
-			: NodeVariableProperty(ContainingNodeProperty)
-			, ConstantProperty(SourcePinProperty)
-			, ArrayIndex(SourceArrayIndex)
-			, LiteralSourcePin(SourcePin)
-		{
-		}
-
-		bool Apply(UObject* Object);
-	};
-
-	// Struct used to record information about one evaluation handler
-	struct FEvaluationHandlerRecord
-	{
-	public:
-		// The node variable that the handler is in
-		class UStructProperty* NodeVariableProperty;
-
-		// The specific evaluation handler inside the specified node
-		class UStructProperty* EvaluationHandlerProperty;
-
-		// Set of properties serviced by this handler (Map from property name to the record for that property)
-		TMap<FName, FAnimNodeSinglePropertyHandler> ServicedProperties;
-
-		// The resulting function name
-		FName HandlerFunctionName;
-	public:
-		FEvaluationHandlerRecord();
-
-		bool IsValid() const;
-
-		void PatchFunctionNameInto(UObject* TargetObject) const;
-
-		void RegisterPin(UEdGraphPin* SourcePin, UProperty* AssociatedProperty, int32 AssociatedPropertyArrayIndex);
-	};
-
-protected:
-	void ProcessWidgetNode(UWidgetGraphNode_Base* VisualWidgetNode);
-
-	void CreateEvaluationHandler(UWidgetGraphNode_Base* VisualAnimNode, FEvaluationHandlerRecord& Record);
-
 	UWidgetBlueprintGeneratedClass* NewWidgetBlueprintClass;
 
 	// Map of properties created for widgets; to aid in debug data generation
 	TMap<class UWidget*, class UProperty*> WidgetToMemberVariableMap;
 
-	///----------------------------------------------------------------
-	// Callback system properties.
-
-	// Map of allocated v3 nodes that are members of the class
-	TMap<class UWidgetGraphNode_Base*, UProperty*> AllocatedAnimNodes;
-	TMap<UProperty*, class UWidgetGraphNode_Base*> AllocatedNodePropertiesToNodes;
-	TMap<int32, UProperty*> AllocatedPropertiesByIndex;
-
-	// Map of true source objects (user edited ones) to the cloned ones that are actually compiled
-	TMap<class UWidgetGraphNode_Base*, UWidgetGraphNode_Base*> SourceNodeToProcessedNodeMap;
-
-	// Index of the nodes (must match up with the runtime discovery process of nodes, which runs thru the property chain)
-	int32 AllocateNodeIndexCounter;
-	TMap<class UWidgetGraphNode_Base*, int32> AllocatedAnimNodeIndices;
-
-	// List of successfully created evaluation handlers
-	TArray<FEvaluationHandlerRecord> ValidEvaluationHandlerList;
-
-	// List of animation node literals (values exposed as pins but never wired up) that need to be pushed into the CDO
-	TArray<FEffectiveConstantRecord> ValidAnimNodePinConstants;
-
 	///----------------------------------------------------------------
 };
 
diff --git a/Engine/Source/Editor/UMGEditor/Private/WidgetGraphNode_Base.cpp b/Engine/Source/Editor/UMGEditor/Private/WidgetGraphNode_Base.cpp
deleted file mode 100644
diff --git a/Engine/Source/Editor/UMGEditor/Private/WidgetGraphNode_Text.cpp b/Engine/Source/Editor/UMGEditor/Private/WidgetGraphNode_Text.cpp
deleted file mode 100644
diff --git a/Engine/Source/Editor/UMGEditor/Public/WidgetTemplate.h b/Engine/Source/Editor/UMGEditor/Public/WidgetTemplate.h
--- a/Engine/Source/Editor/UMGEditor/Public/WidgetTemplate.h
+++ b/Engine/Source/Editor/UMGEditor/Public/WidgetTemplate.h
@@ -13,7 +13,7 @@ public:
 	FWidgetTemplate();
 
 	/** The category this template fits into. */
-	virtual FText GetCategory() = 0;
+	virtual FText GetCategory() const = 0;
 
 	/** Constructs the widget template. */
 	virtual UWidget* Create(class UWidgetTree* Tree) = 0;
diff --git a/Engine/Source/Runtime/UMG/Private/UserWidget.cpp b/Engine/Source/Runtime/UMG/Private/UserWidget.cpp
--- a/Engine/Source/Runtime/UMG/Private/UserWidget.cpp
+++ b/Engine/Source/Runtime/UMG/Private/UserWidget.cpp
@@ -463,3 +463,12 @@ APlayerController* UUserWidget::GetPlayerController() const
 {
 	return PlayerContext.IsValid() ? PlayerContext.GetPlayerController() : NULL;
 }
+
+#if WITH_EDITOR
+
+const FSlateBrush* UUserWidget::GetEditorIcon()
+{
+	return FUMGStyle::Get().GetBrush("Widget.UserWidget");
+}
+
+#endif
\ No newline at end of file
diff --git a/Engine/Source/Runtime/UMG/Private/WidgetActor.cpp b/Engine/Source/Runtime/UMG/Private/WidgetActor.cpp
deleted file mode 100644
diff --git a/Engine/Source/Runtime/UMG/Private/WidgetBlueprintGeneratedClass.cpp b/Engine/Source/Runtime/UMG/Private/WidgetBlueprintGeneratedClass.cpp
--- a/Engine/Source/Runtime/UMG/Private/WidgetBlueprintGeneratedClass.cpp
+++ b/Engine/Source/Runtime/UMG/Private/WidgetBlueprintGeneratedClass.cpp
@@ -35,16 +35,16 @@ void UWidgetBlueprintGeneratedClass::Link(FArchive& Ar, bool bRelinkExistingProp
 	WidgetNodeProperties.Empty();
 
 	// Initialize derived members
-	for ( TFieldIterator<UProperty> It(this); It; ++It )
-	{
-		if ( UStructProperty* StructProp = Cast<UStructProperty>(*It) )
-		{
-			if ( StructProp->Struct->IsChildOf(FWidgetNode_Base::StaticStruct()) )
-			{
-				WidgetNodeProperties.Add(StructProp);
-			}
-		}
-	}
+	//for ( TFieldIterator<UProperty> It(this); It; ++It )
+	//{
+	//	if ( UStructProperty* StructProp = Cast<UStructProperty>(*It) )
+	//	{
+	//		if ( StructProp->Struct->IsChildOf(FWidgetNode_Base::StaticStruct()) )
+	//		{
+	//			WidgetNodeProperties.Add(StructProp);
+	//		}
+	//	}
+	//}
 }
 
 void UWidgetBlueprintGeneratedClass::InitializeWidget(UUserWidget* UserWidget) const
diff --git a/Engine/Source/Runtime/UMG/Public/Blueprint/UserWidget.h b/Engine/Source/Runtime/UMG/Public/Blueprint/UserWidget.h
--- a/Engine/Source/Runtime/UMG/Public/Blueprint/UserWidget.h
+++ b/Engine/Source/Runtime/UMG/Public/Blueprint/UserWidget.h
@@ -94,7 +94,7 @@ DECLARE_DYNAMIC_MULTICAST_DELEGATE_OneParam(FOnVisibilityChangedEvent, ESlateVis
 /**
  * The user widget is extensible by users through the WidgetBlueprint.
  */
-UCLASS(Abstract, editinlinenew, BlueprintType, Blueprintable)
+UCLASS(Abstract, editinlinenew, BlueprintType, Blueprintable, meta=( Category="User Controls" ))
 class UMG_API UUserWidget : public UWidget
 {
 	GENERATED_UCLASS_BODY()
@@ -251,6 +251,12 @@ public:
 	/** Ticks this widget and forwards to the underlying widget to tick */
 	void NativeTick(const FGeometry& MyGeometry, float InDeltaTime );
 
+#if WITH_EDITOR
+	// UWidget interface
+	virtual const FSlateBrush* GetEditorIcon() override;
+	// End UWidget interface
+#endif
+
 public:
 	/** Called when the visibility changes. */
 	UPROPERTY(BlueprintAssignable)
diff --git a/Engine/Source/Runtime/UMG/Public/Blueprint/WidgetActor.h b/Engine/Source/Runtime/UMG/Public/Blueprint/WidgetActor.h
deleted file mode 100644
diff --git a/Engine/Source/Runtime/UMG/Public/Blueprint/WidgetNode_Base.h b/Engine/Source/Runtime/UMG/Public/Blueprint/WidgetNode_Base.h
deleted file mode 100644
diff --git a/Engine/Source/Runtime/UMG/Public/Blueprint/WidgetNode_Text.h b/Engine/Source/Runtime/UMG/Public/Blueprint/WidgetNode_Text.h
deleted file mode 100644
diff --git a/Engine/Source/Runtime/UMG/Public/UMG.h b/Engine/Source/Runtime/UMG/Public/UMG.h
--- a/Engine/Source/Runtime/UMG/Public/UMG.h
+++ b/Engine/Source/Runtime/UMG/Public/UMG.h
@@ -92,12 +92,7 @@
 #include "WidgetTree.h"
 #include "UserWidget.h"
 
-#include "WidgetNode_Base.h"
-#include "WidgetNode_Text.h"
-
 #include "WidgetBlueprintGeneratedClass.h"
 
-#include "WidgetActor.h"
-
 // Interfaces
 #include "IUMGModule.h"
