Build - Fixing a build break

diff --git a/Engine/Source/Editor/UMGEditor/Private/Templates/WidgetTemplateClass.cpp b/Engine/Source/Editor/UMGEditor/Private/Templates/WidgetTemplateClass.cpp
--- a/Engine/Source/Editor/UMGEditor/Private/Templates/WidgetTemplateClass.cpp
+++ b/Engine/Source/Editor/UMGEditor/Private/Templates/WidgetTemplateClass.cpp
@@ -11,7 +11,7 @@ FWidgetTemplateClass::FWidgetTemplateClass(TSubclassOf<UWidget> InWidgetClass)
 {
 }
 
-FText FWidgetTemplateClass::GetCategory()
+FText FWidgetTemplateClass::GetCategory() const
 {
 	if ( Category.IsEmpty() )
 	{
