UMG - Fix missing font warning

diff --git a/Engine/Source/Editor/UMGEditor/Private/Hierarchy/SHierarchyViewItem.cpp b/Engine/Source/Editor/UMGEditor/Private/Hierarchy/SHierarchyViewItem.cpp
--- a/Engine/Source/Editor/UMGEditor/Private/Hierarchy/SHierarchyViewItem.cpp
+++ b/Engine/Source/Editor/UMGEditor/Private/Hierarchy/SHierarchyViewItem.cpp
@@ -172,7 +172,7 @@ FSlateFontInfo SHierarchyViewItem::GetItemFont() const
 	else
 	{
 		// TODO UMG Hacky move into style area
-		return FSlateFontInfo(FPaths::EngineContentDir() / TEXT("Slate/Fonts/Roboto-Normal.ttf"), 9);
+		return FSlateFontInfo(FPaths::EngineContentDir() / TEXT("Slate/Fonts/Roboto-Regular.ttf"), 9);
 	}
 }
 
