UMG Sequencer keying:
- Added ability to key into UMG animations via the details panel
- Fixed some undo/redo issues in sequencer related to autokey
- Fixed auto-key causing property migration

diff --git a/Engine/Source/Developer/MovieSceneTools/Private/PropertyTrackEditor.cpp b/Engine/Source/Developer/MovieSceneTools/Private/PropertyTrackEditor.cpp
--- a/Engine/Source/Developer/MovieSceneTools/Private/PropertyTrackEditor.cpp
+++ b/Engine/Source/Developer/MovieSceneTools/Private/PropertyTrackEditor.cpp
@@ -49,12 +49,18 @@ public:
 		float StartTime = ColorSection->GetStartTime();
 		float EndTime = ColorSection->GetEndTime();
 		float SectionDuration = EndTime - StartTime;
-		
+
 		if (!FMath::IsNearlyZero(SectionDuration))
 		{
+			LayerId = FPropertySection::OnPaintSection( AllottedGeometry, SectionClippingRect, OutDrawElements, LayerId, bParentEnabled );
+
+			FVector2D GradientSize = FVector2D(AllottedGeometry.Size.X, (AllottedGeometry.Size.Y / 4)-3.0f);
+
+			FPaintGeometry PaintGeometry = AllottedGeometry.ToPaintGeometry(FVector2D(0,0), GradientSize );
+
 			// If we are showing a background pattern and the colors is transparent, draw a checker pattern
 			const FSlateBrush* CheckerBrush = FEditorStyle::GetBrush( "Checker" );
-			FSlateDrawElement::MakeBox(OutDrawElements, LayerId, AllottedGeometry.ToPaintGeometry(), CheckerBrush, SectionClippingRect);
+			FSlateDrawElement::MakeBox(OutDrawElements, LayerId, PaintGeometry, CheckerBrush, SectionClippingRect);
 	
 			TArray<FSlateGradientStop> GradientStops;
 
@@ -76,7 +82,7 @@ public:
 				FSlateDrawElement::MakeGradient(
 					OutDrawElements,
 					LayerId + 1,
-					AllottedGeometry.ToPaintGeometry(),
+					PaintGeometry,
 					GradientStops,
 					Orient_Vertical,
 					SectionClippingRect
@@ -336,23 +342,22 @@ void FPropertyTrackEditor::AddKey(const FGuid& ObjectGuid, UObject* AdditionalAs
 }
 
 template <typename Type, typename TrackType>
-void FPropertyTrackEditor::OnAnimatedPropertyChanged( const TArray<UObject*>& InObjectsThatChanged, const IPropertyHandle& PropertyValue, bool bRequireAutoKey )
+void FPropertyTrackEditor::OnAnimatedPropertyChanged(const FKeyPropertyParams& KeyPropertyParams )
 {
-	FName PropertyName = PropertyValue.GetProperty()->GetFName();
-
 	// Get the value from the property 
 	Type Value;
-	FPropertyAccess::Result Result = PropertyValue.GetValue( Value );
+	FPropertyAccess::Result Result = KeyPropertyParams.PropertyHandle->GetValue( Value );
 	check( Result == FPropertyAccess::Success );
 	
-	AnimatablePropertyChanged(TrackType::StaticClass(), bRequireAutoKey,
-		FOnKeyProperty::CreateRaw(this, &FPropertyTrackEditor::OnKeyProperty<Type, TrackType>, PropertyName, InObjectsThatChanged, (const Type*)&Value, bRequireAutoKey));
+	AnimatablePropertyChanged(TrackType::StaticClass(), KeyPropertyParams.bRequireAutoKey,
+		FOnKeyProperty::CreateRaw(this, &FPropertyTrackEditor::OnKeyProperty<Type, TrackType>, KeyPropertyParams, Value) );
 }
 
-void FPropertyTrackEditor::OnAnimatedVectorPropertyChanged( const TArray<UObject*>& InObjectsThatChanged, const IPropertyHandle& PropertyValue, bool bRequireAutoKey )
+
+void FPropertyTrackEditor::OnAnimatedVectorPropertyChanged(const FKeyPropertyParams& KeyPropertyParams)
 {
 	bool bIsVector2D = false, bIsVector = false, bIsVector4 = false;
-	const UStructProperty* StructProp = Cast<const UStructProperty>(PropertyValue.GetProperty());
+	const UStructProperty* StructProp = Cast<const UStructProperty>(KeyPropertyParams.PropertyHandle->GetProperty());
 	if (StructProp && StructProp->Struct)
 	{
 		FName StructName = StructProp->Struct->GetFName();
@@ -361,44 +366,63 @@ void FPropertyTrackEditor::OnAnimatedVectorPropertyChanged( const TArray<UObject
 		bIsVector = StructName == NAME_Vector;
 		bIsVector4 = StructName == NAME_Vector4;
 	}
-	if (!bIsVector2D && !bIsVector && !bIsVector4) {return;}
+	if (!bIsVector2D && !bIsVector && !bIsVector4) 
+	{
+		return;
+	}
 
 
-	FName PropertyName = PropertyValue.GetProperty()->GetFName();
 
 	// Get the vector value from the property
 	FPropertyAccess::Result Result = FPropertyAccess::Fail;
+	const bool bRequiresAutoKey = KeyPropertyParams.bRequireAutoKey;
+
+
 	if (bIsVector2D)
 	{
-		FVector2D Value;
-		Result = PropertyValue.GetValue( Value );
+		FVectorKey<FVector2D> NewKey;
+		NewKey.CurveName = KeyPropertyParams.InnerStructPropertyName;
+		NewKey.bAddKeyEvenIfUnchanged = !bRequiresAutoKey;
+
+		Result = KeyPropertyParams.PropertyHandle->GetValue( NewKey.Value );
+
+		FOnKeyProperty OnKeyPropertyDelegate = FOnKeyProperty::CreateRaw( this, &FPropertyTrackEditor::OnKeyProperty<FVectorKey<FVector2D>, UMovieSceneVectorTrack>, KeyPropertyParams, NewKey );
+
 		check( Result == FPropertyAccess::Success );
-		AnimatablePropertyChanged(UMovieSceneVectorTrack::StaticClass(), bRequireAutoKey,
-			FOnKeyProperty::CreateRaw(this, &FPropertyTrackEditor::OnKeyProperty<FVector2D, UMovieSceneVectorTrack>, PropertyName, InObjectsThatChanged, (const FVector2D*)&Value, bRequireAutoKey));
+		AnimatablePropertyChanged(UMovieSceneVectorTrack::StaticClass(), bRequiresAutoKey, OnKeyPropertyDelegate );
 	}
 	else if (bIsVector)
 	{
-		FVector Value;
-		Result = PropertyValue.GetValue( Value );
+		FVectorKey<FVector> NewKey;
+		NewKey.CurveName = KeyPropertyParams.InnerStructPropertyName;
+		NewKey.bAddKeyEvenIfUnchanged = !bRequiresAutoKey;
+
+		Result = KeyPropertyParams.PropertyHandle->GetValue(NewKey.Value);
+
+		FOnKeyProperty OnKeyPropertyDelegate = FOnKeyProperty::CreateRaw(this, &FPropertyTrackEditor::OnKeyProperty<FVectorKey<FVector>, UMovieSceneVectorTrack>, KeyPropertyParams, NewKey );
+
 		check( Result == FPropertyAccess::Success );
-		AnimatablePropertyChanged(UMovieSceneVectorTrack::StaticClass(), bRequireAutoKey,
-			FOnKeyProperty::CreateRaw(this, &FPropertyTrackEditor::OnKeyProperty<FVector, UMovieSceneVectorTrack>, PropertyName, InObjectsThatChanged, (const FVector*)&Value, bRequireAutoKey));
+		AnimatablePropertyChanged( UMovieSceneVectorTrack::StaticClass(),bRequiresAutoKey, OnKeyPropertyDelegate );
 	}
 	else if (bIsVector4)
 	{
-		FVector4 Value;
-		Result = PropertyValue.GetValue( Value );
+		FVectorKey<FVector4> NewKey;
+		NewKey.CurveName = KeyPropertyParams.InnerStructPropertyName;
+		NewKey.bAddKeyEvenIfUnchanged = !bRequiresAutoKey;
+
+		Result = KeyPropertyParams.PropertyHandle->GetValue(NewKey.Value);
+
+		FOnKeyProperty OnKeyPropertyDelegate = FOnKeyProperty::CreateRaw(this, &FPropertyTrackEditor::OnKeyProperty<FVectorKey<FVector4>, UMovieSceneVectorTrack>, KeyPropertyParams, NewKey );
+
 		check( Result == FPropertyAccess::Success );
-		AnimatablePropertyChanged(UMovieSceneVectorTrack::StaticClass(), bRequireAutoKey,
-			FOnKeyProperty::CreateRaw(this, &FPropertyTrackEditor::OnKeyProperty<FVector4, UMovieSceneVectorTrack>, PropertyName, InObjectsThatChanged, (const FVector4*)&Value, bRequireAutoKey));
+		AnimatablePropertyChanged(UMovieSceneVectorTrack::StaticClass(),bRequiresAutoKey, OnKeyPropertyDelegate );
 	}
-	else {check(0);}
 }
 
-void FPropertyTrackEditor::OnAnimatedColorPropertyChanged( const TArray<UObject*>& InObjectsThatChanged, const IPropertyHandle& PropertyValue, bool bRequireAutoKey )
+void FPropertyTrackEditor::OnAnimatedColorPropertyChanged(const FKeyPropertyParams& KeyPropertyParams )
 {
 	bool bIsFColor = false, bIsFLinearColor = false;
-	const UStructProperty* StructProp = Cast<const UStructProperty>(PropertyValue.GetProperty());
+	const UStructProperty* StructProp = Cast<const UStructProperty>(KeyPropertyParams.PropertyHandle->GetProperty());
 	if (StructProp && StructProp->Struct)
 	{
 		FName StructName = StructProp->Struct->GetFName();
@@ -411,15 +435,15 @@ void FPropertyTrackEditor::OnAnimatedColorPropertyChanged( const TArray<UObject*
 	check(bIsFColor ^ bIsFLinearColor);
 
 
-	FName PropertyName = PropertyValue.GetProperty()->GetFName();
+	FName PropertyName = KeyPropertyParams.PropertyHandle->GetProperty()->GetFName();
 					
-	const UClass* PropertyClass = PropertyValue.GetPropertyClass();
+	const UClass* PropertyClass = KeyPropertyParams.PropertyHandle->GetPropertyClass();
 
 	TSharedPtr<IPropertyHandle> ChildProperties[4] = {
-		PropertyValue.GetChildHandle("R"),
-		PropertyValue.GetChildHandle("G"),
-		PropertyValue.GetChildHandle("B"),
-		PropertyValue.GetChildHandle("A")
+		KeyPropertyParams.PropertyHandle->GetChildHandle("R"),
+		KeyPropertyParams.PropertyHandle->GetChildHandle("G"),
+		KeyPropertyParams.PropertyHandle->GetChildHandle("B"),
+		KeyPropertyParams.PropertyHandle->GetChildHandle("A")
 	};
 
 	uint8 FColorChannels[4];
@@ -441,35 +465,41 @@ void FPropertyTrackEditor::OnAnimatedColorPropertyChanged( const TArray<UObject*
 	{
 		ColorValue = FLinearColor(FLinearColorChannels[0], FLinearColorChannels[1], FLinearColorChannels[2], FLinearColorChannels[3]);
 	}
-	
+
 	if( StructProp->HasMetaData("HideAlphaChannel") )
 	{
 		ColorValue.A = 1;
 	}
 
-	AnimatablePropertyChanged(UMovieSceneColorTrack::StaticClass(), bRequireAutoKey,
-		FOnKeyProperty::CreateRaw(this, &FPropertyTrackEditor::OnKeyProperty<FLinearColor, UMovieSceneColorTrack>, PropertyName, InObjectsThatChanged, (const FLinearColor*)&ColorValue, bRequireAutoKey));
+
+	FColorKey Key;
+	Key.Value = ColorValue;
+	Key.CurveName = KeyPropertyParams.InnerStructPropertyName;
+	Key.bAddKeyEvenIfUnchanged = !KeyPropertyParams.bRequireAutoKey;
+
+	AnimatablePropertyChanged(UMovieSceneColorTrack::StaticClass(), KeyPropertyParams.bRequireAutoKey, FOnKeyProperty::CreateRaw(this, &FPropertyTrackEditor::OnKeyProperty<FColorKey, UMovieSceneColorTrack>, KeyPropertyParams, Key) );
 }
 
 
 template <typename Type, typename TrackType>
-void FPropertyTrackEditor::OnKeyProperty( float KeyTime, FName PropertyName, const TArray<UObject*> InObjectsThatChanged, const Type* Value, bool bAutoKeying )
+void FPropertyTrackEditor::OnKeyProperty( float KeyTime, FKeyPropertyParams KeyPropertyParams, Type Value )
 {
-	for( int32 ObjectIndex = 0; ObjectIndex < InObjectsThatChanged.Num(); ++ObjectIndex )
+	for( UObject* Object : KeyPropertyParams.ObjectsThatChanged )
 	{
-		UObject* Object = InObjectsThatChanged[ObjectIndex];
-
 		FGuid ObjectHandle = FindOrCreateHandleToObject( Object );
 		if (ObjectHandle.IsValid())
 		{
+			FName PropertyName = KeyPropertyParams.PropertyHandle->GetProperty()->GetFName();
+
 			UMovieSceneTrack* Track = GetTrackForObject( ObjectHandle, TrackType::StaticClass(), PropertyName );
 			if( ensure( Track ) )
 			{
 				TrackType* TypedTrack = CastChecked<TrackType>(Track);
+
 				TypedTrack->SetPropertyName( PropertyName );
 				// Find or add a new section at the auto-key time and changing the property same property
 				// AddKeyToSection is not actually a virtual, it's redefined in each class with a different type
-				bool bSuccessfulAdd = TypedTrack->AddKeyToSection( KeyTime, *Value );
+				bool bSuccessfulAdd = TypedTrack->AddKeyToSection( KeyTime, Value );
 				if (bSuccessfulAdd)
 				{
 					TypedTrack->SetAsShowable();
diff --git a/Engine/Source/Developer/MovieSceneTools/Private/PropertyTrackEditor.h b/Engine/Source/Developer/MovieSceneTools/Private/PropertyTrackEditor.h
--- a/Engine/Source/Developer/MovieSceneTools/Private/PropertyTrackEditor.h
+++ b/Engine/Source/Developer/MovieSceneTools/Private/PropertyTrackEditor.h
@@ -4,7 +4,6 @@
 
 class IPropertyHandle;
 
-
 /**
  * Tools for animatable property types such as floats ands vectors
  */
@@ -40,7 +39,7 @@ private:
 	 * @param PropertyValue			Handle to the property value which changed
 	 */
 	template <typename Type, typename TrackType>
-	void OnAnimatedPropertyChanged( const TArray<UObject*>& InObjectsThatChanged, const IPropertyHandle& PropertyValue, bool bRequireAutoKey );
+	void OnAnimatedPropertyChanged( const struct FKeyPropertyParams& KeyPropertyParams );
 
 	/**
 	 * Called by the details panel when an animatable vector property changes
@@ -48,7 +47,7 @@ private:
 	 * @param InObjectsThatChanged	List of objects that changed
 	 * @param PropertyValue			Handle to the property value which changed
 	 */
-	void OnAnimatedVectorPropertyChanged( const TArray<UObject*>& InObjectsThatChanged, const IPropertyHandle& PropertyValue, bool bRequireAutoKey );
+	void OnAnimatedVectorPropertyChanged( const struct FKeyPropertyParams& KeyPropertyParams );
 	
 	/**
 	 * Called by the details panel when an animatable color property changes
@@ -56,11 +55,11 @@ private:
 	 * @param InObjectsThatChanged	List of objects that changed
 	 * @param PropertyValue			Handle to the property value which changed
 	 */
-	void OnAnimatedColorPropertyChanged( const TArray<UObject*>& InObjectsThatChanged, const IPropertyHandle& PropertyValue, bool bRequireAutoKey );
+	void OnAnimatedColorPropertyChanged( const struct FKeyPropertyParams& KeyPropertyParams );
 	
 	/** Delegates for animatable property changed functions */
 	template <typename Type, typename TrackType>
-	void OnKeyProperty( float KeyTime, FName PropertyName, const TArray<UObject*> InObjectsThatChanged, const Type* Value, bool bAutoKeying);
+	void OnKeyProperty( float KeyTime, FKeyPropertyParams KeyPropertyParams, Type Value );
 };
 
 
diff --git a/Engine/Source/Developer/MovieSceneTools/Private/TransformTrackEditor.cpp b/Engine/Source/Developer/MovieSceneTools/Private/TransformTrackEditor.cpp
--- a/Engine/Source/Developer/MovieSceneTools/Private/TransformTrackEditor.cpp
+++ b/Engine/Source/Developer/MovieSceneTools/Private/TransformTrackEditor.cpp
@@ -188,7 +188,7 @@ void FTransformTrackEditor::OnTransformChanged( UObject& InObject )
 
 	}
 
-	if( SceneComponentThatChanged && ensure( ObjectHandle.IsValid() ) )
+	if( SceneComponentThatChanged && ObjectHandle.IsValid() )
 	{
 		// Find an existing transform if possible.  If one exists we will compare against the new one to decide what components of the transform need keys
 		FTransformData ExistingTransform = ObjectToExistingTransform.FindRef( &InObject );
diff --git a/Engine/Source/Editor/EditorStyle/Private/SlateEditorStyle.cpp b/Engine/Source/Editor/EditorStyle/Private/SlateEditorStyle.cpp
--- a/Engine/Source/Editor/EditorStyle/Private/SlateEditorStyle.cpp
+++ b/Engine/Source/Editor/EditorStyle/Private/SlateEditorStyle.cpp
@@ -1324,7 +1324,7 @@ void FSlateEditorStyle::FStyle::SetupGeneralStyles()
 		FLinearColor EditorOrange = FLinearColor(0.728f, 0.364f, 0.003f);
 
 		FTextBlockStyle CreditsNormal = FTextBlockStyle(NormalText)
-			.SetFont(TTF_CORE_FONT("Fonts/Roboto-Normal", 16))
+			.SetFont(TTF_CORE_FONT("Fonts/Roboto-Regular", 16))
 			.SetShadowOffset(FVector2D::UnitVector);
 
 		Set("Credits.Normal", CreditsNormal);
@@ -1403,6 +1403,16 @@ void FSlateEditorStyle::FStyle::SetupGeneralStyles()
 		Set( "Sequencer.AnimationOutliner.BoldFont", TTF_CORE_FONT( "Fonts/Roboto-Bold", 11 ) );
 		Set( "Sequencer.AnimationOutliner.RegularFont", TTF_CORE_FONT( "Fonts/Roboto-Regular", 8 ) );
 		Set( "Sequencer.ShotFilter", new IMAGE_BRUSH( "Sequencer/FilteredArea", FVector2D(74,74), FLinearColor::White, ESlateBrushTileType::Both ) );
+		Set( "Sequencer.KeyMark", new IMAGE_BRUSH("Sequencer/KeyMark", FVector2D(3,21), FLinearColor::White, ESlateBrushTileType::NoTile ) );
+
+		const FButtonStyle DetailsKeyButton = FButtonStyle(NoBorder)
+			.SetNormal( IMAGE_BRUSH("Sequencer/AddKey_Details", FVector2D(11,11) )  )
+			.SetHovered( IMAGE_BRUSH("Sequencer/AddKey_Details", FVector2D(11,11), SelectionColor ) )
+			.SetPressed( IMAGE_BRUSH("Sequencer/AddKey_Details", FVector2D(11,11), SelectionColor_Pressed ) )
+			.SetNormalPadding(FMargin(0, 1))
+			.SetPressedPadding(FMargin(0, 2, 0, 0));
+
+		Set( "Sequencer.AddKey.Details", DetailsKeyButton );
 	}
 
 	// Foliage Edit Mode
diff --git a/Engine/Source/Editor/PropertyEditor/Private/IDetailsViewPrivate.h b/Engine/Source/Editor/PropertyEditor/Private/IDetailsViewPrivate.h
--- a/Engine/Source/Editor/PropertyEditor/Private/IDetailsViewPrivate.h
+++ b/Engine/Source/Editor/PropertyEditor/Private/IDetailsViewPrivate.h
@@ -77,6 +77,8 @@ public:
 	 */
 	virtual bool IsPropertyVisible( const struct FPropertyAndParent& PropertyAndParent ) const = 0;
 
+	virtual TSharedPtr<IDetailKeyframeHandler> GetKeyframeHandler() = 0;
+
 	/**
 	 * @return The thumbnail pool that should be used for thumbnails being rendered in this view
 	 */
diff --git a/Engine/Source/Editor/PropertyEditor/Private/PropertyChangeListener.cpp b/Engine/Source/Editor/PropertyEditor/Private/PropertyChangeListener.cpp
--- a/Engine/Source/Editor/PropertyEditor/Private/PropertyChangeListener.cpp
+++ b/Engine/Source/Editor/PropertyEditor/Private/PropertyChangeListener.cpp
@@ -212,7 +212,7 @@ bool FPropertyChangeListener::ScanForChanges( bool bRecacheNewValues )
 			{
 				bChangesFound = true;
 				TSharedPtr<IPropertyHandle> Handle = PropertyEditorHelpers::GetPropertyHandle( PropertyNode.ToSharedRef(), NULL, NULL );
-				OnPropertyChangedDelegate.Broadcast( ObjectsThatChanged, Handle.ToSharedRef() );
+				OnPropertyChangedDelegate.Broadcast( ObjectsThatChanged, *Handle );
 			}
 		}
 		
@@ -246,7 +246,7 @@ void FPropertyChangeListener::TriggerAllPropertiesChangedDelegate()
 		if( ObjectsThatChanged.Num() > 0 )
 		{
 			TSharedPtr<IPropertyHandle> Handle = PropertyEditorHelpers::GetPropertyHandle( PropertyNode.ToSharedRef(), NULL, NULL );
-			OnPropertyChangedDelegate.Broadcast( ObjectsThatChanged, Handle.ToSharedRef() );
+			OnPropertyChangedDelegate.Broadcast( ObjectsThatChanged, *Handle );
 		}
 	}
 }
diff --git a/Engine/Source/Editor/PropertyEditor/Private/PropertyHandleImpl.cpp b/Engine/Source/Editor/PropertyEditor/Private/PropertyHandleImpl.cpp
--- a/Engine/Source/Editor/PropertyEditor/Private/PropertyHandleImpl.cpp
+++ b/Engine/Source/Editor/PropertyEditor/Private/PropertyHandleImpl.cpp
@@ -1636,7 +1636,7 @@ uint32 FPropertyHandleBase::GetNumOuterObjects() const
 	return NumObjects;
 }
 
-void FPropertyHandleBase::GetOuterObjects( TArray<UObject*>& OuterObjects )
+void FPropertyHandleBase::GetOuterObjects( TArray<UObject*>& OuterObjects ) const
 {
 	FObjectPropertyNode* ObjectNode = Implementation->GetPropertyNode()->FindObjectItemParent();
 	if( ObjectNode )
diff --git a/Engine/Source/Editor/PropertyEditor/Private/PropertyHandleImpl.h b/Engine/Source/Editor/PropertyEditor/Private/PropertyHandleImpl.h
--- a/Engine/Source/Editor/PropertyEditor/Private/PropertyHandleImpl.h
+++ b/Engine/Source/Editor/PropertyEditor/Private/PropertyHandleImpl.h
@@ -389,7 +389,7 @@ public:
 	virtual void AccessRawData( TArray<void*>& RawData ) override;
 	virtual void AccessRawData( TArray<const void*>& RawData ) const override;
 	virtual uint32 GetNumOuterObjects() const override;
-	virtual void GetOuterObjects( TArray<UObject*>& OuterObjects ) override;
+	virtual void GetOuterObjects( TArray<UObject*>& OuterObjects ) const override;
 	virtual FPropertyAccess::Result GetNumChildren( uint32& OutNumChildren ) const override;
 	virtual TSharedPtr<IPropertyHandleArray> AsArray() override { return NULL; }
 	virtual const UClass* GetPropertyClass() const override;
diff --git a/Engine/Source/Editor/PropertyEditor/Private/SDetailSingleItemRow.cpp b/Engine/Source/Editor/PropertyEditor/Private/SDetailSingleItemRow.cpp
--- a/Engine/Source/Editor/PropertyEditor/Private/SDetailSingleItemRow.cpp
+++ b/Engine/Source/Editor/PropertyEditor/Private/SDetailSingleItemRow.cpp
@@ -4,6 +4,7 @@
 #include "SDetailSingleItemRow.h"
 #include "DetailItemNode.h"
 #include "PropertyEditorHelpers.h"
+#include "IDetailKeyFrameHandler.h"
 
 namespace DetailWidgetConstants
 {
@@ -100,6 +101,8 @@ void SDetailSingleItemRow::Construct( const FArguments& InArgs, FDetailLayoutCus
 			ValueWidget->SetEnabled( Row.IsEnabledAttr );
 		}
 
+		TSharedRef<SWidget> KeyFrameButton = CreateKeyframeButton( *Customization, InOwnerTreeNode );
+
 		if( bHasMultipleColumns )
 		{
 			Widget = 
@@ -127,6 +130,14 @@ void SDetailSingleItemRow::Construct( const FArguments& InArgs, FDetailLayoutCus
 					[
 						NameWidget.ToSharedRef()
 					]
+					+ SHorizontalBox::Slot()
+					.Padding(3.0f, 0.0f)
+					.HAlign(HAlign_Right)
+					.VAlign(VAlign_Center)
+					.AutoWidth()
+					[
+						KeyFrameButton
+					]
 				]
 				+ SSplitter::Slot()
 				.Value( ColumnSizeData.RightColumnWidth )
@@ -160,11 +171,18 @@ void SDetailSingleItemRow::Construct( const FArguments& InArgs, FDetailLayoutCus
 				.Padding( DetailWidgetConstants::LeftRowPadding )
 				[
 					Row.WholeRowWidget.Widget
+				]
+				+ SHorizontalBox::Slot()
+				.Padding(3.0f, 0.0f)
+				.HAlign(HAlign_Right)
+				.VAlign(VAlign_Center)
+				[
+					KeyFrameButton
 				];
 		}
-
 	}
 
+
 	this->ChildSlot
 	[	
 		SNew( SBorder )
@@ -276,4 +294,46 @@ const FSlateBrush* SDetailSingleItemRow::GetBorderImage() const
 	{
 		return FEditorStyle::GetBrush("DetailsView.CategoryMiddle");
 	}
+}
+
+TSharedRef<SWidget> SDetailSingleItemRow::CreateKeyframeButton( FDetailLayoutCustomization& InCustomization, TSharedRef<IDetailTreeNode> InTreeNode )
+{
+	IDetailsViewPrivate& DetailsView = InTreeNode->GetDetailsView();
+
+	TSharedPtr<IDetailKeyframeHandler> KeyframeHandler = DetailsView.GetKeyframeHandler();
+
+	EVisibility SetKeyVisibility = EVisibility::Collapsed;
+
+	if (InCustomization.HasPropertyNode() && KeyframeHandler.IsValid() )
+	{
+		TSharedPtr<IPropertyHandle> Handle = PropertyEditorHelpers::GetPropertyHandle(InCustomization.GetPropertyNode().ToSharedRef(), nullptr, nullptr);
+
+		UClass* ObjectClass = InCustomization.GetPropertyNode()->FindObjectItemParent()->GetObjectBaseClass();
+		SetKeyVisibility = KeyframeHandler->IsPropertyKeyable(*ObjectClass, *Handle) ? EVisibility::Visible : EVisibility::Hidden;
+	}
+
+	return 
+		SNew(SButton)
+		.HAlign(HAlign_Left)
+		.VAlign(VAlign_Center)
+		.ContentPadding(0.0f)
+		.ButtonStyle(FEditorStyle::Get(), "Sequencer.AddKey.Details")
+		.Visibility(SetKeyVisibility)
+		.ToolTipText( NSLOCTEXT("PropertyView", "AddKeyframeButton_ToolTip", "Adds a keyframe for this property to the current animation") )
+		.OnClicked( this, &SDetailSingleItemRow::OnAddKeyframeClicked );
+}
+
+FReply SDetailSingleItemRow::OnAddKeyframeClicked()
+{
+	IDetailsViewPrivate& DetailsView = OwnerTreeNode.Pin()->GetDetailsView();
+
+	TSharedPtr<IDetailKeyframeHandler> KeyframeHandler = DetailsView.GetKeyframeHandler();
+
+	check(KeyframeHandler.IsValid());
+
+	TSharedPtr<IPropertyHandle> Handle = PropertyEditorHelpers::GetPropertyHandle(Customization->GetPropertyNode().ToSharedRef(), nullptr, nullptr);
+
+	KeyframeHandler->OnKeyPropertyClicked(*Handle);
+
+	return FReply::Handled();
 }
\ No newline at end of file
diff --git a/Engine/Source/Editor/PropertyEditor/Private/SDetailSingleItemRow.h b/Engine/Source/Editor/PropertyEditor/Private/SDetailSingleItemRow.h
--- a/Engine/Source/Editor/PropertyEditor/Private/SDetailSingleItemRow.h
+++ b/Engine/Source/Editor/PropertyEditor/Private/SDetailSingleItemRow.h
@@ -29,6 +29,8 @@ private:
 	void OnPasteProperty();
 	bool CanPasteProperty() const;
 	const FSlateBrush* GetBorderImage() const;
+	TSharedRef<SWidget> CreateKeyframeButton( FDetailLayoutCustomization& InCustomization, TSharedRef<IDetailTreeNode> InTreeNode );
+	FReply OnAddKeyframeClicked();
 private:
 	/** Customization for this widget */
 	FDetailLayoutCustomization* Customization;
diff --git a/Engine/Source/Editor/PropertyEditor/Private/SDetailsViewBase.cpp b/Engine/Source/Editor/PropertyEditor/Private/SDetailsViewBase.cpp
--- a/Engine/Source/Editor/PropertyEditor/Private/SDetailsViewBase.cpp
+++ b/Engine/Source/Editor/PropertyEditor/Private/SDetailsViewBase.cpp
@@ -218,6 +218,17 @@ bool SDetailsViewBase::IsPropertyEditingEnabled() const
 	return !IsPropertyEditingEnabledDelegate.IsBound() || IsPropertyEditingEnabledDelegate.Execute();
 }
 
+
+void SDetailsViewBase::SetKeyframeHandler( TSharedPtr<class IDetailKeyframeHandler> InKeyframeHandler )
+{
+	KeyframeHandler = InKeyframeHandler;
+}
+
+TSharedPtr<IDetailKeyframeHandler> SDetailsViewBase::GetKeyframeHandler() 
+{
+	return KeyframeHandler;
+}
+
 void SDetailsViewBase::SetGenericLayoutDetailsDelegate(FOnGetDetailCustomizationInstance OnGetGenericDetails)
 {
 	GenericLayoutDelegate = OnGetGenericDetails;
diff --git a/Engine/Source/Editor/PropertyEditor/Private/SDetailsViewBase.h b/Engine/Source/Editor/PropertyEditor/Private/SDetailsViewBase.h
--- a/Engine/Source/Editor/PropertyEditor/Private/SDetailsViewBase.h
+++ b/Engine/Source/Editor/PropertyEditor/Private/SDetailsViewBase.h
@@ -131,6 +131,9 @@ public:
 	 */
 	virtual bool IsPropertyEditingEnabled() const override;
 
+	virtual void SetKeyframeHandler( TSharedPtr<class IDetailKeyframeHandler> InKeyframeHandler ) override;
+	virtual TSharedPtr<IDetailKeyframeHandler> GetKeyframeHandler() override;
+
 	/**
 	 * Requests that an item in the tree be expanded or collapsed
 	 *
@@ -200,7 +203,7 @@ public:
 	/**
 	 * @return True if the property is visible
 	 */
-	bool IsPropertyVisible( const struct FPropertyAndParent& PropertyAndParent ) const override;
+	virtual bool IsPropertyVisible( const struct FPropertyAndParent& PropertyAndParent ) const override;
 
 	/**
 	* Sets a delegate which is regardless of the objects being viewed to lay out generic details not specific to any object
@@ -403,6 +406,8 @@ protected:
 	/** Root tree node that needs to be destroyed when safe */
 	TSharedPtr<FComplexPropertyNode> RootNodePendingKill;
 
+	TSharedPtr<IDetailKeyframeHandler> KeyframeHandler;
+
 	/** External property nodes which need to validated each tick */
 	TArray< TWeakPtr<FObjectPropertyNode> > ExternalRootPropertyNodes;
 };
\ No newline at end of file
diff --git a/Engine/Source/Editor/PropertyEditor/Public/IDetailKeyframeHandler.h b/Engine/Source/Editor/PropertyEditor/Public/IDetailKeyframeHandler.h
new file mode 100644
--- /dev/null
+++ b/Engine/Source/Editor/PropertyEditor/Public/IDetailKeyframeHandler.h
@@ -0,0 +1,14 @@
+// Copyright 1998-2014 Epic Games, Inc. All Rights Reserved.
+
+#pragma once
+
+class IDetailKeyframeHandler
+{
+public:
+	virtual ~IDetailKeyframeHandler(){}
+
+	virtual bool IsPropertyKeyable(const UClass& InObjectClass, const class IPropertyHandle& PropertyHandle) const = 0;
+
+	virtual void OnKeyPropertyClicked(const IPropertyHandle& KeyedPropertyHandle) = 0;
+
+};
diff --git a/Engine/Source/Editor/PropertyEditor/Public/IDetailsView.h b/Engine/Source/Editor/PropertyEditor/Public/IDetailsView.h
--- a/Engine/Source/Editor/PropertyEditor/Public/IDetailsView.h
+++ b/Engine/Source/Editor/PropertyEditor/Public/IDetailsView.h
@@ -140,6 +140,8 @@ public:
 	 */ 
 	virtual void SetIsPropertyEditingEnabledDelegate( FIsPropertyEditingEnabled IsPropertyEditingEnabled ) = 0;
 
+	virtual void SetKeyframeHandler( TSharedPtr<class IDetailKeyframeHandler> InKeyframeHandler ) = 0;
+
 	/**
 	 * @return true if property editing is enabled (based on the FIsPropertyEditingEnabled delegate)
 	 */ 
diff --git a/Engine/Source/Editor/PropertyEditor/Public/IPropertyChangeListener.h b/Engine/Source/Editor/PropertyEditor/Public/IPropertyChangeListener.h
--- a/Engine/Source/Editor/PropertyEditor/Public/IPropertyChangeListener.h
+++ b/Engine/Source/Editor/PropertyEditor/Public/IPropertyChangeListener.h
@@ -24,7 +24,7 @@ class IPropertyChangeListener
 {
 public:
 	/** Delegate called when a property changes */
-	DECLARE_MULTICAST_DELEGATE_TwoParams( FOnPropertyChanged, const TArray<UObject*>&, TSharedRef< const class IPropertyHandle>  );
+	DECLARE_MULTICAST_DELEGATE_TwoParams( FOnPropertyChanged, const TArray<UObject*>&, const class IPropertyHandle&  );
 
 	/**
 	 * Sets the object that should be listened to for changes
diff --git a/Engine/Source/Editor/PropertyEditor/Public/PropertyHandle.h b/Engine/Source/Editor/PropertyEditor/Public/PropertyHandle.h
--- a/Engine/Source/Editor/PropertyEditor/Public/PropertyHandle.h
+++ b/Engine/Source/Editor/PropertyEditor/Public/PropertyHandle.h
@@ -274,7 +274,7 @@ public:
 	 *
 	 * @param OuterObjects	An array that will be populated with the outer objects 
 	 */
-	virtual void GetOuterObjects( TArray<UObject*>& OuterObjects ) = 0;
+	virtual void GetOuterObjects( TArray<UObject*>& OuterObjects ) const = 0;
 
 	/**
 	 * Accesses the raw data of this property.  (Each pointer can be cast to the property data type)
diff --git a/Engine/Source/Editor/Sequencer/Private/Sequencer.cpp b/Engine/Source/Editor/Sequencer/Private/Sequencer.cpp
--- a/Engine/Source/Editor/Sequencer/Private/Sequencer.cpp
+++ b/Engine/Source/Editor/Sequencer/Private/Sequencer.cpp
@@ -1253,6 +1253,16 @@ void FSequencer::FilterToSelectedShotSections(bool bZoomToShotBounds)
 	FilterToShotSections(SelectedShotSections, bZoomToShotBounds);
 }
 
+bool FSequencer::CanKeyProperty(const UClass& ObjectClass, const IPropertyHandle& PropertyHandle) const
+{
+	return ObjectChangeListener->IsTypeKeyable( ObjectClass, PropertyHandle );
+} 
+
+void FSequencer::KeyProperty(const TArray<UObject*>& ObjectsToKey, const class IPropertyHandle& PropertyHandle) 
+{
+	ObjectChangeListener->KeyProperty( ObjectsToKey, PropertyHandle );
+}
+
 TArray< TWeakObjectPtr<UMovieSceneSection> > FSequencer::GetFilteringShotSections() const
 {
 	return FilteringShots;
diff --git a/Engine/Source/Editor/Sequencer/Private/Sequencer.h b/Engine/Source/Editor/Sequencer/Private/Sequencer.h
--- a/Engine/Source/Editor/Sequencer/Private/Sequencer.h
+++ b/Engine/Source/Editor/Sequencer/Private/Sequencer.h
@@ -81,6 +81,8 @@ public:
 	virtual void AddSubMovieScene(UMovieScene* SubMovieScene) override;
 	virtual void FilterToShotSections(const TArray< TWeakObjectPtr<class UMovieSceneSection> >& ShotSections, bool bZoomToShotBounds = true) override;
 	virtual void FilterToSelectedShotSections(bool bZoomToShotBounds = true) override;
+	virtual bool CanKeyProperty(const UClass& ObjectClass, const class IPropertyHandle& PropertyHandle) const override;
+	virtual void KeyProperty(const TArray<UObject*>& ObjectsToKey, const class IPropertyHandle& PropertyHandle) override;
 
 	bool IsPerspectiveViewportPosessionEnabled() const { return bPerspectiveViewportPossessionEnabled; }
 
diff --git a/Engine/Source/Editor/Sequencer/Private/SequencerObjectChangeListener.cpp b/Engine/Source/Editor/Sequencer/Private/SequencerObjectChangeListener.cpp
--- a/Engine/Source/Editor/Sequencer/Private/SequencerObjectChangeListener.cpp
+++ b/Engine/Source/Editor/Sequencer/Private/SequencerObjectChangeListener.cpp
@@ -24,25 +24,47 @@ FSequencerObjectChangeListener::~FSequencerObjectChangeListener()
 	GEditor->OnActorMoved().RemoveAll( this );
 }
 
-void FSequencerObjectChangeListener::OnPropertyChanged( const TArray<UObject*>& ChangedObjects,  TSharedRef< const IPropertyHandle> PropertyValue, bool bRequireAutoKey )
+void FSequencerObjectChangeListener::OnPropertyChanged( const TArray<UObject*>& ChangedObjects, const IPropertyHandle& PropertyHandle, bool bRequireAutoKey ) const
 {
-	TSharedPtr< const IPropertyHandle> HandleToUse = PropertyValue;
-
-	TSharedPtr<const IPropertyHandle> ParentHandle = PropertyValue->GetParentHandle();
-
-	if( ParentHandle.IsValid() && ParentHandle->GetProperty() && ParentHandle->GetProperty()->IsA<UStructProperty>() )
+	const UProperty* Property = PropertyHandle.GetProperty();
+	const UStructProperty* StructProperty = Cast<const UStructProperty>(Property);
+	const UStructProperty* ParentStructProperty = nullptr;
+	TSharedPtr<IPropertyHandle> ParentHandle;
+	if(!StructProperty)
 	{
-		// Handle the case where internal values of a struct are being set.  We want to key the whole struct in this case
-		// @todo Sequencer: Is this always the case?
-		HandleToUse = ParentHandle;
+		ParentHandle = PropertyHandle.GetParentHandle();
+		if(ParentHandle->IsValidHandle())
+		{
+			ParentStructProperty = Cast<const UStructProperty>(ParentHandle->GetProperty());
+		}
 	}
 
-	UProperty* Property = HandleToUse->GetProperty();
-	UStructProperty* StructProperty = Cast<UStructProperty>( Property );
+	FString PropertyVarName;
 
-	FName PropertyName = StructProperty ? StructProperty->Struct->GetFName() : Property->GetClass()->GetFName();
+	// If not in auto-key allow partial keying of specific inner struct property values;
+	FName InnerStructPropName = !bRequireAutoKey && ParentStructProperty ? Property->GetFName() : NAME_None;
 
-	ClassToPropertyChangedMap.FindRef( PropertyName ).Broadcast( ChangedObjects, *HandleToUse, bRequireAutoKey );
+	FKeyPropertyParams Params;
+	Params.InnerStructPropertyName = InnerStructPropName;
+	Params.bRequireAutoKey = bRequireAutoKey;
+	Params.ObjectsThatChanged = ChangedObjects;
+	if(StructProperty)
+	{
+		Params.PropertyHandle = &PropertyHandle;
+		ClassToPropertyChangedMap.FindRef(StructProperty->Struct->GetFName()).Broadcast( Params );
+	}
+	else if(ParentStructProperty)
+	{
+		Params.PropertyHandle = ParentHandle.Get();;
+		// If the property parent is a struct, see if this property parent can be keyed. (e.g R,G,B,A for a color)
+		ClassToPropertyChangedMap.FindRef(ParentStructProperty->Struct->GetFName()).Broadcast( Params );
+	}
+	else
+	{
+		Params.PropertyHandle = &PropertyHandle;
+		// the property in question is not a struct or an inner of the struct. See if it is directly keyable
+		ClassToPropertyChangedMap.FindRef(Property->GetClass()->GetFName()).Broadcast( Params );
+	}
 }
 
 bool FSequencerObjectChangeListener::IsObjectValidForListening( UObject* Object ) const
@@ -62,6 +84,61 @@ FOnPropagateObjectChanges& FSequencerObjectChangeListener::GetOnPropagateObjectC
 	return OnPropagateObjectChanges;
 }
 
+bool FSequencerObjectChangeListener::IsTypeKeyable(const UClass& ObjectClass, const IPropertyHandle& PropertyHandle) const
+{
+	const UProperty* Property = PropertyHandle.GetProperty();
+	const UStructProperty* StructProperty = Cast<const UStructProperty>(Property);
+	const UStructProperty* ParentStructProperty = nullptr;
+	if( !StructProperty )
+	{
+		const TSharedPtr<IPropertyHandle> ParentHandle = PropertyHandle.GetParentHandle();
+		if( ParentHandle->IsValidHandle() )
+		{
+			ParentStructProperty = Cast<const UStructProperty>( ParentHandle->GetProperty() );
+		}
+	}
+
+	FString PropertyVarName;
+
+	bool bFound = false;
+	if( StructProperty )
+	{
+		bFound = ClassToPropertyChangedMap.Contains(StructProperty->Struct->GetFName());
+		PropertyVarName = StructProperty->GetName();
+	}
+	else if( ParentStructProperty )
+	{
+		// If the property parent is a struct, see if this property parent can be keyed. (e.g R,G,B,A for a color)
+		bFound = ClassToPropertyChangedMap.Contains(ParentStructProperty->Struct->GetFName());
+		PropertyVarName = ParentStructProperty->GetName();
+	}
+
+	if( !bFound )
+	{
+		// the property in question is not a struct or an inner of the struct. See if it is directly keyable
+		bFound = ClassToPropertyChangedMap.Contains(Property->GetClass()->GetFName());
+		PropertyVarName = Property->GetName();
+	}
+
+	if( bFound )
+	{
+		static const FString Set(TEXT("Set"));
+
+		const FString FunctionString = Set + PropertyVarName;
+
+		FName FunctionName = FName(*FunctionString);
+
+		bFound = ObjectClass.FindFunctionByName( FunctionName ) != nullptr; 
+	}
+
+	return bFound;
+}
+
+void FSequencerObjectChangeListener::KeyProperty( const TArray<UObject*>& ObjectsToKey, const class IPropertyHandle& PropertyHandle ) const
+{
+	const bool bRequireAutoKey = false;
+	OnPropertyChanged( ObjectsToKey, PropertyHandle, bRequireAutoKey );
+}
 
 void FSequencerObjectChangeListener::OnObjectPreEditChange( UObject* Object, const FEditPropertyChain& PropertyChain )
 {
@@ -175,7 +252,7 @@ void FSequencerObjectChangeListener::TriggerAllPropertiesChanged(UObject* Object
 				Settings.bIgnoreArrayProperties = true;
 				Settings.bIgnoreObjectProperties = true;
 				// Property flags which must be on the property
-				Settings.RequiredPropertyFlags = CPF_Interp;
+				Settings.RequiredPropertyFlags = 0;
 				// Property flags which cannot be on the property
 				Settings.DisallowedPropertyFlags = CPF_EditConst;
 
diff --git a/Engine/Source/Editor/Sequencer/Private/SequencerObjectChangeListener.h b/Engine/Source/Editor/Sequencer/Private/SequencerObjectChangeListener.h
--- a/Engine/Source/Editor/Sequencer/Private/SequencerObjectChangeListener.h
+++ b/Engine/Source/Editor/Sequencer/Private/SequencerObjectChangeListener.h
@@ -16,6 +16,8 @@ public:
 	/** ISequencerObjectChangeListener interface */
 	virtual FOnAnimatablePropertyChanged& GetOnAnimatablePropertyChanged( FName PropertyTypeName ) override;
 	virtual FOnPropagateObjectChanges& GetOnPropagateObjectChanges() override;
+	virtual bool IsTypeKeyable(const UClass& ObjectClass, const IPropertyHandle& PropertyHandle) const override;
+	virtual void KeyProperty( const TArray<UObject*>& ObjectsToKey, const class IPropertyHandle& PropertyHandle ) const override;
 	virtual void TriggerAllPropertiesChanged(UObject* Object) override;
 
 private:
@@ -46,7 +48,7 @@ private:
 	 *
 	 * @param Object	The object that PostEditChange was called on
 	 */
-	void OnPropertyChanged( const TArray<UObject*>& ChangedObjects, TSharedRef< const IPropertyHandle> PropertyValue, bool bRequireAutoKey );
+	void OnPropertyChanged( const TArray<UObject*>& ChangedObjects, const IPropertyHandle& PropertyValue, bool bRequireAutoKey ) const;
 
 	/**
 	 * @return True if an object is valid for listening to property changes 
diff --git a/Engine/Source/Editor/Sequencer/Public/ISequencer.h b/Engine/Source/Editor/Sequencer/Public/ISequencer.h
--- a/Engine/Source/Editor/Sequencer/Public/ISequencer.h
+++ b/Engine/Source/Editor/Sequencer/Public/ISequencer.h
@@ -111,7 +111,13 @@ public:
 	 */
 	virtual class ISequencerObjectChangeListener& GetObjectChangeListener() = 0;
 
+	virtual bool CanKeyProperty(const UClass& ObjectClass, const class IPropertyHandle& PropertyHandle) const = 0;
+
+	virtual void KeyProperty( const TArray<UObject*>& ObjectsToKey, const class IPropertyHandle& PropertyHandle ) = 0;
+
 	virtual void NotifyMovieSceneDataChanged() = 0;
 
 	virtual void UpdateRuntimeInstances() = 0;
+
+
 };
diff --git a/Engine/Source/Editor/Sequencer/Public/ISequencerObjectChangeListener.h b/Engine/Source/Editor/Sequencer/Public/ISequencerObjectChangeListener.h
--- a/Engine/Source/Editor/Sequencer/Public/ISequencerObjectChangeListener.h
+++ b/Engine/Source/Editor/Sequencer/Public/ISequencerObjectChangeListener.h
@@ -4,7 +4,16 @@
 
 class IPropertyHandle;
 
-DECLARE_MULTICAST_DELEGATE_ThreeParams( FOnAnimatablePropertyChanged, const TArray<UObject*>&, const IPropertyHandle&, bool );
+struct FKeyPropertyParams
+{
+	const IPropertyHandle* PropertyHandle;
+	FName InnerStructPropertyName;
+	TArray<UObject*> ObjectsThatChanged;
+	bool bRequireAutoKey;
+};
+
+
+DECLARE_MULTICAST_DELEGATE_OneParam( FOnAnimatablePropertyChanged, const FKeyPropertyParams& );
 
 DECLARE_MULTICAST_DELEGATE_OneParam ( FOnPropagateObjectChanges, UObject* );
 
@@ -28,4 +37,8 @@ public:
 	 * Triggers all properties as changed for the passed in object.
 	 */
 	virtual void TriggerAllPropertiesChanged(UObject* Object) = 0;
+
+	virtual bool IsTypeKeyable(const UClass& ObjectClass, const IPropertyHandle& PropertyHandle) const = 0;
+
+	virtual void KeyProperty( const TArray<UObject*>& ObjectsToKey, const class IPropertyHandle& PropertyHandle ) const = 0;
 };
\ No newline at end of file
diff --git a/Engine/Source/Editor/UMGEditor/Private/Animation/MarginTrackEditor.cpp b/Engine/Source/Editor/UMGEditor/Private/Animation/MarginTrackEditor.cpp
--- a/Engine/Source/Editor/UMGEditor/Private/Animation/MarginTrackEditor.cpp
+++ b/Engine/Source/Editor/UMGEditor/Private/Animation/MarginTrackEditor.cpp
@@ -67,43 +67,38 @@ TSharedRef<ISequencerSection> FMarginTrackEditor::MakeSectionInterface( UMovieSc
 
 	return NewSection;
 }
-/*
 
-void FMarginTrackEditor::AddKey(const FGuid& ObjectGuid, UObject* AdditionalAsset)
-{
-	ISequencerObjectChangeListener& ObjectChangeListener = GetSequencer()->GetObjectChangeListener();
-
-	TArray<UObject*> OutObjects;
-	GetSequencer()->GetRuntimeObjects( GetSequencer()->GetFocusedMovieSceneInstance(), ObjectGuid, OutObjects);
-	for (int32 i = 0; i < OutObjects.Num(); ++i)
-	{
-		ObjectChangeListener.TriggerAllPropertiesChanged(OutObjects[i]);
-	}
-}*/
 
-void FMarginTrackEditor::OnMarginChanged( const TArray<UObject*>& InObjectsThatChanged, const IPropertyHandle& PropertyValue, bool bRequireAutoKey )
+void FMarginTrackEditor::OnMarginChanged(  const FKeyPropertyParams& PropertyKeyParams )
 {
-	FName PropertyName = PropertyValue.GetProperty()->GetFName();
+	FName PropertyName = PropertyKeyParams.PropertyHandle->GetProperty()->GetFName();
 
 	AnimatablePropertyChanged
 	(
 		UMovieSceneMarginTrack::StaticClass(), 
-		bRequireAutoKey,
-		FOnKeyProperty::CreateRaw(this, &FMarginTrackEditor::OnKeyMargin, &InObjectsThatChanged, &PropertyValue, PropertyName ) 
+		PropertyKeyParams.bRequireAutoKey,
+		FOnKeyProperty::CreateRaw(this, &FMarginTrackEditor::OnKeyMargin, &PropertyKeyParams ) 
 	);
 }
 
 
-void FMarginTrackEditor::OnKeyMargin( float KeyTime, const TArray<UObject*>* InObjectsThatChanged, const IPropertyHandle* PropertyValue, FName PropertyName )
+void FMarginTrackEditor::OnKeyMargin( float KeyTime, const FKeyPropertyParams* PropertyKeyParams )
 {
 	TArray<const void*> MarginValues;
-	PropertyValue->AccessRawData( MarginValues );
+	PropertyKeyParams->PropertyHandle->AccessRawData( MarginValues );
+
+	FName PropertyName = PropertyKeyParams->PropertyHandle->GetProperty()->GetFName();
 
-	for( int32 ObjectIndex = 0; ObjectIndex < InObjectsThatChanged->Num(); ++ObjectIndex )
+	for( int32 ObjectIndex = 0; ObjectIndex < PropertyKeyParams->ObjectsThatChanged.Num(); ++ObjectIndex )
 	{
-		UObject* Object = (*InObjectsThatChanged)[ObjectIndex];
+		UObject* Object = PropertyKeyParams->ObjectsThatChanged[ObjectIndex];
 		FMargin MarginValue = *(const FMargin*)MarginValues[ObjectIndex];
 
+		FMarginKey Key;
+		Key.bAddKeyEvenIfUnchanged = !PropertyKeyParams->bRequireAutoKey;
+		Key.CurveName = PropertyKeyParams->InnerStructPropertyName;
+		Key.Value = MarginValue;
+
 		FGuid ObjectHandle = FindOrCreateHandleToObject( Object );
 		if (ObjectHandle.IsValid())
 		{
@@ -114,7 +109,7 @@ void FMarginTrackEditor::OnKeyMargin( float KeyTime, const TArray<UObject*>* InO
 				MarginTrack->SetPropertyName( PropertyName );
 				// Find or add a new section at the auto-key time and changing the property same property
 				// AddKeyToSection is not actually a virtual, it's redefined in each class with a different type
-				bool bSuccessfulAdd = MarginTrack->AddKeyToSection( KeyTime, MarginValue );
+				bool bSuccessfulAdd = MarginTrack->AddKeyToSection( KeyTime, Key );
 				if (bSuccessfulAdd)
 				{
 					MarginTrack->SetAsShowable();
diff --git a/Engine/Source/Editor/UMGEditor/Private/Animation/MarginTrackEditor.h b/Engine/Source/Editor/UMGEditor/Private/Animation/MarginTrackEditor.h
--- a/Engine/Source/Editor/UMGEditor/Private/Animation/MarginTrackEditor.h
+++ b/Engine/Source/Editor/UMGEditor/Private/Animation/MarginTrackEditor.h
@@ -36,10 +36,10 @@ private:
 	 * @param InObjectsThatChanged	List of objects that changed
 	 * @param PropertyValue			Handle to the property value which changed
 	 */
-	void OnMarginChanged( const TArray<UObject*>& InObjectsThatChanged, const IPropertyHandle& PropertyValue, bool bRequireAutoKey );
+	void OnMarginChanged( const struct FKeyPropertyParams& PropertyKeyParams );
 
 	/** Called After OnMarginChanged if we actually can key the margin */
-	void OnKeyMargin( float KeyTime, const TArray<UObject*>* InObjectsThatChanged, const IPropertyHandle* PropertyValue, FName PropertyName );
+	void OnKeyMargin( float KeyTime, const struct FKeyPropertyParams* PropertyKeyParams );
 };
 
 
diff --git a/Engine/Source/Editor/UMGEditor/Private/Animation/UMGDetailKeyframeHandler.cpp b/Engine/Source/Editor/UMGEditor/Private/Animation/UMGDetailKeyframeHandler.cpp
new file mode 100644
--- /dev/null
+++ b/Engine/Source/Editor/UMGEditor/Private/Animation/UMGDetailKeyframeHandler.cpp
@@ -0,0 +1,22 @@
+// Copyright 1998-2014 Epic Games, Inc. All Rights Reserved.
+
+#include "UMGEditorPrivatePCH.h"
+#include "UMGDetailKeyframeHandler.h"
+
+
+FUMGDetailKeyframeHandler::FUMGDetailKeyframeHandler(TSharedPtr<FWidgetBlueprintEditor> InBlueprintEditor)
+	: BlueprintEditor( InBlueprintEditor )
+{}
+
+bool FUMGDetailKeyframeHandler::IsPropertyKeyable(const UClass& InObjectClass, const IPropertyHandle& InPropertyHandle) const
+{
+	return BlueprintEditor.Pin()->GetSequencer()->CanKeyProperty( InObjectClass, InPropertyHandle );
+}
+
+void FUMGDetailKeyframeHandler::OnKeyPropertyClicked(const IPropertyHandle& KeyedPropertyHandle)
+{
+	TArray<UObject*> Objects;
+	KeyedPropertyHandle.GetOuterObjects( Objects );
+
+	BlueprintEditor.Pin()->GetSequencer()->KeyProperty( Objects, KeyedPropertyHandle );
+}
diff --git a/Engine/Source/Editor/UMGEditor/Private/Animation/UMGDetailKeyframeHandler.h b/Engine/Source/Editor/UMGEditor/Private/Animation/UMGDetailKeyframeHandler.h
new file mode 100644
--- /dev/null
+++ b/Engine/Source/Editor/UMGEditor/Private/Animation/UMGDetailKeyframeHandler.h
@@ -0,0 +1,18 @@
+// Copyright 1998-2014 Epic Games, Inc. All Rights Reserved.
+
+#pragma once
+
+#include "Editor/PropertyEditor/Public/IDetailKeyFrameHandler.h"
+
+class FUMGDetailKeyframeHandler : public IDetailKeyframeHandler
+{
+public:
+	FUMGDetailKeyframeHandler( TSharedPtr<class FWidgetBlueprintEditor> InBlueprintEditor );
+
+	virtual bool IsPropertyKeyable(const UClass& InObjectClass, const class IPropertyHandle& PropertyHandle) const override;
+
+	virtual void OnKeyPropertyClicked(const IPropertyHandle& KeyedPropertyHandle) override;
+
+private:
+	TWeakPtr<class FWidgetBlueprintEditor> BlueprintEditor;
+};
\ No newline at end of file
diff --git a/Engine/Source/Editor/UMGEditor/Private/Animation/UMGSequencerObjectBindingManager.cpp b/Engine/Source/Editor/UMGEditor/Private/Animation/UMGSequencerObjectBindingManager.cpp
new file mode 100644
--- /dev/null
+++ b/Engine/Source/Editor/UMGEditor/Private/Animation/UMGSequencerObjectBindingManager.cpp
@@ -0,0 +1,131 @@
+// Copyright 1998-2014 Epic Games, Inc. All Rights Reserved.
+
+#include "UMGEditorPrivatePCH.h"
+#include "UMGSequencerObjectBindingManager.h"
+#include "WidgetBlueprintEditor.h"
+#include "ISequencer.h"
+#include "MovieScene.h"
+
+FUMGSequencerObjectBindingManager::FUMGSequencerObjectBindingManager( FWidgetBlueprintEditor& InWidgetBlueprintEditor, UMovieScene& InMovieScene )
+	: WidgetBlueprintEditor( InWidgetBlueprintEditor )
+	, MovieScene( &InMovieScene )
+{
+	WidgetBlueprintEditor.GetOnWidgetPreviewUpdated().AddRaw( this, &FUMGSequencerObjectBindingManager::OnWidgetPreviewUpdated );
+
+}
+
+FUMGSequencerObjectBindingManager::~FUMGSequencerObjectBindingManager()
+{
+}
+
+FGuid FUMGSequencerObjectBindingManager::FindGuidForObject( const UMovieScene& MovieScene, UObject& Object ) const
+{
+	return PreviewObjectToGuidMap.FindRef( &Object );
+}
+
+bool FUMGSequencerObjectBindingManager::CanPossessObject( UObject& Object ) const
+{
+	UWidgetBlueprint* WidgetBlueprint = WidgetBlueprintEditor.GetWidgetBlueprintObj();
+
+	// Only preview widgets in this blueprint can be possessed
+	UUserWidget* PreviewWidget = WidgetBlueprintEditor.GetPreview();
+
+	return Object.IsA<UVisual>() && Object.IsIn( PreviewWidget );
+}
+
+void FUMGSequencerObjectBindingManager::BindPossessableObject( const FGuid& PossessableGuid, UObject& PossessedObject )
+{
+	PreviewObjectToGuidMap.Add( &PossessedObject, PossessableGuid );
+
+	UWidgetBlueprint* WidgetBlueprint = WidgetBlueprintEditor.GetWidgetBlueprintObj();
+	FWidgetAnimation* WidgetAnimation = WidgetBlueprint->FindAnimationDataForMovieScene(*MovieScene);
+
+	FWidgetAnimationBinding NewBinding;
+	NewBinding.WidgetName = PossessedObject.GetFName();
+	if( PossessedObject.IsA<UPanelSlot>() )
+	{
+		// Save the name of the widget containing the slots.  This is the object to look up that contaisn the slot itself(the thing we are animating)
+		NewBinding.SlotWidgetName = PossessedObject.GetOuter()->GetFName();
+	}
+
+	NewBinding.AnimationGuid = PossessableGuid;
+
+	WidgetAnimation->AnimationBindings.Add( NewBinding );
+}
+
+void FUMGSequencerObjectBindingManager::UnbindPossessableObjects( const FGuid& PossessableGuid )
+{
+	for( auto It = PreviewObjectToGuidMap.CreateIterator(); It; ++It )
+	{
+		if( It.Value() == PossessableGuid )
+		{
+			It.RemoveCurrent();
+		}
+	}
+
+	UWidgetBlueprint* WidgetBlueprint = WidgetBlueprintEditor.GetWidgetBlueprintObj();
+	FWidgetAnimation* WidgetAnimation = WidgetBlueprint->FindAnimationDataForMovieScene(*MovieScene);
+
+	WidgetAnimation->AnimationBindings.RemoveAll( [&]( const FWidgetAnimationBinding& Binding ) { return Binding.AnimationGuid == PossessableGuid; } );
+
+}
+
+void FUMGSequencerObjectBindingManager::GetRuntimeObjects( const TSharedRef<FMovieSceneInstance>& MovieSceneInstance, const FGuid& ObjectGuid, TArray<UObject*>& OutRuntimeObjects ) const
+{
+	for( auto It = PreviewObjectToGuidMap.CreateConstIterator(); It; ++It )
+	{
+		UObject* Object = It.Key().Get();
+		if( Object && It.Value() == ObjectGuid )
+		{
+			OutRuntimeObjects.Add( Object );
+		}
+	}
+}
+
+void FUMGSequencerObjectBindingManager::InitPreviewObjects()
+{
+	UWidgetBlueprint* WidgetBlueprint = WidgetBlueprintEditor.GetWidgetBlueprintObj();
+
+	// Populate preview object to guid map
+	UUserWidget* PreviewWidget = WidgetBlueprintEditor.GetPreview();
+
+	UWidgetTree* WidgetTree = PreviewWidget->WidgetTree;
+
+	const FWidgetAnimation* WidgetAnimation = WidgetBlueprint->FindAnimationDataForMovieScene( *MovieScene );
+	check( WidgetAnimation );
+
+	for( const FWidgetAnimationBinding& Binding : WidgetAnimation->AnimationBindings )
+	{
+		FName ObjectName = Binding.WidgetName;
+		FName SlotWidgetName = Binding.SlotWidgetName;
+
+		FName NameToSearchFor = SlotWidgetName != NAME_None ? SlotWidgetName : ObjectName;
+
+		UObject* FoundObject = FindObject<UObject>(WidgetTree, *NameToSearchFor.ToString());
+		if(FoundObject)
+		{
+			if( SlotWidgetName == NAME_None )
+			{
+				PreviewObjectToGuidMap.Add(FoundObject, Binding.AnimationGuid);
+			}
+			else
+			{
+				FoundObject = FindObject<UObject>( FoundObject, *ObjectName.ToString() );
+				if( FoundObject )
+				{
+					PreviewObjectToGuidMap.Add(FoundObject, Binding.AnimationGuid);
+				}
+			}
+		
+		}
+	}
+}
+
+void FUMGSequencerObjectBindingManager::OnWidgetPreviewUpdated()
+{
+	PreviewObjectToGuidMap.Empty();
+
+	InitPreviewObjects();
+
+	WidgetBlueprintEditor.GetSequencer()->UpdateRuntimeInstances();
+}
\ No newline at end of file
diff --git a/Engine/Source/Editor/UMGEditor/Private/Animation/UMGSequencerObjectBindingManager.h b/Engine/Source/Editor/UMGEditor/Private/Animation/UMGSequencerObjectBindingManager.h
new file mode 100644
--- /dev/null
+++ b/Engine/Source/Editor/UMGEditor/Private/Animation/UMGSequencerObjectBindingManager.h
@@ -0,0 +1,35 @@
+// Copyright 1998-2014 Epic Games, Inc. All Rights Reserved.
+
+#pragma once
+
+#include "ISequencerObjectBindingManager.h"
+
+class FUMGSequencerObjectBindingManager : public ISequencerObjectBindingManager
+{
+public:
+	FUMGSequencerObjectBindingManager( FWidgetBlueprintEditor& InWidgetBlueprintEditor, UMovieScene& InMovieScene );
+	~FUMGSequencerObjectBindingManager();
+
+	/** ISequencerObjectBindingManager interface */
+	virtual bool AllowsSpawnableObjects() const override { return false; }
+	virtual FGuid FindGuidForObject( const UMovieScene& MovieScene, UObject& Object ) const override;
+	virtual void SpawnOrDestroyObjectsForInstance( TSharedRef<FMovieSceneInstance> MovieSceneInstance, bool bDestroyAll ) {}
+	virtual void DestroyAllSpawnedObjects()  {}
+	virtual bool CanPossessObject( UObject& Object ) const override;
+	virtual void BindPossessableObject( const FGuid& PossessableGuid, UObject& PossessedObject ) override;
+	virtual void UnbindPossessableObjects( const FGuid& PossessableGuid ) override;
+	virtual void GetRuntimeObjects( const TSharedRef<FMovieSceneInstance>& MovieSceneInstance, const FGuid& ObjectGuid, TArray<UObject*>& OutRuntimeObjects ) const override;
+	
+	/** Rebuilds mappings to live preview objects */
+	void InitPreviewObjects();
+
+private:
+	void OnWidgetPreviewUpdated();
+private:
+	/** Mapping of preview objects to sequencer guids */
+	TMap< TWeakObjectPtr<UObject>, FGuid> PreviewObjectToGuidMap;
+
+	TWeakObjectPtr<UMovieScene> MovieScene;
+
+	FWidgetBlueprintEditor& WidgetBlueprintEditor;
+};
diff --git a/Engine/Source/Editor/UMGEditor/Private/Customizations/CanvasSlotCustomization.cpp b/Engine/Source/Editor/UMGEditor/Private/Customizations/CanvasSlotCustomization.cpp
--- a/Engine/Source/Editor/UMGEditor/Private/Customizations/CanvasSlotCustomization.cpp
+++ b/Engine/Source/Editor/UMGEditor/Private/Customizations/CanvasSlotCustomization.cpp
@@ -68,6 +68,7 @@ void FCanvasSlotCustomization::CustomizeAnchors(TSharedRef<IPropertyHandle> Prop
 		.NameContent()
 		[
 			SNew(STextBlock)
+			.Font( IDetailLayoutBuilder::GetDetailFont() )
 			.Text(LOCTEXT("Anchors", "Anchors"))
 		]
 		.ValueContent()
diff --git a/Engine/Source/Editor/UMGEditor/Private/Customizations/DetailCustomizations.cpp b/Engine/Source/Editor/UMGEditor/Private/Customizations/DetailCustomizations.cpp
--- a/Engine/Source/Editor/UMGEditor/Private/Customizations/DetailCustomizations.cpp
+++ b/Engine/Source/Editor/UMGEditor/Private/Customizations/DetailCustomizations.cpp
@@ -12,205 +12,200 @@
 
 #define LOCTEXT_NAMESPACE "UMG"
 
-void FBlueprintWidgetCustomization::CustomizeDetails( IDetailLayoutBuilder& DetailLayout )
+void FBlueprintWidgetCustomization::CreateDelegateCustomization( IDetailLayoutBuilder& DetailLayout, UDelegateProperty* Property )
 {
-	TArray< TWeakObjectPtr<UObject> > OutObjects;
-	DetailLayout.GetObjectsBeingCustomized(OutObjects);
+	FString ConstPropertyName = Property->GetName();
+	ConstPropertyName.RemoveFromEnd(TEXT("Delegate"));
 
-	if ( OutObjects.Num() == 1 )
+	//TSharedRef<IPropertyHandle> DelegatePropertyHandle = DetailLayout.GetProperty(Property->GetFName(), PropertyClass);
+	TSharedRef<IPropertyHandle> ConstPropertyHandle = DetailLayout.GetProperty(FName(*ConstPropertyName), CastChecked<UClass>(Property->GetOuter()));
+
+	const bool bHasValidHandle = ConstPropertyHandle->IsValidHandle();
+	if(!bHasValidHandle)
 	{
-		UClass* PropertyClass = OutObjects[0].Get()->GetClass();
+		return;
+	}
 
-		for ( TFieldIterator<UProperty> PropertyIt(PropertyClass, EFieldIteratorFlags::IncludeSuper); PropertyIt; ++PropertyIt )
-		{
-			if ( UDelegateProperty* Property = Cast<UDelegateProperty>(*PropertyIt) )
-			{
-				if ( !Property->GetName().EndsWith(TEXT("Delegate")) )
-				{
-					continue;
-				}
+	UProperty* ConstProperty = ConstPropertyHandle->GetProperty();
 
-				FString ConstPropertyName = Property->GetName();
-				ConstPropertyName.RemoveFromEnd(TEXT("Delegate"));
+	const bool bIsEditable = ConstProperty->HasAnyPropertyFlags(CPF_Edit | CPF_EditConst);
+	const bool bDoSignaturesMatch = Property->SignatureFunction->GetReturnProperty()->SameType(ConstProperty);
 
-				//TSharedRef<IPropertyHandle> DelegatePropertyHandle = DetailLayout.GetProperty(Property->GetFName(), PropertyClass);
-				TSharedRef<IPropertyHandle> ConstPropertyHandle = DetailLayout.GetProperty(FName(*ConstPropertyName), CastChecked<UClass>(Property->GetOuter()));
+	if(!(bIsEditable && bDoSignaturesMatch))
+	{
+		return;
+	}
 
-				const bool bHasValidHandle = ConstPropertyHandle->IsValidHandle();
+	IDetailCategoryBuilder& PropertyCategory = DetailLayout.EditCategory(FObjectEditorUtils::GetCategoryFName(ConstProperty));
 
-				if ( !bHasValidHandle )
-				{
-					continue;
-				}
+	IDetailPropertyRow& PropertyRow = PropertyCategory.AddProperty(ConstPropertyHandle);
 
-				UProperty* ConstProperty = ConstPropertyHandle->GetProperty();
+	TSharedPtr<SWidget> NameWidget;
+	TSharedPtr<SWidget> ValueWidget;
+	FDetailWidgetRow Row;
+	PropertyRow.GetDefaultWidgets(NameWidget, ValueWidget, Row);
 
-				const bool bIsEditable = ConstProperty->HasAnyPropertyFlags(CPF_Edit | CPF_EditConst);
-				const bool bDoSignaturesMatch = Property->SignatureFunction->GetReturnProperty()->SameType(ConstProperty);
+	TSharedRef<SComboButton> DelegateButton = SNew(SComboButton)
+		.OnGetMenuContent(this, &FBlueprintWidgetCustomization::OnGenerateDelegateMenu, ConstPropertyHandle, Property->SignatureFunction)
+		.ContentPadding(0)
+		.ButtonContent()
+		[
+			SNew(SHorizontalBox)
+			+ SHorizontalBox::Slot()
+			.AutoWidth()
+			.VAlign(VAlign_Center)
+			[
+				SNew(SImage)
+				.Image(this, &FBlueprintWidgetCustomization::GetCurrentBindingImage, ConstPropertyHandle)
+				.ColorAndOpacity(FLinearColor(0.25f, 0.25f, 0.25f))
+			]
+			+ SHorizontalBox::Slot()
+			.AutoWidth()
+			.VAlign(VAlign_Center)
+			.Padding(4, 1, 0, 0)
+			[
+				SNew(STextBlock)
+				.Text(this, &FBlueprintWidgetCustomization::GetCurrentBindingText, ConstPropertyHandle)
+				.Font(IDetailLayoutBuilder::GetDetailFont())
+			]
+		];
 
-				if ( !(bIsEditable && bDoSignaturesMatch) )
-				{
-					continue;
-				}
+	const bool bShowChildren = true;
+	PropertyRow.CustomWidget(bShowChildren)
+		.NameContent()
+		.MinDesiredWidth(Row.NameWidget.MinWidth)
+		.MaxDesiredWidth(Row.NameWidget.MaxWidth)
+		[
+			NameWidget.ToSharedRef()
+		]
+		.ValueContent()
+		.MinDesiredWidth(Row.ValueWidget.MinWidth * 2)
+		.MaxDesiredWidth(Row.ValueWidget.MaxWidth * 2)
+		[
+			SNew(SHorizontalBox)
+			+ SHorizontalBox::Slot()
+			.FillWidth(1)
+			.Padding(FMargin(0, 0, 4, 0))
+			[
+				ValueWidget.ToSharedRef()
+			]
+			+ SHorizontalBox::Slot()
+			.AutoWidth()
+			[
+				DelegateButton
+			]
+			+ SHorizontalBox::Slot()
+			.AutoWidth()
+			[
+				SNew(SButton)
+				.ButtonStyle(FEditorStyle::Get(), "HoverHintOnly")
+				.Visibility(this, &FBlueprintWidgetCustomization::GetGotoBindingVisibility, ConstPropertyHandle)
+				.OnClicked(this, &FBlueprintWidgetCustomization::HandleGotoBindingClicked, ConstPropertyHandle)
+				.VAlign(VAlign_Center)
+				.ToolTipText(LOCTEXT("GotoFunction", "Goto Function"))
+				[
+					SNew(SImage)
+					.Image(FEditorStyle::GetBrush("PropertyWindow.Button_Browse"))
+				]
+			]
+		];
+}
 
-				IDetailCategoryBuilder& PropertyCategory = DetailLayout.EditCategory(FObjectEditorUtils::GetCategoryFName(ConstProperty));
+void FBlueprintWidgetCustomization::CreateEventCustomization( IDetailLayoutBuilder& DetailLayout, UDelegateProperty* Property )
+{
+	if(!Property->GetName().EndsWith(TEXT("Event")))
+	{
+		return;
+	}
 
-				IDetailPropertyRow& PropertyRow = PropertyCategory.AddProperty(ConstPropertyHandle);
+	TSharedRef<IPropertyHandle> DelegatePropertyHandle = DetailLayout.GetProperty(Property->GetFName(), CastChecked<UClass>(Property->GetOuter()));
 
-				TSharedPtr<SWidget> NameWidget;
-				TSharedPtr<SWidget> ValueWidget;
-				FDetailWidgetRow Row;
-				PropertyRow.GetDefaultWidgets(NameWidget, ValueWidget, Row);
+	const bool bHasValidHandle = DelegatePropertyHandle->IsValidHandle();
+	if(!bHasValidHandle)
+	{
+		return;
+	}
 
-				TSharedRef<SComboButton> DelegateButton = SNew(SComboButton)
-				.OnGetMenuContent(this, &FBlueprintWidgetCustomization::OnGenerateDelegateMenu, ConstPropertyHandle, Property->SignatureFunction)
-				.ContentPadding(0)
-				.ButtonContent()
-				[
-					SNew(SHorizontalBox)
+	IDetailCategoryBuilder& PropertyCategory = DetailLayout.EditCategory(FObjectEditorUtils::GetCategoryFName(Property));
 
-					+ SHorizontalBox::Slot()
-					.AutoWidth()
-					.VAlign(VAlign_Center)
-					[
-						SNew(SImage)
-						.Image(this, &FBlueprintWidgetCustomization::GetCurrentBindingImage, ConstPropertyHandle)
-						.ColorAndOpacity(FLinearColor(0.25f, 0.25f, 0.25f))
-					]
+	IDetailPropertyRow& PropertyRow = PropertyCategory.AddProperty(DelegatePropertyHandle);
 
-					+ SHorizontalBox::Slot()
-					.AutoWidth()
-					.VAlign(VAlign_Center)
-					.Padding(4,1,0,0)
-					[
-						SNew(STextBlock)
-						.Text(this, &FBlueprintWidgetCustomization::GetCurrentBindingText, ConstPropertyHandle)
-						.Font(IDetailLayoutBuilder::GetDetailFont())
-					]
-				];
+	TSharedPtr<SWidget> NameWidget;
+	TSharedPtr<SWidget> ValueWidget;
+	FDetailWidgetRow Row;
+	PropertyRow.GetDefaultWidgets(NameWidget, ValueWidget, Row);
 
-				const bool bShowChildren = true;
-				PropertyRow.CustomWidget(bShowChildren)
-				.NameContent()
-				.MinDesiredWidth(Row.NameWidget.MinWidth)
-				.MaxDesiredWidth(Row.NameWidget.MaxWidth)
-				[
-					NameWidget.ToSharedRef()
-				]
-				.ValueContent()
-				.MinDesiredWidth(Row.ValueWidget.MinWidth * 2)
-				.MaxDesiredWidth(Row.ValueWidget.MaxWidth * 2)
-				[
-					SNew(SHorizontalBox)
+	TSharedRef<SComboButton> DelegateButton = SNew(SComboButton)
+		.OnGetMenuContent(this, &FBlueprintWidgetCustomization::OnGenerateDelegateMenu, DelegatePropertyHandle, Property->SignatureFunction)
+		.ContentPadding(0)
+		.ButtonContent()
+		[
+			SNew(STextBlock)
+			.Text(this, &FBlueprintWidgetCustomization::GetCurrentBindingText, DelegatePropertyHandle)
+			.Font(IDetailLayoutBuilder::GetDetailFont())
+		];
 
-					+ SHorizontalBox::Slot()
-					.FillWidth(1)
-					.Padding(FMargin(0, 0, 4, 0))
-					[
-						ValueWidget.ToSharedRef()
-					]
-					
-					+ SHorizontalBox::Slot()
-					.AutoWidth()
-					[
-						DelegateButton
-					]
+	const bool bShowChildren = true;
+	PropertyRow.CustomWidget(bShowChildren)
+		.NameContent()
+		.MinDesiredWidth(Row.NameWidget.MinWidth)
+		.MaxDesiredWidth(Row.NameWidget.MaxWidth)
+		[
+			NameWidget.ToSharedRef()
+		]
+	.ValueContent()
+		.MinDesiredWidth(Row.ValueWidget.MinWidth)
+		.MaxDesiredWidth(Row.ValueWidget.MaxWidth)
+		[
+			SNew(SHorizontalBox)
+
+			+ SHorizontalBox::Slot()
+			.FillWidth(1)
+			[
+				DelegateButton
+			]
 
-					+ SHorizontalBox::Slot()
-					.AutoWidth()
+			+ SHorizontalBox::Slot()
+				.AutoWidth()
+				[
+					SNew(SButton)
+					.ButtonStyle(FEditorStyle::Get(), "HoverHintOnly")
+					.Visibility(this, &FBlueprintWidgetCustomization::GetGotoBindingVisibility, DelegatePropertyHandle)
+					.OnClicked(this, &FBlueprintWidgetCustomization::HandleGotoBindingClicked, DelegatePropertyHandle)
+					.VAlign(VAlign_Center)
+					.ToolTipText(LOCTEXT("GotoFunction", "Goto Function"))
 					[
-						SNew(SButton)
-						.ButtonStyle( FEditorStyle::Get(), "HoverHintOnly" )
-						.Visibility( this, &FBlueprintWidgetCustomization::GetGotoBindingVisibility, ConstPropertyHandle )
-						.OnClicked( this, &FBlueprintWidgetCustomization::HandleGotoBindingClicked, ConstPropertyHandle )
-						.VAlign(VAlign_Center)
-						.ToolTipText(LOCTEXT("GotoFunction", "Goto Function"))
-						[
-							SNew(SImage)
-							.Image( FEditorStyle::GetBrush("PropertyWindow.Button_Browse") )
-						]
+						SNew(SImage)
+						.Image(FEditorStyle::GetBrush("PropertyWindow.Button_Browse"))
 					]
-				];
-			}
-		}
+				]
+		];
+}
+
+void FBlueprintWidgetCustomization::CustomizeDetails( IDetailLayoutBuilder& DetailLayout )
+{
+	TArray< TWeakObjectPtr<UObject> > OutObjects;
+	DetailLayout.GetObjectsBeingCustomized(OutObjects);
 
-		//--------------------------------------------
+	if ( OutObjects.Num() == 1 )
+	{
+		UClass* PropertyClass = OutObjects[0].Get()->GetClass();
 
 		for ( TFieldIterator<UProperty> PropertyIt(PropertyClass, EFieldIteratorFlags::IncludeSuper); PropertyIt; ++PropertyIt )
 		{
-			if ( UDelegateProperty* Property = Cast<UDelegateProperty>(*PropertyIt) )
+			UProperty* Property = *PropertyIt;
+			if ( UDelegateProperty* DelegateProperty = Cast<UDelegateProperty>(*PropertyIt) )
 			{
-				if ( !Property->GetName().EndsWith(TEXT("Event")) )
+				if( DelegateProperty->GetName().EndsWith(TEXT("Delegate")) )
 				{
-					continue;
+					CreateDelegateCustomization( DetailLayout, DelegateProperty );
 				}
-
-				TSharedRef<IPropertyHandle> DelegatePropertyHandle = DetailLayout.GetProperty(Property->GetFName(), CastChecked<UClass>(Property->GetOuter()));
-
-				const bool bHasValidHandle = DelegatePropertyHandle->IsValidHandle();
-
-				if ( !bHasValidHandle )
+				else if( DelegateProperty->GetName().EndsWith(TEXT("Event")) )
 				{
-					continue;
+					CreateEventCustomization( DetailLayout, DelegateProperty );
 				}
-
-				IDetailCategoryBuilder& PropertyCategory = DetailLayout.EditCategory(FObjectEditorUtils::GetCategoryFName(Property));
-
-				IDetailPropertyRow& PropertyRow = PropertyCategory.AddProperty(DelegatePropertyHandle);
-
-				TSharedPtr<SWidget> NameWidget;
-				TSharedPtr<SWidget> ValueWidget;
-				FDetailWidgetRow Row;
-				PropertyRow.GetDefaultWidgets(NameWidget, ValueWidget, Row);
-
-				TSharedRef<SComboButton> DelegateButton = SNew(SComboButton)
-				.OnGetMenuContent(this, &FBlueprintWidgetCustomization::OnGenerateDelegateMenu, DelegatePropertyHandle, Property->SignatureFunction)
-				.ContentPadding(0)
-				.ButtonContent()
-				[
-					SNew(STextBlock)
-					.Text(this, &FBlueprintWidgetCustomization::GetCurrentBindingText, DelegatePropertyHandle)
-					.Font(IDetailLayoutBuilder::GetDetailFont())
-				];
-
-				const bool bShowChildren = true;
-				PropertyRow.CustomWidget(bShowChildren)
-				.NameContent()
-				.MinDesiredWidth(Row.NameWidget.MinWidth)
-				.MaxDesiredWidth(Row.NameWidget.MaxWidth)
-				[
-					NameWidget.ToSharedRef()
-				]
-				.ValueContent()
-				.MinDesiredWidth(Row.ValueWidget.MinWidth)
-				.MaxDesiredWidth(Row.ValueWidget.MaxWidth)
-				[
-					SNew(SHorizontalBox)
-					
-					+ SHorizontalBox::Slot()
-					.FillWidth(1)
-					[
-						DelegateButton
-					]
-
-					+ SHorizontalBox::Slot()
-					.AutoWidth()
-					[
-						SNew(SButton)
-						.ButtonStyle( FEditorStyle::Get(), "HoverHintOnly" )
-						.Visibility( this, &FBlueprintWidgetCustomization::GetGotoBindingVisibility, DelegatePropertyHandle )
-						.OnClicked( this, &FBlueprintWidgetCustomization::HandleGotoBindingClicked, DelegatePropertyHandle )
-						.VAlign(VAlign_Center)
-						.ToolTipText(LOCTEXT("GotoFunction", "Goto Function"))
-						[
-							SNew(SImage)
-							.Image( FEditorStyle::GetBrush("PropertyWindow.Button_Browse") )
-						]
-					]
-				];
 			}
 		}
-
 	}
 }
 
diff --git a/Engine/Source/Editor/UMGEditor/Private/Customizations/DetailCustomizations.h b/Engine/Source/Editor/UMGEditor/Private/Customizations/DetailCustomizations.h
--- a/Engine/Source/Editor/UMGEditor/Private/Customizations/DetailCustomizations.h
+++ b/Engine/Source/Editor/UMGEditor/Private/Customizations/DetailCustomizations.h
@@ -50,6 +50,9 @@ private:
 
 	FReply HandleGotoBindingClicked(TSharedRef<IPropertyHandle> PropertyHandle);
 
+	void CreateDelegateCustomization( IDetailLayoutBuilder& DetailLayout, UDelegateProperty* Property );
+
+	void CreateEventCustomization( IDetailLayoutBuilder& DetailLayout, UDelegateProperty* Property );
 private:
 
 	TWeakPtr<FWidgetBlueprintEditor> Editor;
diff --git a/Engine/Source/Editor/UMGEditor/Private/Details/SWidgetDetailsView.cpp b/Engine/Source/Editor/UMGEditor/Private/Details/SWidgetDetailsView.cpp
--- a/Engine/Source/Editor/UMGEditor/Private/Details/SWidgetDetailsView.cpp
+++ b/Engine/Source/Editor/UMGEditor/Private/Details/SWidgetDetailsView.cpp
@@ -24,6 +24,10 @@ void SWidgetDetailsView::Construct(const FArguments& InArgs, TSharedPtr<FWidgetB
 
 	PropertyView = EditModule.CreateDetailView(DetailsViewArgs);
 
+	// Create a handler for keyframing via the details panel
+	TSharedRef<IDetailKeyframeHandler> KeyframeHandler = MakeShareable( new FUMGDetailKeyframeHandler( InBlueprintEditor ) );
+	PropertyView->SetKeyframeHandler( KeyframeHandler );
+
 	ChildSlot
 	[
 		SNew(SVerticalBox)
@@ -303,15 +307,19 @@ void SWidgetDetailsView::HandleIsVariableChanged(ESlateCheckBoxState::Type Check
 
 void SWidgetDetailsView::NotifyPreChange(FEditPropertyChain* PropertyAboutToChange)
 {
-	TSharedPtr<FWidgetBlueprintEditor> Editor = BlueprintEditor.Pin();
+	// During auto-key do not migrate values
+	if( !BlueprintEditor.Pin()->GetSequencer()->IsAutoKeyEnabled() )
+	{
+		TSharedPtr<FWidgetBlueprintEditor> Editor = BlueprintEditor.Pin();
 
-	const bool bIsModify = true;
-	Editor->MigrateFromChain(PropertyAboutToChange, bIsModify);
+		const bool bIsModify = true;
+		Editor->MigrateFromChain(PropertyAboutToChange, bIsModify);
+	}
 }
 
 void SWidgetDetailsView::NotifyPostChange(const FPropertyChangedEvent& PropertyChangedEvent, FEditPropertyChain* PropertyThatChanged)
 {
-	if ( PropertyChangedEvent.ChangeType != EPropertyChangeType::Interactive )
+	if ( PropertyChangedEvent.ChangeType != EPropertyChangeType::Interactive && !BlueprintEditor.Pin()->GetSequencer()->IsAutoKeyEnabled() )
 	{
 		TSharedPtr<FWidgetBlueprintEditor> Editor = BlueprintEditor.Pin();
 
diff --git a/Engine/Source/Editor/UMGEditor/Private/UMGSequencerObjectBindingManager.cpp b/Engine/Source/Editor/UMGEditor/Private/UMGSequencerObjectBindingManager.cpp
deleted file mode 100644
diff --git a/Engine/Source/Editor/UMGEditor/Private/UMGSequencerObjectBindingManager.h b/Engine/Source/Editor/UMGEditor/Private/UMGSequencerObjectBindingManager.h
deleted file mode 100644
diff --git a/Engine/Source/Editor/UMGEditor/Private/WidgetBlueprintEditor.cpp b/Engine/Source/Editor/UMGEditor/Private/WidgetBlueprintEditor.cpp
--- a/Engine/Source/Editor/UMGEditor/Private/WidgetBlueprintEditor.cpp
+++ b/Engine/Source/Editor/UMGEditor/Private/WidgetBlueprintEditor.cpp
@@ -6,7 +6,7 @@
 #include "WidgetBlueprintEditor.h"
 #include "MovieScene.h"
 #include "Editor/Sequencer/Public/ISequencerModule.h"
-#include "UMGSequencerObjectBindingManager.h"
+#include "Animation/UMGSequencerObjectBindingManager.h"
 
 #include "WidgetBlueprintApplicationModes.h"
 //#include "WidgetDefafaultsApplicationMode.h"
diff --git a/Engine/Source/Runtime/MovieSceneCore/Classes/MovieSceneSection.h b/Engine/Source/Runtime/MovieSceneCore/Classes/MovieSceneSection.h
--- a/Engine/Source/Runtime/MovieSceneCore/Classes/MovieSceneSection.h
+++ b/Engine/Source/Runtime/MovieSceneCore/Classes/MovieSceneSection.h
@@ -56,6 +56,8 @@ public:
 	{
 		check(NewRange.HasLowerBound() && NewRange.HasUpperBound());
 
+		Modify();
+
 		StartTime = NewRange.GetLowerBoundValue();
 		EndTime = NewRange.GetUpperBoundValue();
 	}
@@ -78,6 +80,8 @@ public:
 	 */
 	virtual void MoveSection( float DeltaTime )
 	{
+		Modify();
+
 		StartTime += DeltaTime;
 		EndTime += DeltaTime;
 	}
@@ -90,6 +94,8 @@ public:
 	 */
 	virtual void DilateSection( float DilationFactor, float Origin )
 	{
+		Modify();
+
 		StartTime = (StartTime - Origin) * DilationFactor + Origin;
 		EndTime = (EndTime - Origin) * DilationFactor + Origin;
 	}
diff --git a/Engine/Source/Runtime/MovieSceneCoreTypes/Classes/MovieSceneColorSection.h b/Engine/Source/Runtime/MovieSceneCoreTypes/Classes/MovieSceneColorSection.h
--- a/Engine/Source/Runtime/MovieSceneCoreTypes/Classes/MovieSceneColorSection.h
+++ b/Engine/Source/Runtime/MovieSceneCoreTypes/Classes/MovieSceneColorSection.h
@@ -4,6 +4,8 @@
 
 #include "MovieSceneColorSection.generated.h"
 
+struct FColorKey;
+
 /**
  * A single floating point section
  */
@@ -29,7 +31,7 @@ public:
 	 * @param Time	The location in time where the key should be added
 	 * @param Value	The value of the key
 	 */
-	void AddKey( float Time, FLinearColor Value );
+	void AddKey( float Time, const FColorKey& Key );
 	
 	/** 
 	 * Determines if a new key would be new data, or just a duplicate of existing data
@@ -72,6 +74,8 @@ public:
 	const FRichCurve& GetAlphaCurve() const { return AlphaCurve; }
 
 private:
+	void AddKeyToNamedCurve(float Time, const FColorKey& Key);
+
 	/**
 	 * Adds a key to a rich curve, finding an existing key to modify or adding a new one
 	 *
diff --git a/Engine/Source/Runtime/MovieSceneCoreTypes/Classes/MovieSceneColorTrack.h b/Engine/Source/Runtime/MovieSceneCoreTypes/Classes/MovieSceneColorTrack.h
--- a/Engine/Source/Runtime/MovieSceneCoreTypes/Classes/MovieSceneColorTrack.h
+++ b/Engine/Source/Runtime/MovieSceneCoreTypes/Classes/MovieSceneColorTrack.h
@@ -7,6 +7,13 @@
 #include "MovieScenePropertyTrack.h"
 #include "MovieSceneColorTrack.generated.h"
 
+struct FColorKey
+{
+	FLinearColor Value;
+	FName CurveName;
+	bool bAddKeyEvenIfUnchanged;
+};
+
 /**
  * Handles manipulation of float properties in a movie scene
  */
@@ -27,7 +34,7 @@ public:
 	 * @param Value				The value of the key
 	 * @return True if the key was successfully added.
 	 */
-	virtual bool AddKeyToSection( float Time, FLinearColor Value );
+	virtual bool AddKeyToSection( float Time, const FColorKey& Key );
 	
 	/**
 	 * Evaluates the track at the playback position
diff --git a/Engine/Source/Runtime/MovieSceneCoreTypes/Classes/MovieSceneVectorSection.h b/Engine/Source/Runtime/MovieSceneCoreTypes/Classes/MovieSceneVectorSection.h
--- a/Engine/Source/Runtime/MovieSceneCoreTypes/Classes/MovieSceneVectorSection.h
+++ b/Engine/Source/Runtime/MovieSceneCoreTypes/Classes/MovieSceneVectorSection.h
@@ -32,7 +32,7 @@ public:
 	 * @param Time	The location in time where the key should be added
 	 * @param Value	The value of the key
 	 */
-	void AddKey( float Time, const FVector4& Value );
+	void AddKey( float Time, FName CurveName, const FVector4& Value );
 	
 	/** 
 	 * Determines if a new key would be new data, or just a duplicate of existing data
@@ -52,6 +52,9 @@ public:
 	/** Gets the number of channels in use */
 	int32 GetChannelsUsed() const {return ChannelsUsed;}
 
+private:
+	void AddKeyToNamedCurve(float Time, FName CurveName, const FVector4& Value);
+
 private:
 	/** Vector t */
 	UPROPERTY()
diff --git a/Engine/Source/Runtime/MovieSceneCoreTypes/Classes/MovieSceneVectorTrack.h b/Engine/Source/Runtime/MovieSceneCoreTypes/Classes/MovieSceneVectorTrack.h
--- a/Engine/Source/Runtime/MovieSceneCoreTypes/Classes/MovieSceneVectorTrack.h
+++ b/Engine/Source/Runtime/MovieSceneCoreTypes/Classes/MovieSceneVectorTrack.h
@@ -9,6 +9,14 @@
 
 #include "MovieSceneVectorTrack.generated.h"
 
+template<typename VecType>
+struct FVectorKey
+{
+	VecType Value;
+	FName CurveName;
+	bool bAddKeyEvenIfUnchanged;
+};
+
 /**
  * Handles manipulation of component transforms in a movie scene
  */
@@ -31,10 +39,10 @@ public:
 	 * @param InChannelsUsed	The number of channels used, 2, 3, or 4, determining which type of vector this is
 	 * @return True if the key was successfully added.
 	 */
-	virtual bool AddKeyToSection( float Time, const FVector4& InKey, int32 InChannelsUsed );
-	virtual bool AddKeyToSection( float Time, const FVector4& InKey);
-	virtual bool AddKeyToSection( float Time, const FVector& InKey);
-	virtual bool AddKeyToSection( float Time, const FVector2D& InKey);
+
+	virtual bool AddKeyToSection( float Time, const FVectorKey<FVector4>& Key );
+	virtual bool AddKeyToSection( float Time, const FVectorKey<FVector>& Key );
+	virtual bool AddKeyToSection( float Time, const FVectorKey<FVector2D>& Key );
 	
 	/**
 	 * Evaluates the track at the playback position
@@ -48,6 +56,9 @@ public:
 
 	/** @return Get the number of channels used by the vector */
 	int32 GetNumChannelsUsed() const { return NumChannelsUsed; }
+
+private:
+	virtual bool AddKeyToSection( float Time, const FVector4& InKey, int32 InChannelsUsed, FName CurveName, bool bAddKeyEvenIfUnchanged );
 private:
 	/** The number of channels used by the vector (2,3, or 4) */
 	UPROPERTY()
diff --git a/Engine/Source/Runtime/MovieSceneCoreTypes/Private/MovieSceneBoolSection.cpp b/Engine/Source/Runtime/MovieSceneCoreTypes/Private/MovieSceneBoolSection.cpp
--- a/Engine/Source/Runtime/MovieSceneCoreTypes/Private/MovieSceneBoolSection.cpp
+++ b/Engine/Source/Runtime/MovieSceneCoreTypes/Private/MovieSceneBoolSection.cpp
@@ -23,13 +23,15 @@ void UMovieSceneBoolSection::MoveSection( float DeltaPosition )
 
 void UMovieSceneBoolSection::DilateSection( float DilationFactor, float Origin )
 {
+	Super::DilateSection(DilationFactor, Origin);
+	
 	BoolCurve.ScaleCurve(Origin, DilationFactor);
 
-	Super::DilateSection(DilationFactor, Origin);
 }
 
 void UMovieSceneBoolSection::AddKey( float Time, bool Value )
 {
+	Modify();
 	BoolCurve.UpdateOrAddKey(Time, Value ? 1 : 0);
 }
 
diff --git a/Engine/Source/Runtime/MovieSceneCoreTypes/Private/MovieSceneBoolTrack.cpp b/Engine/Source/Runtime/MovieSceneCoreTypes/Private/MovieSceneBoolTrack.cpp
--- a/Engine/Source/Runtime/MovieSceneCoreTypes/Private/MovieSceneBoolTrack.cpp
+++ b/Engine/Source/Runtime/MovieSceneCoreTypes/Private/MovieSceneBoolTrack.cpp
@@ -14,7 +14,7 @@ UMovieSceneBoolTrack::UMovieSceneBoolTrack( const FPostConstructInitializeProper
 
 UMovieSceneSection* UMovieSceneBoolTrack::CreateNewSection()
 {
-	return ConstructObject<UMovieSceneSection>( UMovieSceneBoolSection::StaticClass(), this );
+	return ConstructObject<UMovieSceneSection>( UMovieSceneBoolSection::StaticClass(), this, NAME_None, RF_Transactional );
 }
 
 TSharedPtr<IMovieSceneTrackInstance> UMovieSceneBoolTrack::CreateInstance()
@@ -27,8 +27,10 @@ bool UMovieSceneBoolTrack::AddKeyToSection( float Time, bool Value )
 	const UMovieSceneSection* NearestSection = MovieSceneHelpers::FindSectionAtTime( Sections, Time );
 	if (!NearestSection || CastChecked<UMovieSceneBoolSection>(NearestSection)->NewKeyIsNewData(Time, Value))
 	{
-		UMovieSceneBoolSection* NewSection = CastChecked<UMovieSceneBoolSection>(FindOrAddSection(  Time ));
+		Modify();
 
+		UMovieSceneBoolSection* NewSection = CastChecked<UMovieSceneBoolSection>(FindOrAddSection(  Time ));
+	
 		NewSection->AddKey( Time, Value );
 
 		return true;
diff --git a/Engine/Source/Runtime/MovieSceneCoreTypes/Private/MovieSceneColorSection.cpp b/Engine/Source/Runtime/MovieSceneCoreTypes/Private/MovieSceneColorSection.cpp
--- a/Engine/Source/Runtime/MovieSceneCoreTypes/Private/MovieSceneColorSection.cpp
+++ b/Engine/Source/Runtime/MovieSceneCoreTypes/Private/MovieSceneColorSection.cpp
@@ -22,12 +22,13 @@ void UMovieSceneColorSection::MoveSection( float DeltaTime )
 
 void UMovieSceneColorSection::DilateSection( float DilationFactor, float Origin )
 {
+	Super::DilateSection(DilationFactor, Origin);
+
 	RedCurve.ScaleCurve(Origin, DilationFactor);
 	GreenCurve.ScaleCurve(Origin, DilationFactor);
 	BlueCurve.ScaleCurve(Origin, DilationFactor);
 	AlphaCurve.ScaleCurve(Origin, DilationFactor);
 
-	Super::DilateSection(DilationFactor, Origin);
 }
 
 FLinearColor UMovieSceneColorSection::Eval( float Position ) const
@@ -38,12 +39,21 @@ FLinearColor UMovieSceneColorSection::Eval( float Position ) const
 						AlphaCurve.Eval(Position));
 }
 
-void UMovieSceneColorSection::AddKey( float Time, FLinearColor Value )
+void UMovieSceneColorSection::AddKey( float Time, const FColorKey& Key )
 {
-	AddKeyToCurve(RedCurve, Time, Value.R);
-	AddKeyToCurve(GreenCurve, Time, Value.G);
-	AddKeyToCurve(BlueCurve, Time, Value.B);
-	AddKeyToCurve(AlphaCurve, Time, Value.A);
+	Modify();
+
+	if( Key.CurveName == NAME_None )
+	{
+		AddKeyToCurve(RedCurve, Time, Key.Value.R);
+		AddKeyToCurve(GreenCurve, Time, Key.Value.G);
+		AddKeyToCurve(BlueCurve, Time, Key.Value.B);
+		AddKeyToCurve(AlphaCurve, Time, Key.Value.A);
+	}
+	else
+	{
+		AddKeyToNamedCurve( Time, Key );
+	}
 }
 
 bool UMovieSceneColorSection::NewKeyIsNewData(float Time, FLinearColor Value) const
@@ -55,12 +65,38 @@ bool UMovieSceneColorSection::NewKeyIsNewData(float Time, FLinearColor Value) co
 		!Eval(Time).Equals(Value);
 }
 
+void UMovieSceneColorSection::AddKeyToNamedCurve(float Time, const FColorKey& Key)
+{
+	static FName R("R");
+	static FName G("G");
+	static FName B("B");
+	static FName A("A");
 
+	FName CurveName = Key.CurveName;
+	if (CurveName == R)
+	{
+		AddKeyToCurve(RedCurve, Time, Key.Value.R);
+	}
+	else if (CurveName == G)
+	{
+		AddKeyToCurve(GreenCurve, Time, Key.Value.G);
+	}
+	else if (CurveName == B)
+	{
+		AddKeyToCurve(BlueCurve, Time, Key.Value.B);
+	}
+	else if (CurveName == A)
+	{
+		AddKeyToCurve(AlphaCurve, Time, Key.Value.A);
+	}
+}
 
 void UMovieSceneColorSection::AddKeyToCurve( FRichCurve& InCurve, float Time, float Value )
 {
 	if( IsTimeWithinSection(Time) )
 	{
+		Modify();
+
 		InCurve.UpdateOrAddKey(Time, Value);
 	}
 }
diff --git a/Engine/Source/Runtime/MovieSceneCoreTypes/Private/MovieSceneColorTrack.cpp b/Engine/Source/Runtime/MovieSceneCoreTypes/Private/MovieSceneColorTrack.cpp
--- a/Engine/Source/Runtime/MovieSceneCoreTypes/Private/MovieSceneColorTrack.cpp
+++ b/Engine/Source/Runtime/MovieSceneCoreTypes/Private/MovieSceneColorTrack.cpp
@@ -13,7 +13,7 @@ UMovieSceneColorTrack::UMovieSceneColorTrack( const FPostConstructInitializeProp
 
 UMovieSceneSection* UMovieSceneColorTrack::CreateNewSection()
 {
-	return ConstructObject<UMovieSceneSection>( UMovieSceneColorSection::StaticClass(), this );
+	return ConstructObject<UMovieSceneSection>( UMovieSceneColorSection::StaticClass(), this, NAME_None, RF_Transactional );
 }
 
 TSharedPtr<IMovieSceneTrackInstance> UMovieSceneColorTrack::CreateInstance()
@@ -22,14 +22,15 @@ TSharedPtr<IMovieSceneTrackInstance> UMovieSceneColorTrack::CreateInstance()
 }
 
 
-bool UMovieSceneColorTrack::AddKeyToSection( float Time, FLinearColor Value )
+bool UMovieSceneColorTrack::AddKeyToSection( float Time, const FColorKey& Key )
 {
 	const UMovieSceneSection* NearestSection = MovieSceneHelpers::FindSectionAtTime( Sections, Time );
-	if (!NearestSection || CastChecked<UMovieSceneColorSection>(NearestSection)->NewKeyIsNewData(Time, Value))
+	if (!NearestSection || Key.bAddKeyEvenIfUnchanged || CastChecked<UMovieSceneColorSection>(NearestSection)->NewKeyIsNewData(Time, Key.Value))
 	{
+		Modify();
 		UMovieSceneColorSection* NewSection = CastChecked<UMovieSceneColorSection>( FindOrAddSection( Time ) );
 
-		NewSection->AddKey( Time, Value );
+		NewSection->AddKey( Time, Key );
 
 		return true;
 	}
diff --git a/Engine/Source/Runtime/MovieSceneCoreTypes/Private/MovieSceneFloatSection.cpp b/Engine/Source/Runtime/MovieSceneCoreTypes/Private/MovieSceneFloatSection.cpp
--- a/Engine/Source/Runtime/MovieSceneCoreTypes/Private/MovieSceneFloatSection.cpp
+++ b/Engine/Source/Runtime/MovieSceneCoreTypes/Private/MovieSceneFloatSection.cpp
@@ -23,14 +23,15 @@ void UMovieSceneFloatSection::MoveSection( float DeltaPosition )
 }
 
 void UMovieSceneFloatSection::DilateSection( float DilationFactor, float Origin )
-{
-	FloatCurve.ScaleCurve(Origin, DilationFactor);
-
+{	
 	Super::DilateSection(DilationFactor, Origin);
+	
+	FloatCurve.ScaleCurve(Origin, DilationFactor);
 }
 
 void UMovieSceneFloatSection::AddKey( float Time, float Value )
 {
+	Modify();
 	FloatCurve.UpdateOrAddKey(Time, Value);
 }
 
diff --git a/Engine/Source/Runtime/MovieSceneCoreTypes/Private/MovieSceneFloatTrack.cpp b/Engine/Source/Runtime/MovieSceneCoreTypes/Private/MovieSceneFloatTrack.cpp
--- a/Engine/Source/Runtime/MovieSceneCoreTypes/Private/MovieSceneFloatTrack.cpp
+++ b/Engine/Source/Runtime/MovieSceneCoreTypes/Private/MovieSceneFloatTrack.cpp
@@ -14,7 +14,7 @@ UMovieSceneFloatTrack::UMovieSceneFloatTrack( const FPostConstructInitializeProp
 
 UMovieSceneSection* UMovieSceneFloatTrack::CreateNewSection()
 {
-	return ConstructObject<UMovieSceneSection>( UMovieSceneFloatSection::StaticClass(), this );
+	return ConstructObject<UMovieSceneSection>( UMovieSceneFloatSection::StaticClass(), this, NAME_None, RF_Transactional );
 }
 
 TSharedPtr<IMovieSceneTrackInstance> UMovieSceneFloatTrack::CreateInstance()
@@ -28,6 +28,8 @@ bool UMovieSceneFloatTrack::AddKeyToSection( float Time, float Value )
 	const UMovieSceneSection* NearestSection = MovieSceneHelpers::FindSectionAtTime( Sections, Time );
 	if (!NearestSection || CastChecked<UMovieSceneFloatSection>(NearestSection)->NewKeyIsNewData(Time, Value))
 	{
+		Modify();
+
 		UMovieSceneFloatSection* NewSection = Cast<UMovieSceneFloatSection>( FindOrAddSection( Time ) );
 
 		NewSection->AddKey( Time, Value );
diff --git a/Engine/Source/Runtime/MovieSceneCoreTypes/Private/MovieSceneTransformSection.cpp b/Engine/Source/Runtime/MovieSceneCoreTypes/Private/MovieSceneTransformSection.cpp
--- a/Engine/Source/Runtime/MovieSceneCoreTypes/Private/MovieSceneTransformSection.cpp
+++ b/Engine/Source/Runtime/MovieSceneCoreTypes/Private/MovieSceneTransformSection.cpp
@@ -25,16 +25,18 @@ void UMovieSceneTransformSection::MoveSection( float DeltaTime )
 
 void UMovieSceneTransformSection::DilateSection( float DilationFactor, float Origin )
 {
+	Super::DilateSection(DilationFactor, Origin);
+
 	for( int32 Axis = 0; Axis < 3; ++Axis )
 	{
 		Translation[Axis].ScaleCurve( Origin, DilationFactor );
 		Rotation[Axis].ScaleCurve( Origin, DilationFactor );
 		Scale[Axis].ScaleCurve( Origin, DilationFactor );
 	}
-
-	Super::DilateSection(DilationFactor, Origin);
 }
 
+	
+
 bool UMovieSceneTransformSection::EvalTranslation( float Time, FVector& OutTranslation ) const
 {
 	bool bHasTranslationKeys = Translation[0].GetNumKeys() > 0 || Translation[1].GetNumKeys() > 0 || Translation[2].GetNumKeys() > 0;
@@ -124,6 +126,7 @@ void UMovieSceneTransformSection::AddKeyToCurve( FRichCurve& InCurve, float Time
 {
 	if( IsTimeWithinSection(Time) )
 	{
+		Modify();
 		InCurve.UpdateOrAddKey(Time, Value);
 	}
 }
diff --git a/Engine/Source/Runtime/MovieSceneCoreTypes/Private/MovieSceneTransformTrack.cpp b/Engine/Source/Runtime/MovieSceneCoreTypes/Private/MovieSceneTransformTrack.cpp
--- a/Engine/Source/Runtime/MovieSceneCoreTypes/Private/MovieSceneTransformTrack.cpp
+++ b/Engine/Source/Runtime/MovieSceneCoreTypes/Private/MovieSceneTransformTrack.cpp
@@ -64,7 +64,7 @@ UMovieSceneTransformTrack::UMovieSceneTransformTrack( const FPostConstructInitia
 
 UMovieSceneSection* UMovieSceneTransformTrack::CreateNewSection()
 {
-	return ConstructObject<UMovieSceneSection>( UMovieSceneTransformSection::StaticClass(), this );
+	return ConstructObject<UMovieSceneSection>( UMovieSceneTransformSection::StaticClass(), this, NAME_None, RF_Transactional );
 }
 
 TSharedPtr<IMovieSceneTrackInstance> UMovieSceneTransformTrack::CreateInstance()
@@ -77,6 +77,8 @@ bool UMovieSceneTransformTrack::AddKeyToSection( const FGuid& ObjectHandle, cons
 	const UMovieSceneSection* NearestSection = MovieSceneHelpers::FindSectionAtTime(Sections, InKey.GetKeyTime());
 	if (!NearestSection || CastChecked<UMovieSceneTransformSection>(NearestSection)->NewKeyIsNewData(InKey))
 	{
+		Modify();
+
 		UMovieSceneTransformSection* NewSection = Cast<UMovieSceneTransformSection>( FindOrAddSection( InKey.GetKeyTime() ) );
 
 		// key each component of the transform
diff --git a/Engine/Source/Runtime/MovieSceneCoreTypes/Private/MovieSceneVectorSection.cpp b/Engine/Source/Runtime/MovieSceneCoreTypes/Private/MovieSceneVectorSection.cpp
--- a/Engine/Source/Runtime/MovieSceneCoreTypes/Private/MovieSceneVectorSection.cpp
+++ b/Engine/Source/Runtime/MovieSceneCoreTypes/Private/MovieSceneVectorSection.cpp
@@ -19,13 +19,22 @@ FVector4 UMovieSceneVectorSection::Eval( float Position ) const
 		Curves[3].Eval( Position ) );
 }
 
-void UMovieSceneVectorSection::AddKey( float Time, const FVector4& Value )
+void UMovieSceneVectorSection::AddKey( float Time, FName CurveName, const FVector4& Value )
 {
-	check(ChannelsUsed >= 2 && ChannelsUsed <= 4);
+	Modify();
 
-	for (int32 i = 0; i < ChannelsUsed; ++i)
+	if( CurveName == NAME_None )
 	{
-		Curves[i].UpdateOrAddKey(Time, Value[i]);
+		check(ChannelsUsed >= 2 && ChannelsUsed <= 4);
+
+		for (int32 i = 0; i < ChannelsUsed; ++i)
+		{
+			Curves[i].UpdateOrAddKey(Time, Value[i]);
+		}
+	}
+	else
+	{
+		AddKeyToNamedCurve( Time, CurveName, Value );
 	}
 }
 
@@ -47,6 +56,31 @@ bool UMovieSceneVectorSection::NewKeyIsNewData(float Time, const FVector4& Value
 	return bNewData;
 }
 
+void UMovieSceneVectorSection::AddKeyToNamedCurve(float Time, FName CurveName, const FVector4& Value)
+{
+	static FName X("X");
+	static FName Y("Y");
+	static FName Z("Z");
+	static FName W("W");
+
+	if (CurveName == X)
+	{
+		Curves[0].UpdateOrAddKey(Time, Value.X);
+	}
+	else if (CurveName == Y)
+	{
+		Curves[1].UpdateOrAddKey(Time, Value.Y);
+	}
+	else if (CurveName == Z)
+	{
+		Curves[2].UpdateOrAddKey(Time, Value.Z);
+	}
+	else if (CurveName == W)
+	{
+		Curves[3].UpdateOrAddKey(Time, Value.W);
+	}
+}
+
 void UMovieSceneVectorSection::MoveSection( float DeltaTime )
 {
 	check(ChannelsUsed >= 2 && ChannelsUsed <= 4);
@@ -62,11 +96,10 @@ void UMovieSceneVectorSection::MoveSection( float DeltaTime )
 void UMovieSceneVectorSection::DilateSection( float DilationFactor, float Origin )
 {
 	check(ChannelsUsed >= 2 && ChannelsUsed <= 4);
+	Super::DilateSection(DilationFactor, Origin);
 
 	for (int32 i = 0; i < ChannelsUsed; ++i)
 	{
 		Curves[i].ScaleCurve(Origin, DilationFactor);
 	}
-
-	Super::DilateSection(DilationFactor, Origin);
 }
diff --git a/Engine/Source/Runtime/MovieSceneCoreTypes/Private/MovieSceneVectorTrack.cpp b/Engine/Source/Runtime/MovieSceneCoreTypes/Private/MovieSceneVectorTrack.cpp
--- a/Engine/Source/Runtime/MovieSceneCoreTypes/Private/MovieSceneVectorTrack.cpp
+++ b/Engine/Source/Runtime/MovieSceneCoreTypes/Private/MovieSceneVectorTrack.cpp
@@ -22,16 +22,18 @@ TSharedPtr<IMovieSceneTrackInstance> UMovieSceneVectorTrack::CreateInstance()
 	return MakeShareable( new FMovieSceneVectorTrackInstance( *this ) );
 }
 
-bool UMovieSceneVectorTrack::AddKeyToSection( float Time, const FVector4& Value, int32 InChannelsUsed )
+bool UMovieSceneVectorTrack::AddKeyToSection( float Time, const FVector4& Value, int32 InChannelsUsed, FName CurveName, bool bAddKeyEvenIfUnchanged )
 {
 	const UMovieSceneSection* NearestSection = MovieSceneHelpers::FindSectionAtTime( Sections, Time );
-	if (!NearestSection || CastChecked<UMovieSceneVectorSection>(NearestSection)->NewKeyIsNewData(Time, Value))
+	if (!NearestSection || bAddKeyEvenIfUnchanged || CastChecked<UMovieSceneVectorSection>(NearestSection)->NewKeyIsNewData(Time, Value))
 	{
+		Modify();
+
 		UMovieSceneVectorSection* NewSection = Cast<UMovieSceneVectorSection>( FindOrAddSection( Time ) );
 		// @todo Sequencer - I don't like setting the channels used here. It should only be checked here, and set on section creation
 		NewSection->SetChannelsUsed(InChannelsUsed);
 
-		NewSection->AddKey( Time, Value );
+		NewSection->AddKey( Time, CurveName, Value );
 
 		// We dont support one track containing multiple sections of differing channel amounts
 		NumChannelsUsed = InChannelsUsed;
@@ -41,19 +43,19 @@ bool UMovieSceneVectorTrack::AddKeyToSection( float Time, const FVector4& Value,
 	return false;
 }
 
-bool UMovieSceneVectorTrack::AddKeyToSection( float Time, const FVector4& Value )
+bool UMovieSceneVectorTrack::AddKeyToSection( float Time, const FVectorKey<FVector4>& Key )
 {
-	return AddKeyToSection(Time, Value, 4);
+	return AddKeyToSection(Time, Key.Value, 4, Key.CurveName, Key.bAddKeyEvenIfUnchanged );
 }
 
-bool UMovieSceneVectorTrack::AddKeyToSection( float Time, const FVector& Value )
+bool UMovieSceneVectorTrack::AddKeyToSection( float Time, const FVectorKey<FVector>& Key )
 {
-	return AddKeyToSection(Time, FVector4(Value.X, Value.Y, Value.Z, 0.f), 3);
+	return AddKeyToSection(Time, FVector4(Key.Value.X, Key.Value.Y, Key.Value.Z, 0.f), 3, Key.CurveName, Key.bAddKeyEvenIfUnchanged );
 }
 
-bool UMovieSceneVectorTrack::AddKeyToSection( float Time, const FVector2D& Value )
+bool UMovieSceneVectorTrack::AddKeyToSection( float Time, const FVectorKey<FVector2D>& Key )
 {
-	return AddKeyToSection(Time, FVector4(Value.X, Value.Y, 0.f, 0.f), 2);
+	return AddKeyToSection(Time, FVector4(Key.Value.X, Key.Value.Y, 0.f, 0.f), 2, Key.CurveName, Key.bAddKeyEvenIfUnchanged );
 }
 
 bool UMovieSceneVectorTrack::Eval( float Position, float LastPosition, FVector4& OutVector ) const
@@ -67,3 +69,4 @@ bool UMovieSceneVectorTrack::Eval( float Position, float LastPosition, FVector4&
 
 	return Section != NULL;
 }
+
diff --git a/Engine/Source/Runtime/MovieSceneCoreTypes/Private/SubMovieSceneTrack.cpp b/Engine/Source/Runtime/MovieSceneCoreTypes/Private/SubMovieSceneTrack.cpp
--- a/Engine/Source/Runtime/MovieSceneCoreTypes/Private/SubMovieSceneTrack.cpp
+++ b/Engine/Source/Runtime/MovieSceneCoreTypes/Private/SubMovieSceneTrack.cpp
@@ -13,7 +13,7 @@ USubMovieSceneTrack::USubMovieSceneTrack( const FPostConstructInitializeProperti
 
 UMovieSceneSection* USubMovieSceneTrack::CreateNewSection()
 {
-	return ConstructObject<USubMovieSceneSection>( USubMovieSceneSection::StaticClass(), this );
+	return ConstructObject<USubMovieSceneSection>( USubMovieSceneSection::StaticClass(), this, NAME_None, RF_Transactional );
 }
 
 FName USubMovieSceneTrack::GetTrackName() const
diff --git a/Engine/Source/Runtime/UMG/Private/Animation/MovieSceneMarginSection.cpp b/Engine/Source/Runtime/UMG/Private/Animation/MovieSceneMarginSection.cpp
--- a/Engine/Source/Runtime/UMG/Private/Animation/MovieSceneMarginSection.cpp
+++ b/Engine/Source/Runtime/UMG/Private/Animation/MovieSceneMarginSection.cpp
@@ -2,7 +2,7 @@
 
 #include "UMGPrivatePCH.h"
 #include "MovieSceneMarginSection.h"
-
+#include "MovieSceneMarginTrack.h"
 
 UMovieSceneMarginSection::UMovieSceneMarginSection( const FPostConstructInitializeProperties& PCIP )
 	: Super( PCIP )
@@ -22,12 +22,12 @@ void UMovieSceneMarginSection::MoveSection( float DeltaTime )
 
 void UMovieSceneMarginSection::DilateSection( float DilationFactor, float Origin )
 {
+	Super::DilateSection(DilationFactor, Origin);
+
 	LeftCurve.ScaleCurve(Origin, DilationFactor);
 	TopCurve.ScaleCurve(Origin, DilationFactor);
 	RightCurve.ScaleCurve(Origin, DilationFactor);
 	BottomCurve.ScaleCurve(Origin, DilationFactor);
-
-	Super::DilateSection(DilationFactor, Origin);
 }
 
 FMargin UMovieSceneMarginSection::Eval( float Position ) const
@@ -38,12 +38,21 @@ FMargin UMovieSceneMarginSection::Eval( float Position ) const
 					BottomCurve.Eval(Position));
 }
 
-void UMovieSceneMarginSection::AddKey( float Time, const FMargin& Value )
+void UMovieSceneMarginSection::AddKey( float Time, const FMarginKey& MarginKey )
 {
-	AddKeyToCurve(LeftCurve, Time, Value.Left);
-	AddKeyToCurve(TopCurve, Time, Value.Top);
-	AddKeyToCurve(RightCurve, Time, Value.Right);
-	AddKeyToCurve(BottomCurve, Time, Value.Bottom);
+	Modify();
+
+	if( MarginKey.CurveName == NAME_None )
+	{
+		AddKeyToCurve(LeftCurve, Time, MarginKey.Value.Left);
+		AddKeyToCurve(TopCurve, Time, MarginKey.Value.Top);
+		AddKeyToCurve(RightCurve, Time, MarginKey.Value.Right);
+		AddKeyToCurve(BottomCurve, Time, MarginKey.Value.Bottom);
+	}
+	else
+	{
+		AddKeyToNamedCurve( Time, MarginKey );
+	}
 }
 
 bool UMovieSceneMarginSection::NewKeyIsNewData(float Time, const FMargin& Value) const
@@ -55,7 +64,32 @@ bool UMovieSceneMarginSection::NewKeyIsNewData(float Time, const FMargin& Value)
 		(Eval(Time) != Value);
 }
 
+void UMovieSceneMarginSection::AddKeyToNamedCurve( float Time, const FMarginKey& MarginKey )
+{
+	static FName Left("Left");
+	static FName Top("Top");
+	static FName Right("Right");
+	static FName Bottom("Bottom");
 
+	FName CurveName = MarginKey.CurveName;
+
+	if( CurveName == Left )
+	{
+		AddKeyToCurve(LeftCurve, Time, MarginKey.Value.Left);
+	}
+	else if( CurveName == Top )
+	{
+		AddKeyToCurve(TopCurve, Time, MarginKey.Value.Top);
+	}
+	else if( CurveName == Right )
+	{
+		AddKeyToCurve(RightCurve, Time, MarginKey.Value.Right);
+	}
+	else if( CurveName == Bottom )
+	{
+		AddKeyToCurve(BottomCurve, Time, MarginKey.Value.Bottom);
+	}
+}
 
 void UMovieSceneMarginSection::AddKeyToCurve( FRichCurve& InCurve, float Time, float Value )
 {
diff --git a/Engine/Source/Runtime/UMG/Private/Animation/MovieSceneMarginTrack.cpp b/Engine/Source/Runtime/UMG/Private/Animation/MovieSceneMarginTrack.cpp
--- a/Engine/Source/Runtime/UMG/Private/Animation/MovieSceneMarginTrack.cpp
+++ b/Engine/Source/Runtime/UMG/Private/Animation/MovieSceneMarginTrack.cpp
@@ -14,7 +14,7 @@ UMovieSceneMarginTrack::UMovieSceneMarginTrack(const FPostConstructInitializePro
 
 UMovieSceneSection* UMovieSceneMarginTrack::CreateNewSection()
 {
-	return ConstructObject<UMovieSceneSection>( UMovieSceneMarginSection::StaticClass(), this );
+	return ConstructObject<UMovieSceneSection>( UMovieSceneMarginSection::StaticClass(), this, NAME_None, RF_Transactional );
 }
 
 TSharedPtr<IMovieSceneTrackInstance> UMovieSceneMarginTrack::CreateInstance()
@@ -23,14 +23,16 @@ TSharedPtr<IMovieSceneTrackInstance> UMovieSceneMarginTrack::CreateInstance()
 }
 
 
-bool UMovieSceneMarginTrack::AddKeyToSection( float Time, FMargin Value )
+bool UMovieSceneMarginTrack::AddKeyToSection( float Time, const FMarginKey& MarginKey )
 {
 	const UMovieSceneSection* NearestSection = MovieSceneHelpers::FindSectionAtTime( Sections, Time );
-	if (!NearestSection || CastChecked<UMovieSceneMarginSection>(NearestSection)->NewKeyIsNewData(Time, Value))
+	if (!NearestSection || MarginKey.bAddKeyEvenIfUnchanged || CastChecked<UMovieSceneMarginSection>(NearestSection)->NewKeyIsNewData(Time, MarginKey.Value) )
 	{
+		Modify();
+
 		UMovieSceneMarginSection* NewSection = CastChecked<UMovieSceneMarginSection>( FindOrAddSection( Time ) );
 
-		NewSection->AddKey( Time, Value );
+		NewSection->AddKey( Time, MarginKey );
 
 		return true;
 	}
diff --git a/Engine/Source/Runtime/UMG/Public/Animation/MovieSceneMarginSection.h b/Engine/Source/Runtime/UMG/Public/Animation/MovieSceneMarginSection.h
--- a/Engine/Source/Runtime/UMG/Public/Animation/MovieSceneMarginSection.h
+++ b/Engine/Source/Runtime/UMG/Public/Animation/MovieSceneMarginSection.h
@@ -6,6 +6,8 @@
 
 #include "MovieSceneMarginSection.generated.h"
 
+struct FMarginKey;
+
 /**
  * A section in a Margin track
  */
@@ -31,7 +33,7 @@ public:
 	 * @param Time	The location in time where the key should be added
 	 * @param Value	The value of the key
 	 */
-	void AddKey( float Time, const FMargin& Value );
+	void AddKey( float Time, const FMarginKey& MarginKey );
 	
 	/** 
 	 * Determines if a new key would be new data, or just a duplicate of existing data
@@ -74,6 +76,7 @@ public:
 	const FRichCurve& GetBottomCurve() const { return BottomCurve; }
 
 private:
+	void AddKeyToNamedCurve( float Time, const FMarginKey& MarginKey );
 	/**
 	 * Adds a key to a rich curve, finding an existing key to modify or adding a new one
 	 *
diff --git a/Engine/Source/Runtime/UMG/Public/Animation/MovieSceneMarginTrack.h b/Engine/Source/Runtime/UMG/Public/Animation/MovieSceneMarginTrack.h
--- a/Engine/Source/Runtime/UMG/Public/Animation/MovieSceneMarginTrack.h
+++ b/Engine/Source/Runtime/UMG/Public/Animation/MovieSceneMarginTrack.h
@@ -5,6 +5,13 @@
 #include "MovieScenePropertyTrack.h"
 #include "MovieSceneMarginTrack.generated.h"
 
+struct FMarginKey
+{
+	FMargin Value;
+	FName CurveName;
+	bool bAddKeyEvenIfUnchanged;
+};
+
 /**
  * Handles manipulation of FMargins in a movie scene
  */
@@ -25,7 +32,7 @@ public:
 	 * @param Value				The value of the key
 	 * @return True if the key was successfully added.
 	 */
-	virtual bool AddKeyToSection( float Time, FMargin Value );
+	virtual bool AddKeyToSection( float Time, const FMarginKey& MarginKey );
 	
 	/**
 	 * Evaluates the track at the playback position
