Linux: fix build on case-sensitive fs.

diff --git a/Engine/Source/Editor/PropertyEditor/Private/SDetailSingleItemRow.cpp b/Engine/Source/Editor/PropertyEditor/Private/SDetailSingleItemRow.cpp
--- a/Engine/Source/Editor/PropertyEditor/Private/SDetailSingleItemRow.cpp
+++ b/Engine/Source/Editor/PropertyEditor/Private/SDetailSingleItemRow.cpp
@@ -4,7 +4,7 @@
 #include "SDetailSingleItemRow.h"
 #include "DetailItemNode.h"
 #include "PropertyEditorHelpers.h"
-#include "IDetailKeyFrameHandler.h"
+#include "IDetailKeyframeHandler.h"
 
 namespace DetailWidgetConstants
 {
diff --git a/Engine/Source/Editor/UMGEditor/Private/Animation/UMGDetailKeyframeHandler.h b/Engine/Source/Editor/UMGEditor/Private/Animation/UMGDetailKeyframeHandler.h
--- a/Engine/Source/Editor/UMGEditor/Private/Animation/UMGDetailKeyframeHandler.h
+++ b/Engine/Source/Editor/UMGEditor/Private/Animation/UMGDetailKeyframeHandler.h
@@ -2,7 +2,7 @@
 
 #pragma once
 
-#include "Editor/PropertyEditor/Public/IDetailKeyFrameHandler.h"
+#include "Editor/PropertyEditor/Public/IDetailKeyframeHandler.h"
 
 class FUMGDetailKeyframeHandler : public IDetailKeyframeHandler
 {
