Merging UMG change from 4.4 back to main to keep them in sync

diff --git a/Engine/Source/Editor/UMGEditor/Private/Customizations/DetailCustomizations.cpp b/Engine/Source/Editor/UMGEditor/Private/Customizations/DetailCustomizations.cpp
--- a/Engine/Source/Editor/UMGEditor/Private/Customizations/DetailCustomizations.cpp
+++ b/Engine/Source/Editor/UMGEditor/Private/Customizations/DetailCustomizations.cpp
@@ -91,6 +91,7 @@ void FBlueprintWidgetCustomization::CreateDelegateCustomization( IDetailLayoutBu
 			]
 			+ SHorizontalBox::Slot()
 			.AutoWidth()
+			.VAlign(VAlign_Center)
 			[
 				DelegateButton
 			]
@@ -161,6 +162,7 @@ void FBlueprintWidgetCustomization::CreateEventCustomization( IDetailLayoutBuild
 
 			+ SHorizontalBox::Slot()
 			.FillWidth(1)
+			.VAlign(VAlign_Center)
 			[
 				DelegateButton
 			]
diff --git a/Engine/Source/Editor/UMGEditor/Private/Designer/SDesignerView.cpp b/Engine/Source/Editor/UMGEditor/Private/Designer/SDesignerView.cpp
--- a/Engine/Source/Editor/UMGEditor/Private/Designer/SDesignerView.cpp
+++ b/Engine/Source/Editor/UMGEditor/Private/Designer/SDesignerView.cpp
@@ -101,6 +101,12 @@ void SDesignerView::Construct(const FArguments& InArgs, TSharedPtr<FWidgetBluepr
 	BlueprintEditor = InBlueprintEditor;
 	CurrentHandle = DH_NONE;
 
+	//TODO UMG Store as a setting
+	//PreviewWidth = 1920;
+	//PreviewHeight = 1080;
+	PreviewWidth = 1280;
+	PreviewHeight = 720;
+
 	HoverTime = 0;
 
 	bMouseDown = false;
@@ -118,6 +124,7 @@ void SDesignerView::Construct(const FArguments& InArgs, TSharedPtr<FWidgetBluepr
 	DragDirections[DH_BOTTOM_CENTER] = FVector2D(0, 1);
 	DragDirections[DH_BOTTOM_RIGHT] = FVector2D(1, 1);
 
+	// TODO UMG - Register these with the module through some public interface to allow for new extensions to be registered.
 	Register(MakeShareable(new FVerticalSlotExtension()));
 	Register(MakeShareable(new FHorizontalSlotExtension()));
 	Register(MakeShareable(new FCanvasSlotExtension()));
@@ -137,10 +144,18 @@ void SDesignerView::Construct(const FArguments& InArgs, TSharedPtr<FWidgetBluepr
 			.HAlign(HAlign_Fill)
 			.VAlign(VAlign_Fill)
 			[
-				SAssignNew(PreviewSurface, SZoomPan)
+				SAssignNew(PreviewHitTestRoot, SZoomPan)
 				.Visibility(EVisibility::HitTestInvisible)
 				.ZoomAmount(this, &SDesignerView::GetZoomAmount)
 				.ViewOffset(this, &SDesignerView::GetViewOffset)
+				[
+					SAssignNew(PreviewSurface, SBox)
+					.WidthOverride(this, &SDesignerView::GetPreviewWidth)
+					.HeightOverride(this, &SDesignerView::GetPreviewHeight)
+					.HAlign(HAlign_Fill)
+					.VAlign(VAlign_Fill)
+					.Visibility(EVisibility::SelfHitTestInvisible)
+				]
 			]
 
 			// A layer in the overlay where we put all the user intractable widgets, like the reorder widgets.
@@ -152,6 +167,29 @@ void SDesignerView::Construct(const FArguments& InArgs, TSharedPtr<FWidgetBluepr
 				.Visibility(EVisibility::SelfHitTestInvisible)
 			]
 
+			// Top bar with buttons for changing the designer
+			+ SOverlay::Slot()
+			.HAlign(HAlign_Right)
+			.VAlign(VAlign_Top)
+			[
+				SNew(SHorizontalBox)
+				
+				+ SHorizontalBox::Slot()
+				.AutoWidth()
+				[
+					SNew(SComboButton)
+					.ButtonStyle(FEditorStyle::Get(), "PropertyEditor.AssetComboStyle")
+					.ForegroundColor(FEditorStyle::GetColor("PropertyEditor.AssetName.ColorAndOpacity"))
+					.OnGetMenuContent(this, &SDesignerView::GetAspectMenu)
+					.ContentPadding(2.0f)
+					.ButtonContent()
+					[
+						SNew(STextBlock)
+						.Text(LOCTEXT("AspectRatio", "Aspect Ratio"))
+					]
+				]
+			]
+
 			// Bottom-right corner text indicating the experimental nature of umg
 			+ SOverlay::Slot()
 			.Padding(10)
@@ -183,6 +221,16 @@ SDesignerView::~SDesignerView()
 	}
 }
 
+FOptionalSize SDesignerView::GetPreviewWidth() const
+{
+	return (float)PreviewWidth;
+}
+
+FOptionalSize SDesignerView::GetPreviewHeight() const
+{
+	return (float)PreviewHeight;
+}
+
 void SDesignerView::OnEditorSelectionChanged()
 {
 	TSet<FWidgetReference> PendingSelectedWidgets = BlueprintEditor.Pin()->GetSelectedWidgets();
@@ -228,41 +276,84 @@ void SDesignerView::CreateExtensionWidgetsForSelection()
 	// Remove all the current extension widgets
 	ExtensionWidgetCanvas->ClearChildren();
 
-	ExtensionWidgets.Reset();
-
 	TArray<FWidgetReference> Selected;
 	if ( SelectedWidget.IsValid() )
 	{
 		Selected.Add(SelectedWidget);
 	}
 
+	TArray< TSharedRef<FDesignerSurfaceElement> > ExtensionElements;
+
 	// Build extension widgets for new selection
 	for ( TSharedRef<FDesignerExtension>& Ext : DesignerExtensions )
 	{
-		Ext->BuildWidgetsForSelection(Selected, ExtensionWidgets);
+		if ( Ext->CanExtendSelection(Selected) )
+		{
+			Ext->ExtendSelection(Selected, ExtensionElements);
+		}
 	}
 
-	//TODO UMG We need to have a better layout system for extensions.
-	//         currently everything is just laid out across the top.
-
-	TSharedRef<SHorizontalBox> ExtensionBox = SNew(SHorizontalBox);
-
 	// Add Widgets to designer surface
-	for ( TSharedRef<SWidget>& ExWidget : ExtensionWidgets )
+	for ( TSharedRef<FDesignerSurfaceElement>& ExtElement : ExtensionElements )
 	{
-		ExtensionBox->AddSlot()
-			.AutoWidth()
+		ExtensionWidgetCanvas->AddSlot()
+			.Position(TAttribute<FVector2D>::Create(TAttribute<FVector2D>::FGetter::CreateSP(this, &SDesignerView::GetExtensionPosition, ExtElement)))
+			.Size(TAttribute<FVector2D>::Create(TAttribute<FVector2D>::FGetter::CreateSP(this, &SDesignerView::GetExtensionSize, ExtElement)))
 			[
-				ExWidget
+				ExtElement->GetWidget()
 			];
 	}
+}
 
-	ExtensionWidgetCanvas->AddSlot()
-		.Position(TAttribute<FVector2D>(this, &SDesignerView::GetCachedSelectionDesignerWidgetsLocation))
-		.Size(FVector2D(100, 25))
-		[
-			ExtensionBox
-		];
+FVector2D SDesignerView::GetExtensionPosition(TSharedRef<FDesignerSurfaceElement> ExtensionElement) const
+{
+	const FVector2D TopLeft = CachedDesignerWidgetLocation;
+	const FVector2D Size = CachedDesignerWidgetSize * GetZoomAmount();
+
+	FVector2D FinalPosition(0, 0);
+
+	switch ( ExtensionElement->GetLocation() )
+	{
+	case EExtensionLayoutLocation::Absolute:
+		break;
+
+	case EExtensionLayoutLocation::TopLeft:
+		FinalPosition = TopLeft;
+		break;
+	case EExtensionLayoutLocation::TopCenter:
+		FinalPosition = TopLeft + FVector2D(Size.X * 0.5f, 0);
+		break;
+	case EExtensionLayoutLocation::TopRight:
+		FinalPosition = TopLeft + FVector2D(Size.X, 0);
+		break;
+
+	case EExtensionLayoutLocation::CenterLeft:
+		FinalPosition = TopLeft + FVector2D(0, Size.Y * 0.5f);
+		break;
+	case EExtensionLayoutLocation::CenterCenter:
+		FinalPosition = TopLeft + FVector2D(Size.X * 0.5f, Size.Y * 0.5f);
+		break;
+	case EExtensionLayoutLocation::CenterRight:
+		FinalPosition = TopLeft + FVector2D(Size.X, Size.Y * 0.5f);
+		break;
+
+	case EExtensionLayoutLocation::BottomLeft:
+		FinalPosition = TopLeft + FVector2D(0, Size.Y);
+		break;
+	case EExtensionLayoutLocation::BottomCenter:
+		FinalPosition = TopLeft + FVector2D(Size.X * 0.5f, Size.Y);
+		break;
+	case EExtensionLayoutLocation::BottomRight:
+		FinalPosition = TopLeft + Size;
+		break;
+	}
+
+	return FinalPosition + ExtensionElement->GetOffset();
+}
+
+FVector2D SDesignerView::GetExtensionSize(TSharedRef<FDesignerSurfaceElement> ExtensionElement) const
+{
+	return ExtensionElement->GetWidget()->GetDesiredSize();
 }
 
 UWidgetBlueprint* SDesignerView::GetBlueprint() const
@@ -306,11 +397,11 @@ FWidgetReference SDesignerView::GetWidgetAtCursor(const FGeometry& MyGeometry, c
 
 	FArrangedChildren Children(EVisibility::All);
 
-	PreviewSurface->SetVisibility(EVisibility::Visible);
-	FArrangedWidget WindowWidgetGeometry(PreviewSurface.ToSharedRef(), MyGeometry);
+	PreviewHitTestRoot->SetVisibility(EVisibility::Visible);
+	FArrangedWidget WindowWidgetGeometry(PreviewHitTestRoot.ToSharedRef(), MyGeometry);
 	LocateWidgetsUnderCursor_Helper(WindowWidgetGeometry, MouseEvent.GetScreenSpacePosition(), Children, true);
 
-	PreviewSurface->SetVisibility(EVisibility::HitTestInvisible);
+	PreviewHitTestRoot->SetVisibility(EVisibility::HitTestInvisible);
 
 	UUserWidget* WidgetActor = BlueprintEditor.Pin()->GetPreview();
 	if ( WidgetActor )
@@ -514,109 +605,20 @@ void SDesignerView::ShowContextMenu(const FGeometry& MyGeometry, const FPointerE
 	}
 }
 
-bool SDesignerView::GetArrangedWidget(TSharedRef<SWidget> Widget, FArrangedWidget& ArrangedWidget) const
-{
-	// We can't screenshot the widget unless there's a valid window handle to draw it in.
-	TSharedPtr<SWindow> WidgetWindow = FSlateApplication::Get().FindWidgetWindow(Widget);
-	if ( !WidgetWindow.IsValid() )
-	{
-		return false;
-	}
-
-	TSharedRef<SWindow> CurrentWindowRef = WidgetWindow.ToSharedRef();
-
-	FWidgetPath WidgetPath;
-	if ( FSlateApplication::Get().GeneratePathToWidgetUnchecked(Widget, WidgetPath) )
-	{
-		ArrangedWidget = WidgetPath.FindArrangedWidget(Widget);
-		return true;
-	}
-
-	return false;
-}
-
-bool SDesignerView::GetArrangedWidgetRelativeToWindow(TSharedRef<SWidget> Widget, FArrangedWidget& ArrangedWidget) const
-{
-	// We can't screenshot the widget unless there's a valid window handle to draw it in.
-	TSharedPtr<SWindow> WidgetWindow = FSlateApplication::Get().FindWidgetWindow(Widget);
-	if ( !WidgetWindow.IsValid() )
-	{
-		return false;
-	}
-
-	TSharedRef<SWindow> CurrentWindowRef = WidgetWindow.ToSharedRef();
-
-	FWidgetPath WidgetPath;
-	if ( FSlateApplication::Get().GeneratePathToWidgetUnchecked(Widget, WidgetPath) )
-	{
-		ArrangedWidget = WidgetPath.FindArrangedWidget(Widget);
-		ArrangedWidget.Geometry.AbsolutePosition -= CurrentWindowRef->GetPositionInScreen();
-		return true;
-	}
-
-	return false;
-}
-
-bool SDesignerView::GetArrangedWidgetRelativeToDesigner(TSharedRef<SWidget> Widget, FArrangedWidget& ArrangedWidget) const
-{
-	FWidgetPath WidgetPath;
-	if ( FSlateApplication::Get().GeneratePathToWidgetUnchecked(Widget, WidgetPath) )
-	{
-		FArrangedWidget ArrangedDesigner = WidgetPath.FindArrangedWidget(this->AsShared());
-
-		ArrangedWidget = WidgetPath.FindArrangedWidget(Widget);
-		ArrangedWidget.Geometry.AbsolutePosition -= ArrangedDesigner.Geometry.AbsolutePosition;
-		return true;
-	}
-
-	return false;
-}
-
-FVector2D SDesignerView::GetCachedSelectionDesignerWidgetsLocation() const
-{
-	return CachedDesignerWidgetLocation;
-}
-
-FVector2D SDesignerView::GetSelectionDesignerWidgetsLocation() const
+void SDesignerView::CacheSelectedWidgetGeometry()
 {
 	if ( SelectedSlateWidget.IsValid() )
 	{
 		TSharedRef<SWidget> Widget = SelectedSlateWidget.Pin().ToSharedRef();
 
 		FArrangedWidget ArrangedWidget(SNullWidget::NullWidget, FGeometry());
-		GetArrangedWidgetRelativeToDesigner(Widget, ArrangedWidget);
+		FDesignTimeUtils::GetArrangedWidgetRelativeToParent(Widget, AsShared(), ArrangedWidget);
 
-		return ArrangedWidget.Geometry.AbsolutePosition - FVector2D(0, 25);
+		CachedDesignerWidgetLocation = ArrangedWidget.Geometry.AbsolutePosition;
+		CachedDesignerWidgetSize = ArrangedWidget.Geometry.Size;
 	}
-
-	return FVector2D(0, 0);
 }
 
-//void SDesignerView::OnArrangeChildren(const FGeometry& AllottedGeometry, FArrangedChildren& ArrangedChildren) const
-//{
-//	const TSlotlessChildren<SNode>& ChildrenToArrange = ArrangedChildren.Accepts(EVisibility::Hidden) ? Children : VisibleChildren;
-//	// First pass nodes
-//	for ( int32 ChildIndex = 0; ChildIndex < ChildrenToArrange.Num(); ++ChildIndex )
-//	{
-//		const TSharedRef<SNode>& SomeChild = ChildrenToArrange[ChildIndex];
-//		if ( !SomeChild->RequiresSecondPassLayout() )
-//		{
-//			ArrangedChildren.AddWidget(AllottedGeometry.MakeChild(SomeChild, SomeChild->GetPosition() - ViewOffset, SomeChild->GetDesiredSize(), GetZoomAmount()));
-//		}
-//	}
-//
-//	// Second pass nodes
-//	for ( int32 ChildIndex = 0; ChildIndex < ChildrenToArrange.Num(); ++ChildIndex )
-//	{
-//		const TSharedRef<SNode>& SomeChild = ChildrenToArrange[ChildIndex];
-//		if ( SomeChild->RequiresSecondPassLayout() )
-//		{
-//			SomeChild->PerformSecondPassLayout(NodeToWidgetLookup);
-//			ArrangedChildren.AddWidget(AllottedGeometry.MakeChild(SomeChild, SomeChild->GetPosition() - ViewOffset, SomeChild->GetDesiredSize(), GetZoomAmount()));
-//		}
-//	}
-//}
-
 int32 SDesignerView::OnPaint(const FPaintArgs& Args, const FGeometry& AllottedGeometry, const FSlateRect& MyClippingRect, FSlateWindowElementList& OutDrawElements, int32 LayerId, const FWidgetStyle& InWidgetStyle, bool bParentEnabled) const
 {
 	SDesignSurface::OnPaint(Args, AllottedGeometry, MyClippingRect, OutDrawElements, LayerId, InWidgetStyle, bParentEnabled);
@@ -638,7 +640,7 @@ int32 SDesignerView::OnPaint(const FPaintArgs& Args, const FGeometry& AllottedGe
 		TSharedRef<SWidget> Widget = HoveredSlateWidget.Pin().ToSharedRef();
 
 		FArrangedWidget ArrangedWidget(SNullWidget::NullWidget, FGeometry());
-		GetArrangedWidgetRelativeToWindow(Widget, ArrangedWidget);
+		FDesignTimeUtils::GetArrangedWidgetRelativeToWindow(Widget, ArrangedWidget);
 
 		// Draw hovered effect
 		// Azure = 0x007FFF
@@ -665,7 +667,7 @@ int32 SDesignerView::OnPaint(const FPaintArgs& Args, const FGeometry& AllottedGe
 		TSharedRef<SWidget> Widget = SelectedSlateWidget.Pin().ToSharedRef();
 
 		FArrangedWidget ArrangedWidget(SNullWidget::NullWidget, FGeometry());
-		GetArrangedWidgetRelativeToWindow(Widget, ArrangedWidget);
+		FDesignTimeUtils::GetArrangedWidgetRelativeToWindow(Widget, ArrangedWidget);
 
 		const FLinearColor Tint(0, 1, 0);
 
@@ -708,12 +710,11 @@ int32 SDesignerView::OnPaint(const FPaintArgs& Args, const FGeometry& AllottedGe
 			const FSlateBrush* KeyBrush = FEditorStyle::GetBrush("CurveEd.CurveKey");
 			FLinearColor KeyColor = FLinearColor(1.0f, 0.0f, 0.0f, 1.f);
 
+			// Normalized offset of the handle, 0, no offset, 1 move the handle completely out by its own size.
+			//const float HandleOffset = 1.0f;
+
 			if ( UCanvasPanelSlot* CanvasSlot = Cast<UCanvasPanelSlot>(SelectedWidget.GetTemplate()->Slot) )
 			{
-				//Set("UMGEditor.AnchorCenter", new IMAGE_BRUSH("Icons/umg_anchor_center", Icon16x16));
-				//Set("UMGEditor.AnchorTopBottom", new IMAGE_BRUSH("Icons/umg_anchor_top_bottom", FVector2D(16, 32)));
-				//Set("UMGEditor.AnchorLeftRight", new IMAGE_BRUSH("Icons/umg_anchor_left_right", FVector2D(32, 16)));
-
 				// Left
 				FSlateDrawElement::MakeBox(
 					OutDrawElements,
@@ -824,7 +825,7 @@ SDesignerView::DragHandle SDesignerView::HitTestDragHandles(const FGeometry& All
 		TSharedRef<SWidget> Widget = SelectedSlateWidget.Pin().ToSharedRef();
 
 		FArrangedWidget ArrangedWidget(SNullWidget::NullWidget, FGeometry());
-		GetArrangedWidget(Widget, ArrangedWidget);
+		FDesignTimeUtils::GetArrangedWidget(Widget, ArrangedWidget);
 
 		FVector2D WidgetLocalPosition = AllottedGeometry.AbsoluteToLocal(ArrangedWidget.Geometry.AbsolutePosition);
 		float X = WidgetLocalPosition.X;
@@ -889,10 +890,8 @@ FCursorReply SDesignerView::OnCursorQuery(const FGeometry& MyGeometry, const FPo
 	return FCursorReply::Unhandled();
 }
 
-void SDesignerView::Tick(const FGeometry& AllottedGeometry, const double InCurrentTime, const float InDeltaTime)
+void SDesignerView::UpdatePreviewWidget()
 {
-	HoverTime += InDeltaTime;
-
 	UUserWidget* LatestPreviewWidget = BlueprintEditor.Pin()->GetPreview();
 
 	if ( LatestPreviewWidget != PreviewWidget )
@@ -935,14 +934,21 @@ void SDesignerView::Tick(const FGeometry& AllottedGeometry, const double InCurre
 			];
 		}
 	}
+}
+
+void SDesignerView::Tick(const FGeometry& AllottedGeometry, const double InCurrentTime, const float InDeltaTime)
+{
+	HoverTime += InDeltaTime;
+
+	UpdatePreviewWidget();
 
 	// Update the selected widget to match the selected template.
-	if ( LatestPreviewWidget )
+	if ( PreviewWidget )
 	{
 		if ( SelectedWidget.IsValid() )
 		{
 			// Set the selected widget so that we can draw the highlight
-			SelectedSlateWidget = LatestPreviewWidget->GetWidgetFromName(SelectedWidget.GetTemplate()->GetName());
+			SelectedSlateWidget = PreviewWidget->GetWidgetFromName(SelectedWidget.GetTemplate()->GetName());
 		}
 		else
 		{
@@ -951,7 +957,7 @@ void SDesignerView::Tick(const FGeometry& AllottedGeometry, const double InCurre
 
 		if ( HoveredWidget.IsValid() )
 		{
-			HoveredSlateWidget = LatestPreviewWidget->GetWidgetFromName(HoveredWidget.GetTemplate()->GetName());
+			HoveredSlateWidget = PreviewWidget->GetWidgetFromName(HoveredWidget.GetTemplate()->GetName());
 		}
 		else
 		{
@@ -959,7 +965,13 @@ void SDesignerView::Tick(const FGeometry& AllottedGeometry, const double InCurre
 		}
 	}
 
-	CachedDesignerWidgetLocation = GetSelectionDesignerWidgetsLocation();
+	CacheSelectedWidgetGeometry();
+
+	// Tick all designer extensions in case they need to update widgets
+	for ( const TSharedRef<FDesignerExtension>& Ext : DesignerExtensions )
+	{
+		Ext->Tick(AllottedGeometry, InCurrentTime, InDeltaTime);
+	}
 
 	SCompoundWidget::Tick(AllottedGeometry, InCurrentTime, InDeltaTime);
 }
@@ -1259,6 +1271,47 @@ FReply SDesignerView::OnDrop(const FGeometry& MyGeometry, const FDragDropEvent&
 	return FReply::Unhandled();
 }
 
+void SDesignerView::HandleCommonResolutionSelected(int32 Width, int32 Height)
+{
+	PreviewWidth = Width;
+	PreviewHeight = Height;
+}
+
+void SDesignerView::AddScreenResolutionSection(FMenuBuilder& MenuBuilder, const TArray<FPlayScreenResolution>& Resolutions, const FText& SectionName)
+{
+	MenuBuilder.BeginSection(NAME_None, SectionName);
+	{
+		for ( auto Iter = Resolutions.CreateConstIterator(); Iter; ++Iter )
+		{
+			FUIAction Action(FExecuteAction::CreateRaw(this, &SDesignerView::HandleCommonResolutionSelected, Iter->Width, Iter->Height));
+
+			FInternationalization& I18N = FInternationalization::Get();
+
+			FFormatNamedArguments Args;
+			Args.Add(TEXT("Width"), FText::AsNumber(Iter->Width, NULL, I18N.GetInvariantCulture()));
+			Args.Add(TEXT("Height"), FText::AsNumber(Iter->Height, NULL, I18N.GetInvariantCulture()));
+			Args.Add(TEXT("AspectRatio"), FText::FromString(Iter->AspectRatio));
+
+			MenuBuilder.AddMenuEntry(FText::FromString(Iter->Description), FText::Format(LOCTEXT("CommonResolutionFormat", "{Width} x {Height} (AspectRatio)"), Args), FSlateIcon(), Action);
+		}
+	}
+	MenuBuilder.EndSection();
+}
+
+TSharedRef<SWidget> SDesignerView::GetAspectMenu()
+{
+	const ULevelEditorPlaySettings* PlaySettings = GetDefault<ULevelEditorPlaySettings>();
+	FMenuBuilder MenuBuilder(true, NULL);
+
+	AddScreenResolutionSection(MenuBuilder, PlaySettings->PhoneScreenResolutions, LOCTEXT("CommonPhonesSectionHeader", "Phones"));
+	AddScreenResolutionSection(MenuBuilder, PlaySettings->TabletScreenResolutions, LOCTEXT("CommonTabletsSectionHeader", "Tablets"));
+	AddScreenResolutionSection(MenuBuilder, PlaySettings->LaptopScreenResolutions, LOCTEXT("CommonLaptopsSectionHeader", "Laptops"));
+	AddScreenResolutionSection(MenuBuilder, PlaySettings->MonitorScreenResolutions, LOCTEXT("CommoMonitorsSectionHeader", "Monitors"));
+	AddScreenResolutionSection(MenuBuilder, PlaySettings->TelevisionScreenResolutions, LOCTEXT("CommonTelevesionsSectionHeader", "Televisions"));
+
+	return MenuBuilder.MakeWidget();
+}
+
 void SDesignerView::BeginTransaction(const FText& SessionName)
 {
 	if ( ScopedTransaction == NULL )
diff --git a/Engine/Source/Editor/UMGEditor/Private/Designer/SDesignerView.h b/Engine/Source/Editor/UMGEditor/Private/Designer/SDesignerView.h
--- a/Engine/Source/Editor/UMGEditor/Private/Designer/SDesignerView.h
+++ b/Engine/Source/Editor/UMGEditor/Private/Designer/SDesignerView.h
@@ -49,13 +49,18 @@ public:
 	void Register(TSharedRef<FDesignerExtension> Extension);
 
 private:
+	FOptionalSize GetPreviewWidth() const;
+
+	FOptionalSize GetPreviewHeight() const;
+
+	/** Updates the designer to display the latest preview widget */
+	void UpdatePreviewWidget();
+
 	void CreateExtensionWidgetsForSelection();
 
 	/** Displays the context menu when you right click */
 	void ShowContextMenu(const FGeometry& MyGeometry, const FPointerEvent& MouseEvent);
 
-	FVector2D GetZoomVector() const;
-
 	void OnEditorSelectionChanged();
 
 	/** @returns Gets the widget under the cursor based on a mouse pointer event. */
@@ -68,12 +73,12 @@ private:
 	void OnBlueprintChanged(UBlueprint* InBlueprint);
 	void OnObjectPropertyChanged(UObject* ObjectBeingModified, FPropertyChangedEvent& PropertyChangedEvent);
 
-	bool GetArrangedWidget(TSharedRef<SWidget> Widget, FArrangedWidget& ArrangedWidget) const;
-	bool GetArrangedWidgetRelativeToWindow(TSharedRef<SWidget> Widget, FArrangedWidget& ArrangedWidget) const;
-	bool GetArrangedWidgetRelativeToDesigner(TSharedRef<SWidget> Widget, FArrangedWidget& ArrangedWidget) const;
+	void CacheSelectedWidgetGeometry();
 
-	FVector2D GetSelectionDesignerWidgetsLocation() const;
-	FVector2D GetCachedSelectionDesignerWidgetsLocation() const;
+	// Handles selecting a common screen resolution.
+	void HandleCommonResolutionSelected(int32 Width, int32 Height);
+	void AddScreenResolutionSection(FMenuBuilder& MenuBuilder, const TArray<FPlayScreenResolution>& Resolutions, const FText& SectionName);
+	TSharedRef<SWidget> GetAspectMenu();
 
 	void BeginTransaction(const FText& SessionName);
 	void EndTransaction();
@@ -102,14 +107,17 @@ private:
 	/** Extensions for the designer to allow for custom widgets to be inserted onto the design surface as selection changes. */
 	TArray< TSharedRef<FDesignerExtension> > DesignerExtensions;
 
-	/** Temporary widgets that are destroyed every time selection changes and are built from the current selection by extensions. */
-	TArray< TSharedRef<SWidget> > ExtensionWidgets;
-
 private:
 	void DrawDragHandles(const FPaintGeometry& SelectionGeometry, const FSlateRect& MyClippingRect, FSlateWindowElementList& OutDrawElements, int32 LayerId, const FWidgetStyle& InWidgetStyle) const;
 	DragHandle HitTestDragHandles(const FGeometry& AllottedGeometry, const FPointerEvent& PointerEvent) const;
 
 	UWidget* ProcessDropAndAddWidget(const FGeometry& MyGeometry, const FDragDropEvent& DragDropEvent, bool bIsPreview);
+
+private:
+
+	FVector2D GetExtensionPosition(TSharedRef<FDesignerSurfaceElement> ExtensionElement) const;
+
+	FVector2D GetExtensionSize(TSharedRef<FDesignerSurfaceElement> ExtensionElement) const;
 	
 private:
 	/** A reference to the BP Editor that owns this designer */
@@ -127,12 +135,14 @@ private:
 	UWidget* DropPreviewWidget;
 	UPanelWidget* DropPreviewParent;
 
-	TSharedPtr<class SZoomPan> PreviewSurface;
+	TSharedPtr<class SZoomPan> PreviewHitTestRoot;
+	TSharedPtr<SBox> PreviewSurface;
 	TSharedPtr<SCanvas> ExtensionWidgetCanvas;
 
 	DragHandle CurrentHandle;
 
 	FVector2D CachedDesignerWidgetLocation;
+	FVector2D CachedDesignerWidgetSize;
 
 	/** The currently selected set of widgets */
 	TSet< FWidgetReference > SelectedWidgets;
@@ -161,4 +171,7 @@ private:
 
 	/** The current slate widget being hovered, this is refreshed every frame in case it changes */
 	TWeakPtr<SWidget> HoveredSlateWidget;
+
+	int32 PreviewWidth;
+	int32 PreviewHeight;
 };
diff --git a/Engine/Source/Editor/UMGEditor/Private/Designer/SZoomPan.cpp b/Engine/Source/Editor/UMGEditor/Private/Designer/SZoomPan.cpp
--- a/Engine/Source/Editor/UMGEditor/Private/Designer/SZoomPan.cpp
+++ b/Engine/Source/Editor/UMGEditor/Private/Designer/SZoomPan.cpp
@@ -32,7 +32,7 @@ void SZoomPan::OnArrangeChildren(const FGeometry& AllottedGeometry, FArrangedChi
 		ArrangedChildren.AddWidget( ChildVisibility, AllottedGeometry.MakeChild(
 				ChildSlot.GetWidget(),
 				FVector2D(XResult.Offset, YResult.Offset) - ViewOffset.Get(),
-				FVector2D(XResult.Size, YResult.Size),
+				ChildSlot.Widget->GetDesiredSize(),
 				ZoomAmount.Get()
 		) );
 	}
diff --git a/Engine/Source/Editor/UMGEditor/Private/Extensions/CanvasSlotExtension.cpp b/Engine/Source/Editor/UMGEditor/Private/Extensions/CanvasSlotExtension.cpp
--- a/Engine/Source/Editor/UMGEditor/Private/Extensions/CanvasSlotExtension.cpp
+++ b/Engine/Source/Editor/UMGEditor/Private/Extensions/CanvasSlotExtension.cpp
@@ -21,7 +21,7 @@
 	ExtensionId = FName(TEXT("CanvasSlot"));
 }
 
-bool FCanvasSlotExtension::IsActive(const TArray< FWidgetReference >& Selection)
+bool FCanvasSlotExtension::CanExtendSelection(const TArray< FWidgetReference >& Selection) const
 {
 	for ( const FWidgetReference& Widget : Selection )
 	{
@@ -34,15 +34,10 @@
 	return Selection.Num() == 1;
 }
 
-void FCanvasSlotExtension::BuildWidgetsForSelection(const TArray< FWidgetReference >& Selection, TArray< TSharedRef<SWidget> >& Widgets)
+void FCanvasSlotExtension::ExtendSelection(const TArray< FWidgetReference >& Selection, TArray< TSharedRef<FDesignerSurfaceElement> >& SurfaceElements)
 {
 	SelectionCache = Selection;
 
-	if ( !IsActive(Selection) )
-	{
-		return;
-	}
-
 	MoveHandle =
 		SNew(SBorder)
 		//.BorderImage(FEditorStyle::Get().GetBrush("ToolBar.Background"))
@@ -56,7 +51,9 @@
 			.Image(FCoreStyle::Get().GetBrush("SoftwareCursor_CardinalCross"))
 		];
 
-	Widgets.Add(MoveHandle.ToSharedRef());
+	MoveHandle->SlatePrepass();
+
+	SurfaceElements.Add(MakeShareable(new FDesignerSurfaceElement(MoveHandle.ToSharedRef(), EExtensionLayoutLocation::TopLeft, FVector2D(0, -(MoveHandle->GetDesiredSize().Y + 10)))));
 }
 
 bool FCanvasSlotExtension::GetCollisionSegmentsForSlot(UCanvasPanel* Canvas, int32 SlotIndex, TArray<FVector2D>& Segments)
@@ -102,7 +99,91 @@
 	Segments[7] = ArrangedGeometry.Position + ArrangedGeometry.Size;
 }
 
+void FCanvasSlotExtension::Tick(const FGeometry& AllottedGeometry, const double InCurrentTime, const float InDeltaTime)
+{
+
+}
+
 void FCanvasSlotExtension::Paint(const TSet< FWidgetReference >& Selection, const FGeometry& AllottedGeometry, const FSlateRect& MyClippingRect, FSlateWindowElementList& OutDrawElements, int32 LayerId, const FWidgetStyle& InWidgetStyle, bool bParentEnabled) const
 {
+	PaintCollisionLines(Selection, AllottedGeometry, MyClippingRect, OutDrawElements, LayerId, InWidgetStyle, bParentEnabled);
+}
+
+FReply FCanvasSlotExtension::HandleBeginDrag(const FGeometry& Geometry, const FPointerEvent& Event)
+{
+	bDragging = true;
+
+	BeginTransaction(LOCTEXT("MoveWidget", "Move Widget"));
+
+	return FReply::Handled().CaptureMouse(MoveHandle.ToSharedRef());
+}
+
+FReply FCanvasSlotExtension::HandleEndDrag(const FGeometry& Geometry, const FPointerEvent& Event)
+{
+	bDragging = false;
+
+	EndTransaction();
+
+	FBlueprintEditorUtils::MarkBlueprintAsModified(Blueprint);
+
+	return FReply::Handled().ReleaseMouseCapture();
+}
+
+FReply FCanvasSlotExtension::HandleDragging(const FGeometry& Geometry, const FPointerEvent& Event)
+{
+	if ( bDragging )
+	{
+		for ( FWidgetReference& Selection : SelectionCache )
+		{
+			MoveByAmount(Selection, Event.GetCursorDelta());
+		}
+
+		return FReply::Handled();
+	}
+
+	return FReply::Unhandled();
+}
+
+void FCanvasSlotExtension::MoveByAmount(FWidgetReference& WidgetRef, FVector2D Delta)
+{
+	if ( Delta.IsZero() )
+	{
+		return;
+	}
+
+	UWidget* Widget = WidgetRef.GetPreview();
+
+	UCanvasPanelSlot* CanvasSlot = CastChecked<UCanvasPanelSlot>(Widget->Slot);
+	UCanvasPanel* Parent = CastChecked<UCanvasPanel>(CanvasSlot->Parent);
+
+	FMargin Offsets = CanvasSlot->LayoutData.Offsets;
+	Offsets.Left += Delta.X;
+	Offsets.Top += Delta.Y;
+
+	// If the slot is stretched horizontally we need to move the right side as it no longer represents width, but
+	// now represents margin from the right stretched side.
+	if ( CanvasSlot->LayoutData.Anchors.IsStretchedHorizontal() )
+	{
+		Offsets.Right -= Delta.X;
+	}
+
+	// If the slot is stretched vertically we need to move the bottom side as it no longer represents width, but
+	// now represents margin from the bottom stretched side.
+	if ( CanvasSlot->LayoutData.Anchors.IsStretchedVertical() )
+	{
+		Offsets.Bottom -= Delta.Y;
+	}
+
+	CanvasSlot->SetOffsets(Offsets);
+
+	// Update the Template widget to match
+	UWidget* TemplateWidget = WidgetRef.GetTemplate();
+	UCanvasPanelSlot* TemplateSlot = CastChecked<UCanvasPanelSlot>(TemplateWidget->Slot);
+
+	TemplateSlot->SetOffsets(Offsets);
+}
+
+void FCanvasSlotExtension::PaintCollisionLines(const TSet< FWidgetReference >& Selection, const FGeometry& AllottedGeometry, const FSlateRect& MyClippingRect, FSlateWindowElementList& OutDrawElements, int32 LayerId, const FWidgetStyle& InWidgetStyle, bool bParentEnabled) const
+{
 	for ( const FWidgetReference& WidgetRef : Selection )
 	{
@@ -196,78 +277,4 @@
 	}
 }
 
-FReply FCanvasSlotExtension::HandleBeginDrag(const FGeometry& Geometry, const FPointerEvent& Event)
-{
-	bDragging = true;
-
-	BeginTransaction(LOCTEXT("MoveWidget", "Move Widget"));
-
-	return FReply::Handled().CaptureMouse(MoveHandle.ToSharedRef());
-}
-
-FReply FCanvasSlotExtension::HandleEndDrag(const FGeometry& Geometry, const FPointerEvent& Event)
-{
-	bDragging = false;
-
-	EndTransaction();
-
-	FBlueprintEditorUtils::MarkBlueprintAsModified(Blueprint);
-
-	return FReply::Handled().ReleaseMouseCapture();
-}
-
-FReply FCanvasSlotExtension::HandleDragging(const FGeometry& Geometry, const FPointerEvent& Event)
-{
-	if ( bDragging )
-	{
-		for ( FWidgetReference& Selection : SelectionCache )
-		{
-			MoveByAmount(Selection, Event.GetCursorDelta());
-		}
-
-		return FReply::Handled();
-	}
-
-	return FReply::Unhandled();
-}
-
-void FCanvasSlotExtension::MoveByAmount(FWidgetReference& WidgetRef, FVector2D Delta)
-{
-	if ( Delta.IsZero() )
-	{
-		return;
-	}
-
-	UWidget* Widget = WidgetRef.GetPreview();
-
-	UCanvasPanelSlot* CanvasSlot = CastChecked<UCanvasPanelSlot>(Widget->Slot);
-	UCanvasPanel* Parent = CastChecked<UCanvasPanel>(CanvasSlot->Parent);
-
-	FMargin Offsets = CanvasSlot->LayoutData.Offsets;
-	Offsets.Left += Delta.X;
-	Offsets.Top += Delta.Y;
-
-	// If the slot is stretched horizontally we need to move the right side as it no longer represents width, but
-	// now represents margin from the right stretched side.
-	if ( CanvasSlot->LayoutData.Anchors.IsStretchedHorizontal() )
-	{
-		Offsets.Right -= Delta.X;
-	}
-
-	// If the slot is stretched vertically we need to move the bottom side as it no longer represents width, but
-	// now represents margin from the bottom stretched side.
-	if ( CanvasSlot->LayoutData.Anchors.IsStretchedVertical() )
-	{
-		Offsets.Bottom -= Delta.Y;
-	}
-
-	CanvasSlot->SetOffsets(Offsets);
-
-	// Update the Template widget to match
-	UWidget* TemplateWidget = WidgetRef.GetTemplate();
-	UCanvasPanelSlot* TemplateSlot = CastChecked<UCanvasPanelSlot>(TemplateWidget->Slot);
-
-	TemplateSlot->SetOffsets(Offsets);
-}
-
 #undef LOCTEXT_NAMESPACE
diff --git a/Engine/Source/Editor/UMGEditor/Private/Extensions/CanvasSlotExtension.h b/Engine/Source/Editor/UMGEditor/Private/Extensions/CanvasSlotExtension.h
--- a/Engine/Source/Editor/UMGEditor/Private/Extensions/CanvasSlotExtension.h
+++ b/Engine/Source/Editor/UMGEditor/Private/Extensions/CanvasSlotExtension.h
@@ -4,14 +4,18 @@
 
 #include "DesignerExtension.h"
 
+/**
+ * The canvas slot extension provides design time widgets for widgets that are selected in the canvas.
+ */
 class FCanvasSlotExtension : public FDesignerExtension
 {
 public:
 	FCanvasSlotExtension();
 
-	bool IsActive(const TArray< FWidgetReference >& Selection);
+	virtual bool CanExtendSelection(const TArray< FWidgetReference >& Selection) const override;
 	
-	virtual void BuildWidgetsForSelection(const TArray< FWidgetReference >& Selection, TArray< TSharedRef<SWidget> >& Widgets) override;
+	virtual void ExtendSelection(const TArray< FWidgetReference >& Selection, TArray< TSharedRef<FDesignerSurfaceElement> >& SurfaceElements) override;
+	virtual void Tick(const FGeometry& AllottedGeometry, const double InCurrentTime, const float InDeltaTime);
 	virtual void Paint(const TSet< FWidgetReference >& Selection, const FGeometry& AllottedGeometry, const FSlateRect& MyClippingRect, FSlateWindowElementList& OutDrawElements, int32 LayerId, const FWidgetStyle& InWidgetStyle, bool bParentEnabled) const override;
 
 private:
@@ -26,6 +30,10 @@ private:
 	static bool GetCollisionSegmentsForSlot(UCanvasPanel* Canvas, UCanvasPanelSlot* Slot, TArray<FVector2D>& Segments);
 	static void GetCollisionSegmentsFromGeometry(FGeometry ArrangedGeometry, TArray<FVector2D>& Segments);
 
+	void PaintCollisionLines(const TSet< FWidgetReference >& Selection, const FGeometry& AllottedGeometry, const FSlateRect& MyClippingRect, FSlateWindowElementList& OutDrawElements, int32 LayerId, const FWidgetStyle& InWidgetStyle, bool bParentEnabled) const;
+
+private:
+
 	bool bDragging;
 
 	TSharedPtr<SBorder> MoveHandle;
diff --git a/Engine/Source/Editor/UMGEditor/Private/Extensions/DesignerExtension.cpp b/Engine/Source/Editor/UMGEditor/Private/Extensions/DesignerExtension.cpp
--- a/Engine/Source/Editor/UMGEditor/Private/Extensions/DesignerExtension.cpp
+++ b/Engine/Source/Editor/UMGEditor/Private/Extensions/DesignerExtension.cpp
@@ -6,71 +6,6 @@
 
 #define LOCTEXT_NAMESPACE "UMG"
 
-// FWidgetReference
-
-FWidgetReference::FWidgetReference()
-	: WidgetEditor()
-	, TemplateWidget(NULL)
-{
-}
-
-FWidgetReference::FWidgetReference(TSharedPtr<FWidgetBlueprintEditor> WidgetEditor, UWidget* TemplateWidget)
-	: WidgetEditor(WidgetEditor)
-	, TemplateWidget(TemplateWidget)
-{
-}
-
-FWidgetReference FWidgetReference::FromTemplate(TSharedPtr<FWidgetBlueprintEditor> InWidgetEditor, UWidget* TemplateWidget)
-{
-	return FWidgetReference(InWidgetEditor, TemplateWidget);
-}
-
-FWidgetReference FWidgetReference::FromPreview(TSharedPtr<FWidgetBlueprintEditor> InWidgetEditor, UWidget* PreviewWidget)
-{
-	if ( InWidgetEditor.IsValid() )
-	{
-		UUserWidget* PreviewRoot = InWidgetEditor->GetPreview();
-		if ( PreviewRoot )
-		{
-			UWidgetBlueprint* Blueprint = InWidgetEditor->GetWidgetBlueprintObj();
-
-			if ( PreviewWidget )
-			{
-				FString Name = PreviewWidget->GetName();
-				return FromTemplate(InWidgetEditor, Blueprint->WidgetTree->FindWidget(Name));
-			}
-		}
-	}
-
-	return FWidgetReference();
-}
-
-bool FWidgetReference::IsValid() const
-{
-	return TemplateWidget.Get() != NULL && GetPreview();
-}
-
-UWidget* FWidgetReference::GetTemplate() const
-{
-	return TemplateWidget.Get();
-}
-
-UWidget* FWidgetReference::GetPreview() const
-{
-	if ( WidgetEditor.IsValid() )
-	{
-		UUserWidget* PreviewRoot = WidgetEditor.Pin()->GetPreview();
-
-		if ( PreviewRoot && TemplateWidget.Get() )
-		{
-			UWidget* PreviewWidget = PreviewRoot->GetHandleFromName(TemplateWidget.Get()->GetName());
-			return PreviewWidget;
-		}
-	}
-
-	return NULL;
-}
-
 // Designer Extension
 
 FDesignerExtension::FDesignerExtension()
diff --git a/Engine/Source/Editor/UMGEditor/Private/Extensions/HorizontalSlotExtension.cpp b/Engine/Source/Editor/UMGEditor/Private/Extensions/HorizontalSlotExtension.cpp
--- a/Engine/Source/Editor/UMGEditor/Private/Extensions/HorizontalSlotExtension.cpp
+++ b/Engine/Source/Editor/UMGEditor/Private/Extensions/HorizontalSlotExtension.cpp
@@ -10,7 +10,7 @@ FHorizontalSlotExtension::FHorizontalSlotExtension()
 	ExtensionId = FName(TEXT("HorizontalSlot"));
 }
 
-bool FHorizontalSlotExtension::IsActive(const TArray< FWidgetReference >& Selection)
+bool FHorizontalSlotExtension::CanExtendSelection(const TArray< FWidgetReference >& Selection) const
 {
 	for ( const FWidgetReference& Widget : Selection )
 	{
@@ -23,27 +23,33 @@ bool FHorizontalSlotExtension::IsActive(const TArray< FWidgetReference >& Select
 	return Selection.Num() == 1;
 }
 
-void FHorizontalSlotExtension::BuildWidgetsForSelection(const TArray< FWidgetReference >& Selection, TArray< TSharedRef<SWidget> >& Widgets)
+void FHorizontalSlotExtension::ExtendSelection(const TArray< FWidgetReference >& Selection, TArray< TSharedRef<FDesignerSurfaceElement> >& SurfaceElements)
 {
 	SelectionCache = Selection;
 
-	if ( !IsActive(Selection) )
-	{
-		return;
-	}
-
-	TSharedRef<SButton> LeftButton =
-		SNew(SButton)
+	TSharedRef<SButton> LeftArrow = SNew(SButton)
 		.Text(LOCTEXT("LeftArrow", "←"))
+		.ContentPadding(FMargin(2, 6))
+		.IsEnabled(this, &FHorizontalSlotExtension::CanShift, -1)
 		.OnClicked(this, &FHorizontalSlotExtension::HandleShift, -1);
 
-	TSharedRef<SButton> RightButton =
-		SNew(SButton)
+	TSharedRef<SButton> RightArrow = SNew(SButton)
 		.Text(LOCTEXT("RightArrow", "→"))
+		.ContentPadding(FMargin(2, 6))
+		.IsEnabled(this, &FHorizontalSlotExtension::CanShift, 1)
 		.OnClicked(this, &FHorizontalSlotExtension::HandleShift, 1);
 
-	Widgets.Add(LeftButton);
-	Widgets.Add(RightButton);
+	LeftArrow->SlatePrepass();
+	RightArrow->SlatePrepass();
+
+	SurfaceElements.Add(MakeShareable(new FDesignerSurfaceElement(LeftArrow, EExtensionLayoutLocation::CenterLeft, FVector2D(-LeftArrow->GetDesiredSize().X, LeftArrow->GetDesiredSize().Y * -0.5f))));
+	SurfaceElements.Add(MakeShareable(new FDesignerSurfaceElement(RightArrow, EExtensionLayoutLocation::CenterRight, FVector2D(0, RightArrow->GetDesiredSize().Y * -0.5f))));
+}
+
+bool FHorizontalSlotExtension::CanShift(int32 ShiftAmount) const
+{
+	//TODO UMG Provide feedback if shifting is possible.  Tricky with multiple items selected, if we ever support that.
+	return true;
 }
 
 FReply FHorizontalSlotExtension::HandleShift(int32 ShiftAmount)
diff --git a/Engine/Source/Editor/UMGEditor/Private/Extensions/HorizontalSlotExtension.h b/Engine/Source/Editor/UMGEditor/Private/Extensions/HorizontalSlotExtension.h
--- a/Engine/Source/Editor/UMGEditor/Private/Extensions/HorizontalSlotExtension.h
+++ b/Engine/Source/Editor/UMGEditor/Private/Extensions/HorizontalSlotExtension.h
@@ -9,12 +9,14 @@ class FHorizontalSlotExtension : public FDesignerExtension
 public:
 	FHorizontalSlotExtension();
 
-	bool IsActive(const TArray< FWidgetReference >& Selection);
+	virtual bool CanExtendSelection(const TArray< FWidgetReference >& Selection) const override;
 	
-	virtual void BuildWidgetsForSelection(const TArray< FWidgetReference >& Selection, TArray< TSharedRef<SWidget> >& Widgets) override;
+	virtual void ExtendSelection(const TArray< FWidgetReference >& Selection, TArray< TSharedRef<FDesignerSurfaceElement> >& SurfaceElements) override;
 
 private:
 
+	bool CanShift(int32 ShiftAmount) const;
+
 	FReply HandleShift(int32 ShiftAmount);
 	
 	void ShiftHorizontal(UWidget* Widget, int32 ShiftAmount);
diff --git a/Engine/Source/Editor/UMGEditor/Private/Extensions/UniformGridSlotExtension.cpp b/Engine/Source/Editor/UMGEditor/Private/Extensions/UniformGridSlotExtension.cpp
--- a/Engine/Source/Editor/UMGEditor/Private/Extensions/UniformGridSlotExtension.cpp
+++ b/Engine/Source/Editor/UMGEditor/Private/Extensions/UniformGridSlotExtension.cpp
@@ -9,7 +9,7 @@ FUniformGridSlotExtension::FUniformGridSlotExtension()
 	ExtensionId = FName(TEXT("UniformGridSlot"));
 }
 
-bool FUniformGridSlotExtension::IsActive(const TArray< FWidgetReference >& Selection)
+bool FUniformGridSlotExtension::CanExtendSelection(const TArray< FWidgetReference >& Selection) const
 {
 	for ( const FWidgetReference& Widget : Selection )
 	{
@@ -22,39 +22,39 @@ bool FUniformGridSlotExtension::IsActive(const TArray< FWidgetReference >& Selec
 	return Selection.Num() == 1;
 }
 
-void FUniformGridSlotExtension::BuildWidgetsForSelection(const TArray< FWidgetReference >& Selection, TArray< TSharedRef<SWidget> >& Widgets)
+void FUniformGridSlotExtension::ExtendSelection(const TArray< FWidgetReference >& Selection, TArray< TSharedRef<FDesignerSurfaceElement> >& SurfaceElements)
 {
 	SelectionCache = Selection;
 
-	if ( !IsActive(Selection) )
-	{
-		return;
-	}
-
-	TSharedRef<SButton> UpButton =
-		SNew(SButton)
+	TSharedRef<SButton> UpArrow = SNew(SButton)
 		.Text(LOCTEXT("UpArrow", "↑"))
+		.ContentPadding(FMargin(6, 2))
 		.OnClicked(this, &FUniformGridSlotExtension::HandleShiftRow, -1);
 
-	TSharedRef<SButton> DownButton =
-		SNew(SButton)
+	TSharedRef<SButton> DownArrow = SNew(SButton)
 		.Text(LOCTEXT("DownArrow", "↓"))
+		.ContentPadding(FMargin(6, 2))
 		.OnClicked(this, &FUniformGridSlotExtension::HandleShiftRow, 1);
 
-	TSharedRef<SButton> LeftButton =
-		SNew(SButton)
+	TSharedRef<SButton> LeftArrow = SNew(SButton)
 		.Text(LOCTEXT("LeftArrow", "←"))
+		.ContentPadding(FMargin(2, 6))
 		.OnClicked(this, &FUniformGridSlotExtension::HandleShiftColumn, -1);
 
-	TSharedRef<SButton> RightButton =
-		SNew(SButton)
+	TSharedRef<SButton> RightArrow = SNew(SButton)
 		.Text(LOCTEXT("RightArrow", "→"))
+		.ContentPadding(FMargin(2, 6))
 		.OnClicked(this, &FUniformGridSlotExtension::HandleShiftColumn, 1);
 
-	Widgets.Add(UpButton);
-	Widgets.Add(DownButton);
-	Widgets.Add(LeftButton);
-	Widgets.Add(RightButton);
+	UpArrow->SlatePrepass();
+	DownArrow->SlatePrepass();
+	LeftArrow->SlatePrepass();
+	RightArrow->SlatePrepass();
+
+	SurfaceElements.Add(MakeShareable(new FDesignerSurfaceElement(LeftArrow, EExtensionLayoutLocation::CenterLeft, FVector2D(-LeftArrow->GetDesiredSize().X, LeftArrow->GetDesiredSize().Y * -0.5f))));
+	SurfaceElements.Add(MakeShareable(new FDesignerSurfaceElement(RightArrow, EExtensionLayoutLocation::CenterRight, FVector2D(0, RightArrow->GetDesiredSize().Y * -0.5f))));
+	SurfaceElements.Add(MakeShareable(new FDesignerSurfaceElement(UpArrow, EExtensionLayoutLocation::TopCenter, FVector2D(UpArrow->GetDesiredSize().X * -0.5f, -UpArrow->GetDesiredSize().Y))));
+	SurfaceElements.Add(MakeShareable(new FDesignerSurfaceElement(DownArrow, EExtensionLayoutLocation::BottomCenter, FVector2D(DownArrow->GetDesiredSize().X * -0.5f, 0))));
 }
 
 FReply FUniformGridSlotExtension::HandleShiftRow(int32 ShiftAmount)
diff --git a/Engine/Source/Editor/UMGEditor/Private/Extensions/UniformGridSlotExtension.h b/Engine/Source/Editor/UMGEditor/Private/Extensions/UniformGridSlotExtension.h
--- a/Engine/Source/Editor/UMGEditor/Private/Extensions/UniformGridSlotExtension.h
+++ b/Engine/Source/Editor/UMGEditor/Private/Extensions/UniformGridSlotExtension.h
@@ -9,9 +9,9 @@ class FUniformGridSlotExtension : public FDesignerExtension
 public:
 	FUniformGridSlotExtension();
 
-	bool IsActive(const TArray< FWidgetReference >& Selection);
+	virtual bool CanExtendSelection(const TArray< FWidgetReference >& Selection) const override;
 	
-	virtual void BuildWidgetsForSelection(const TArray< FWidgetReference >& Selection, TArray< TSharedRef<SWidget> >& Widgets) override;
+	virtual void ExtendSelection(const TArray< FWidgetReference >& Selection, TArray< TSharedRef<FDesignerSurfaceElement> >& SurfaceElements) override;
 
 private:
 
diff --git a/Engine/Source/Editor/UMGEditor/Private/Extensions/VerticalSlotExtension.cpp b/Engine/Source/Editor/UMGEditor/Private/Extensions/VerticalSlotExtension.cpp
--- a/Engine/Source/Editor/UMGEditor/Private/Extensions/VerticalSlotExtension.cpp
+++ b/Engine/Source/Editor/UMGEditor/Private/Extensions/VerticalSlotExtension.cpp
@@ -10,7 +10,7 @@ FVerticalSlotExtension::FVerticalSlotExtension()
 	ExtensionId = FName(TEXT("VerticalSlot"));
 }
 
-bool FVerticalSlotExtension::IsActive(const TArray< FWidgetReference >& Selection)
+bool FVerticalSlotExtension::CanExtendSelection(const TArray< FWidgetReference >& Selection) const
 {
 	for ( const FWidgetReference& Widget : Selection )
 	{
@@ -23,30 +23,27 @@ bool FVerticalSlotExtension::IsActive(const TArray< FWidgetReference >& Selectio
 	return Selection.Num() == 1;
 }
 
-void FVerticalSlotExtension::BuildWidgetsForSelection(const TArray< FWidgetReference >& Selection, TArray< TSharedRef<SWidget> >& Widgets)
+void FVerticalSlotExtension::ExtendSelection(const TArray< FWidgetReference >& Selection, TArray< TSharedRef<FDesignerSurfaceElement> >& SurfaceElements)
 {
 	SelectionCache = Selection;
 
-	if ( !IsActive(Selection) )
-	{
-		return;
-	}
-
-	TSharedRef<SButton> UpButton =
-		SNew(SButton)
+	TSharedRef<SButton> UpArrow = SNew(SButton)
 		.Text(LOCTEXT("UpArrow", "↑"))
+		.ContentPadding(FMargin(6, 2))
+		.IsEnabled(this, &FVerticalSlotExtension::CanShift, -1)
 		.OnClicked(this, &FVerticalSlotExtension::HandleShiftVertical, -1);
 
-	TSharedRef<SButton> DownButton =
-		SNew(SButton)
+	TSharedRef<SButton> DownArrow = SNew(SButton)
 		.Text(LOCTEXT("DownArrow", "↓"))
+		.ContentPadding(FMargin(6, 2))
+		.IsEnabled(this, &FVerticalSlotExtension::CanShift, 1)
 		.OnClicked(this, &FVerticalSlotExtension::HandleShiftVertical, 1);
 
-	UpButton->SetEnabled(TAttribute<bool>::Create(TAttribute<bool>::FGetter::CreateSP(this, &FVerticalSlotExtension::CanShift, -1)));
-	DownButton->SetEnabled(TAttribute<bool>::Create(TAttribute<bool>::FGetter::CreateSP(this, &FVerticalSlotExtension::CanShift, 1)));
+	UpArrow->SlatePrepass();
+	DownArrow->SlatePrepass();
 
-	Widgets.Add(UpButton);
-	Widgets.Add(DownButton);
+	SurfaceElements.Add(MakeShareable(new FDesignerSurfaceElement(UpArrow, EExtensionLayoutLocation::TopCenter, FVector2D(UpArrow->GetDesiredSize().X * -0.5f, -UpArrow->GetDesiredSize().Y))));
+	SurfaceElements.Add(MakeShareable(new FDesignerSurfaceElement(DownArrow, EExtensionLayoutLocation::BottomCenter, FVector2D(DownArrow->GetDesiredSize().X * -0.5f, 0))));
 }
 
 bool FVerticalSlotExtension::CanShift(int32 ShiftAmount) const
@@ -65,6 +62,8 @@ FReply FVerticalSlotExtension::HandleShiftVertical(int32 ShiftAmount)
 		ShiftVertical(Selection.GetTemplate(), ShiftAmount);
 	}
 
+	EndTransaction();
+
 	//TODO UMG Reorder the live slot without rebuilding the structure
 	FBlueprintEditorUtils::MarkBlueprintAsStructurallyModified(Blueprint);
 
diff --git a/Engine/Source/Editor/UMGEditor/Private/Extensions/VerticalSlotExtension.h b/Engine/Source/Editor/UMGEditor/Private/Extensions/VerticalSlotExtension.h
--- a/Engine/Source/Editor/UMGEditor/Private/Extensions/VerticalSlotExtension.h
+++ b/Engine/Source/Editor/UMGEditor/Private/Extensions/VerticalSlotExtension.h
@@ -9,9 +9,9 @@ class FVerticalSlotExtension : public FDesignerExtension
 public:
 	FVerticalSlotExtension();
 
-	bool IsActive(const TArray< FWidgetReference >& Selection);
+	virtual bool CanExtendSelection(const TArray< FWidgetReference >& Selection) const override;
 	
-	virtual void BuildWidgetsForSelection(const TArray< FWidgetReference >& Selection, TArray< TSharedRef<SWidget> >& Widgets) override;
+	virtual void ExtendSelection(const TArray< FWidgetReference >& Selection, TArray< TSharedRef<FDesignerSurfaceElement> >& SurfaceElements) override;
 
 private:
 
diff --git a/Engine/Source/Editor/UMGEditor/Private/Hierarchy/SHierarchyView.cpp b/Engine/Source/Editor/UMGEditor/Private/Hierarchy/SHierarchyView.cpp
--- a/Engine/Source/Editor/UMGEditor/Private/Hierarchy/SHierarchyView.cpp
+++ b/Engine/Source/Editor/UMGEditor/Private/Hierarchy/SHierarchyView.cpp
@@ -20,29 +20,22 @@
 void SHierarchyView::Construct(const FArguments& InArgs, TSharedPtr<FWidgetBlueprintEditor> InBlueprintEditor, USimpleConstructionScript* InSCS)
 {
 	BlueprintEditor = InBlueprintEditor;
-	bRefreshRequested = false;
+	bRebuildTreeRequested = false;
 
+	// register for any objects replaced
+	GEditor->OnObjectsReplaced().AddRaw(this, &SHierarchyView::OnObjectsReplaced);
+
+	// Create the filter for searching in the tree
 	SearchBoxWidgetFilter = MakeShareable(new WidgetTextFilter(WidgetTextFilter::FItemToStringArray::CreateSP(this, &SHierarchyView::TransformWidgetToString)));
 
 	UWidgetBlueprint* Blueprint = GetBlueprint();
-	Blueprint->OnChanged().AddSP(this, &SHierarchyView::OnBlueprintChanged);
+	Blueprint->OnChanged().AddRaw(this, &SHierarchyView::OnBlueprintChanged);
 
 	FilterHandler = MakeShareable(new TreeFilterHandler<UWidget*>());
 	FilterHandler->SetFilter(SearchBoxWidgetFilter.Get());
 	FilterHandler->SetRootItems(&RootWidgets, &TreeRootWidgets);
 	FilterHandler->SetGetChildrenDelegate(TreeFilterHandler<UWidget*>::FOnGetChildren::CreateRaw(this, &SHierarchyView::WidgetHierarchy_OnGetChildren));
 
-	SAssignNew(WidgetTreeView, STreeView< UWidget* >)
-	.ItemHeight(20.0f)
-	.SelectionMode(ESelectionMode::Single)
-	.OnGetChildren(FilterHandler.ToSharedRef(), &TreeFilterHandler<UWidget*>::OnGetFilteredChildren)
-	.OnGenerateRow(this, &SHierarchyView::WidgetHierarchy_OnGenerateRow)
-	.OnSelectionChanged(this, &SHierarchyView::WidgetHierarchy_OnSelectionChanged)
-	.OnContextMenuOpening(this, &SHierarchyView::WidgetHierarchy_OnContextMenuOpening)
-	.TreeItemsSource(&TreeRootWidgets);
-
-	FilterHandler->SetTreeView(WidgetTreeView.Get());
-
 	ChildSlot
 	[
 		SNew(SBorder)
@@ -62,14 +55,15 @@ void SHierarchyView::Construct(const FArguments& InArgs, TSharedPtr<FWidgetBluep
 			+ SVerticalBox::Slot()
 			.FillHeight(1.0f)
 			[
-				SNew(SScrollBorder, WidgetTreeView.ToSharedRef())
-				[
-					WidgetTreeView.ToSharedRef()
-				]
+				SAssignNew(TreeViewArea, SBorder)
+				.Padding(0)
+				.BorderImage( FEditorStyle::GetBrush( "NoBrush" ) )
 			]
 		]
 	];
 
+	RebuildTreeView();
+
 	BlueprintEditor.Pin()->OnSelectedWidgetsChanged.AddRaw(this, &SHierarchyView::OnEditorSelectionChanged);
 
 	bRefreshRequested = true;
@@ -87,10 +81,18 @@ SHierarchyView::~SHierarchyView()
 	{
 		BlueprintEditor.Pin()->OnSelectedWidgetsChanged.RemoveAll(this);
 	}
+
+	GEditor->OnObjectsReplaced().RemoveAll(this);
 }
 
 void SHierarchyView::Tick(const FGeometry& AllottedGeometry, const double InCurrentTime, const float InDeltaTime)
 {
+	if ( bRebuildTreeRequested )
+	{
+		bRebuildTreeRequested = false;
+		RebuildTreeView();
+	}
+
 	if ( bRefreshRequested )
 	{
 		bRefreshRequested = false;
@@ -182,7 +184,7 @@ void SHierarchyView::OnBlueprintChanged(UBlueprint* InBlueprint)
 {
 	if ( InBlueprint )
 	{
-		RefreshTree();
+		bRefreshRequested = true;
 	}
 }
 
@@ -209,6 +211,8 @@ TSharedPtr<SWidget> SHierarchyView::WidgetHierarchy_OnContextMenuOpening()
 
 void SHierarchyView::WidgetHierarchy_OnGetChildren(UWidget* InParent, TArray< UWidget* >& OutChildren)
 {
+	VisibleItems.Add(InParent);
+
 	UPanelWidget* Widget = Cast<UPanelWidget>(InParent);
 	if ( Widget )
 	{
@@ -258,6 +262,8 @@ FReply SHierarchyView::HandleDeleteSelected()
 
 void SHierarchyView::RefreshTree()
 {
+	VisibleItems.Empty();
+
 	RootWidgets.Empty();
 	UWidgetBlueprint* Blueprint = GetBlueprint();
 	if (Blueprint->WidgetTree->RootWidget)
@@ -267,6 +273,41 @@ void SHierarchyView::RefreshTree()
 	FilterHandler->RefreshAndFilterTree();
 }
 
+void SHierarchyView::RebuildTreeView()
+{
+	SAssignNew(WidgetTreeView, STreeView< UWidget* >)
+		.ItemHeight(20.0f)
+		.SelectionMode(ESelectionMode::Single)
+		.OnGetChildren(FilterHandler.ToSharedRef(), &TreeFilterHandler<UWidget*>::OnGetFilteredChildren)
+		.OnGenerateRow(this, &SHierarchyView::WidgetHierarchy_OnGenerateRow)
+		.OnSelectionChanged(this, &SHierarchyView::WidgetHierarchy_OnSelectionChanged)
+		.OnContextMenuOpening(this, &SHierarchyView::WidgetHierarchy_OnContextMenuOpening)
+		.TreeItemsSource(&TreeRootWidgets);
+
+	FilterHandler->SetTreeView(WidgetTreeView.Get());
+
+	TreeViewArea->SetContent(
+		SNew(SScrollBorder, WidgetTreeView.ToSharedRef())
+		[
+			WidgetTreeView.ToSharedRef()
+		]);
+}
+
+void SHierarchyView::OnObjectsReplaced(const TMap<UObject*, UObject*>& ReplacementMap)
+{
+	if ( !bRebuildTreeRequested )
+	{
+		for ( UWidget* Widget : VisibleItems )
+		{
+			if ( ReplacementMap.Contains(Widget) )
+			{
+				bRefreshRequested = true;
+				bRebuildTreeRequested = true;
+			}
+		}
+	}
+}
+
 //@TODO UMG Drop widgets onto the tree, when nothing is present, if there is a root node present, what happens then, let the root node attempt to place it?
 
 #undef LOCTEXT_NAMESPACE
diff --git a/Engine/Source/Editor/UMGEditor/Private/Hierarchy/SHierarchyView.h b/Engine/Source/Editor/UMGEditor/Private/Hierarchy/SHierarchyView.h
--- a/Engine/Source/Editor/UMGEditor/Private/Hierarchy/SHierarchyView.h
+++ b/Engine/Source/Editor/UMGEditor/Private/Hierarchy/SHierarchyView.h
@@ -57,6 +57,9 @@ private:
 	/** Rebuilds the tree structure based on the current filter options */
 	void RefreshTree();
 
+	/** Completely regenerates the treeview */
+	void RebuildTreeView();
+
 	/** Handles the menu clicking to delete the selected widgets. */
 	FReply HandleDeleteSelected();
 
@@ -75,6 +78,9 @@ private:
 	/** Transforms the widget into a searchable string */
 	void TransformWidgetToString(UWidget* Widget, OUT TArray< FString >& Array);
 
+	/** Called when a Blueprint is recompiled and live objects are swapped out for replacements */
+	void OnObjectsReplaced(const TMap<UObject*, UObject*>& ReplacementMap);
+
 private:
 
 	/** Cached pointer to the blueprint editor that owns this tree. */
@@ -89,6 +95,12 @@ private:
 	/** The root widgets which are actually displayed by the TreeView which will be managed by the TreeFilterHandler. */
 	TArray< UWidget* > TreeRootWidgets;
 
+	/** Gets the visible items displayed in the tree view. */
+	TSet< UWidget* > VisibleItems;
+
+	/** The widget containing the treeview */
+	TSharedPtr<SBorder> TreeViewArea;
+
 	/** The widget hierarchy slate treeview widget */
 	TSharedPtr< STreeView< UWidget* > > WidgetTreeView;
 
@@ -97,4 +109,7 @@ private:
 
 	/** Has a full refresh of the tree been requested?  This happens when the user is filtering the tree */
 	bool bRefreshRequested;
+
+	/** Is the tree in such a changed state that the whole widget needs rebuilding? */
+	bool bRebuildTreeRequested;
 };
diff --git a/Engine/Source/Editor/UMGEditor/Private/Palette/SPaletteView.cpp b/Engine/Source/Editor/UMGEditor/Private/Palette/SPaletteView.cpp
--- a/Engine/Source/Editor/UMGEditor/Private/Palette/SPaletteView.cpp
+++ b/Engine/Source/Editor/UMGEditor/Private/Palette/SPaletteView.cpp
@@ -144,15 +144,14 @@ void SPaletteView::Construct(const FArguments& InArgs, TSharedPtr<FBlueprintEdit
 {
 	BlueprintEditor = InBlueprintEditor;
 
+	// register for any objects replaced
+	GEditor->OnObjectsReplaced().AddRaw(this, &SPaletteView::OnObjectsReplaced);
+
 	UBlueprint* BP = InBlueprintEditor->GetBlueprintObj();
 
 	WidgetFilter = MakeShareable(new WidgetViewModelTextFilter(
 		WidgetViewModelTextFilter::FItemToStringArray::CreateSP(this, &SPaletteView::TransformWidgetViewModelToString)));
 
-	BuildWidgetList();
-
-	//FCoreDelegates::OnObjectPropertyChanged.AddRaw(this, &SPaletteView::OnObjectPropertyChanged);
-
 	FilterHandler = MakeShareable(new PaletteFilterHandler());
 	FilterHandler->SetFilter(WidgetFilter.Get());
 	FilterHandler->SetRootItems(&WidgetViewModels, &TreeWidgetViewModels);
@@ -189,14 +188,16 @@ void SPaletteView::Construct(const FArguments& InArgs, TSharedPtr<FBlueprintEdit
 		]
 	];
 
-	LoadItemExpanssion();
 	bRefreshRequested = true;
+
+	BuildWidgetList();
+	LoadItemExpanssion();
+
+	bRebuildRquested = false;
 }
 
 SPaletteView::~SPaletteView()
 {
-	//FCoreDelegates::OnObjectPropertyChanged.RemoveAll(this);
-
 	// If the filter is enabled, disable it before saving the expanded items since
 	// filtering expands all items by default.
 	if (FilterHandler->GetIsEnabled())
@@ -332,6 +333,27 @@ TSharedRef<ITableRow> SPaletteView::OnGenerateWidgetTemplateItem(TSharedPtr<FWid
 
 void SPaletteView::Tick(const FGeometry& AllottedGeometry, const double InCurrentTime, const float InDeltaTime)
 {
+	if ( bRebuildRquested )
+	{
+		// Save the old expanded items temporarily
+		TSet<TSharedPtr<FWidgetViewModel>> ExpandedItems;
+		WidgetTemplatesView->GetExpandedItems(ExpandedItems);
+
+		BuildWidgetList();
+
+		// Restore the expansion state
+		for ( TSharedPtr<FWidgetViewModel>& ExpandedItem : ExpandedItems )
+		{
+			for ( TSharedPtr<FWidgetViewModel>& ViewModel : WidgetViewModels )
+			{
+				if ( ViewModel->GetName().EqualTo(ExpandedItem->GetName()) )
+				{
+					WidgetTemplatesView->SetItemExpansion(ViewModel, true);
+				}
+			}
+		}
+	}
+
 	if (bRefreshRequested)
 	{
 		bRefreshRequested = false;
@@ -344,4 +366,10 @@ void SPaletteView::TransformWidgetViewModelToString(TSharedPtr<FWidgetViewModel>
 	Array.Add(WidgetViewModel->GetFilterString());
 }
 
+void SPaletteView::OnObjectsReplaced(const TMap<UObject*, UObject*>& ReplacementMap)
+{
+	//bRefreshRequested = true;
+	//bRebuildRquested = true;
+}
+
 #undef LOCTEXT_NAMESPACE
diff --git a/Engine/Source/Editor/UMGEditor/Private/Palette/SPaletteView.h b/Engine/Source/Editor/UMGEditor/Private/Palette/SPaletteView.h
--- a/Engine/Source/Editor/UMGEditor/Private/Palette/SPaletteView.h
+++ b/Engine/Source/Editor/UMGEditor/Private/Palette/SPaletteView.h
@@ -59,6 +59,9 @@ private:
 	void LoadItemExpanssion();
 	void SaveItemExpansion();
 
+	/** Called when a Blueprint is recompiled and live objects are swapped out for replacements */
+	void OnObjectsReplaced(const TMap<UObject*, UObject*>& ReplacementMap);
+
 	void AddWidgetTemplate(TSharedPtr<FWidgetTemplate> Template);
 
 	/** Transforms the widget view model into a searchable string. */
@@ -88,4 +91,7 @@ private:
 
 	bool bRefreshRequested;
 	FText SearchText;
+
+	/** Controls rebuilding the list of spawnable widgets */
+	bool bRebuildRquested;
 };
diff --git a/Engine/Source/Editor/UMGEditor/Private/UMGEditorPrivatePCH.h b/Engine/Source/Editor/UMGEditor/Private/UMGEditorPrivatePCH.h
--- a/Engine/Source/Editor/UMGEditor/Private/UMGEditorPrivatePCH.h
+++ b/Engine/Source/Editor/UMGEditor/Private/UMGEditorPrivatePCH.h
@@ -29,6 +29,8 @@
 
 #include "AssetTypeActions_Base.h"
 
+#include "DesignTimeUtils.h"
+
 #include "DesignerExtension.h"
 #include "VerticalSlotExtension.h"
 #include "HorizontalSlotExtension.h"
diff --git a/Engine/Source/Editor/UMGEditor/Private/WidgetBlueprintEditor.cpp b/Engine/Source/Editor/UMGEditor/Private/WidgetBlueprintEditor.cpp
--- a/Engine/Source/Editor/UMGEditor/Private/WidgetBlueprintEditor.cpp
+++ b/Engine/Source/Editor/UMGEditor/Private/WidgetBlueprintEditor.cpp
@@ -30,6 +30,8 @@ FWidgetBlueprintEditor::~FWidgetBlueprintEditor()
 	{
 		Blueprint->OnChanged().RemoveAll(this);
 	}
+
+	GEditor->OnObjectsReplaced().RemoveAll(this);
 }
 
 void FWidgetBlueprintEditor::InitWidgetBlueprintEditor(const EToolkitMode::Type Mode, const TSharedPtr< IToolkitHost >& InitToolkitHost, const TArray<UBlueprint*>& InBlueprints, bool bShouldOpenInDefaultsMode)
@@ -39,6 +41,9 @@ void FWidgetBlueprintEditor::InitWidgetBlueprintEditor(const EToolkitMode::Type
 
 	InitBlueprintEditor(Mode, InitToolkitHost, InBlueprints, bShouldOpenInDefaultsMode);
 
+	// register for any objects replaced
+	GEditor->OnObjectsReplaced().AddSP(this, &FWidgetBlueprintEditor::OnObjectsReplaced);
+
 	UWidgetBlueprint* Blueprint = GetWidgetBlueprintObj();
 
 	// If this blueprint is empty, add a canvas panel as the root widget.
@@ -176,6 +181,34 @@ void FWidgetBlueprintEditor::OnBlueprintChanged(UBlueprint* InBlueprint)
 	}
 }
 
+void FWidgetBlueprintEditor::OnObjectsReplaced(const TMap<UObject*, UObject*>& ReplacementMap)
+{
+	// Objects being replaced could force a selection change
+	OnSelectedWidgetsChanging.Broadcast();
+
+	TSet<FWidgetReference> TempSelection;
+	for ( const FWidgetReference& WidgetRef : SelectedWidgets )
+	{
+		UWidget* OldTemplate = WidgetRef.GetTemplate();
+
+		UObject* const* NewObject = ReplacementMap.Find(OldTemplate);
+		if ( NewObject )
+		{
+			UWidget* NewTemplate = Cast<UWidget>(*NewObject);
+			TempSelection.Add(FWidgetReference::FromTemplate(SharedThis(this), NewTemplate));
+		}
+		else
+		{
+			TempSelection.Add(WidgetRef);
+		}
+	}
+
+	SelectedWidgets = TempSelection;
+
+	// Fire the selection updated event to ensure everyone is watching the same widgets.
+	OnSelectedWidgetsChanged.Broadcast();
+}
+
 bool FWidgetBlueprintEditor::CanDeleteSelectedWidgets()
 {
 	TSet<FWidgetReference> Widgets = GetSelectedWidgets();
diff --git a/Engine/Source/Editor/UMGEditor/Private/WidgetBlueprintEditor.h b/Engine/Source/Editor/UMGEditor/Private/WidgetBlueprintEditor.h
--- a/Engine/Source/Editor/UMGEditor/Private/WidgetBlueprintEditor.h
+++ b/Engine/Source/Editor/UMGEditor/Private/WidgetBlueprintEditor.h
@@ -90,6 +90,9 @@ private:
 	/** Called whenever the blueprint is structurally changed. */
 	virtual void OnBlueprintChanged(UBlueprint* InBlueprint) override;
 
+	/** Called when objects need to be swapped out for new versions, like after a blueprint recompile. */
+	void OnObjectsReplaced(const TMap<UObject*, UObject*>& ReplacementMap);
+
 	/** Destroy the current preview GUI object */
 	void DestroyPreview();
 
diff --git a/Engine/Source/Editor/UMGEditor/Private/WidgetReference.cpp b/Engine/Source/Editor/UMGEditor/Private/WidgetReference.cpp
new file mode 100644
--- /dev/null
+++ b/Engine/Source/Editor/UMGEditor/Private/WidgetReference.cpp
@@ -0,0 +1,88 @@
+// Copyright 1998-2014 Epic Games, Inc. All Rights Reserved.
+
+#include "UMGEditorPrivatePCH.h"
+
+#define LOCTEXT_NAMESPACE "UMG"
+
+// FWidgetReference
+
+FWidgetReference::FWidgetReference()
+	: WidgetEditor()
+	, TemplateWidget(NULL)
+{
+}
+
+FWidgetReference::FWidgetReference(TSharedPtr<FWidgetBlueprintEditor> WidgetEditor, UWidget* TemplateWidget)
+	: WidgetEditor(WidgetEditor)
+	, TemplateWidget(TemplateWidget)
+{
+	// register for any objects replaced
+	GEditor->OnObjectsReplaced().AddRaw(this, &FWidgetReference::OnObjectsReplaced);
+}
+
+FWidgetReference::~FWidgetReference()
+{
+	GEditor->OnObjectsReplaced().RemoveAll(this);
+}
+
+FWidgetReference FWidgetReference::FromTemplate(TSharedPtr<FWidgetBlueprintEditor> InWidgetEditor, UWidget* TemplateWidget)
+{
+	return FWidgetReference(InWidgetEditor, TemplateWidget);
+}
+
+FWidgetReference FWidgetReference::FromPreview(TSharedPtr<FWidgetBlueprintEditor> InWidgetEditor, UWidget* PreviewWidget)
+{
+	if ( InWidgetEditor.IsValid() )
+	{
+		UUserWidget* PreviewRoot = InWidgetEditor->GetPreview();
+		if ( PreviewRoot )
+		{
+			UWidgetBlueprint* Blueprint = InWidgetEditor->GetWidgetBlueprintObj();
+
+			if ( PreviewWidget )
+			{
+				FString Name = PreviewWidget->GetName();
+				return FromTemplate(InWidgetEditor, Blueprint->WidgetTree->FindWidget(Name));
+			}
+		}
+	}
+
+	return FWidgetReference();
+}
+
+bool FWidgetReference::IsValid() const
+{
+	return TemplateWidget.Get() != NULL && GetPreview();
+}
+
+UWidget* FWidgetReference::GetTemplate() const
+{
+	return TemplateWidget.Get();
+}
+
+UWidget* FWidgetReference::GetPreview() const
+{
+	if ( WidgetEditor.IsValid() )
+	{
+		UUserWidget* PreviewRoot = WidgetEditor.Pin()->GetPreview();
+
+		if ( PreviewRoot && TemplateWidget.Get() )
+		{
+			UWidget* PreviewWidget = PreviewRoot->GetHandleFromName(TemplateWidget.Get()->GetName());
+			return PreviewWidget;
+		}
+	}
+
+	return NULL;
+}
+
+void FWidgetReference::OnObjectsReplaced(const TMap<UObject*, UObject*>& ReplacementMap)
+{
+	UObject* const* NewObject = ReplacementMap.Find(TemplateWidget.Get());
+	if ( NewObject )
+	{
+		TemplateWidget = Cast<UWidget>(*NewObject);
+	}
+}
+
+#undef LOCTEXT_NAMESPACE
diff --git a/Engine/Source/Editor/UMGEditor/Public/DesignerExtension.h b/Engine/Source/Editor/UMGEditor/Public/DesignerExtension.h
--- a/Engine/Source/Editor/UMGEditor/Public/DesignerExtension.h
+++ b/Engine/Source/Editor/UMGEditor/Public/DesignerExtension.h
@@ -8,6 +8,74 @@ class UWidget;
 class SWidget;
 class UWidgetBlueprint;
 
+/**
+ * The location where the widget should appear
+ */
+namespace EExtensionLayoutLocation
+{
+	enum Type
+	{
+		Absolute,
+
+		TopLeft,
+		TopCenter,
+		TopRight,
+
+		CenterLeft,
+		CenterCenter,
+		CenterRight,
+
+		BottomLeft,
+		BottomCenter,
+		BottomRight,
+	};
+}
+
+/**
+ * The basic element returned for extending the design surface.
+ */
+class UMGEDITOR_API FDesignerSurfaceElement : public TSharedFromThis < FDesignerSurfaceElement >
+{
+public:
+	FDesignerSurfaceElement(TSharedRef<SWidget> InWidget, EExtensionLayoutLocation::Type InLocation, FVector2D InOffset = FVector2D(0,0))
+		: Widget(InWidget)
+		, Location(InLocation)
+		, Offset(InOffset)
+	{
+	}
+
+	/** Gets the widget that will be laid out in the design surface for extending the capability of the selected widget */
+	TSharedRef<SWidget> GetWidget() const
+	{
+		return Widget;
+	}
+
+	/** Gets the location where the widget will be appear */
+	EExtensionLayoutLocation::Type GetLocation() const
+	{
+		return Location;
+	}
+
+	/** Sets the offset after laid out in that location */
+	void SetOffset(FVector2D InOffset)
+	{
+		Offset = InOffset;
+	}
+
+	/** Gets the offset after being laid out. */
+	FVector2D GetOffset() const
+	{
+		return Offset;
+	}
+
+protected:
+	TSharedRef<SWidget> Widget;
+
+	EExtensionLayoutLocation::Type Location;
+
+	FVector2D Offset;
+};
+
 /**
  * The Designer extension allows developers to provide additional widgets and custom painting to the designer surface for
  * specific widgets.  Allowing for a more customized and specific editors for the different widgets.
@@ -21,8 +89,19 @@ public:
 	/** Initializes the designer extension, this is called the first time a designer extension is registered */
 	virtual void Initialize(UWidgetBlueprint* InBlueprint);
 
+	virtual bool CanExtendSelection(const TArray< FWidgetReference >& Selection) const
+	{
+		return false;
+	}
+
 	/** Called every time the selection in the designer changes. */
-	virtual void BuildWidgetsForSelection(const TArray< FWidgetReference >& Selection, TArray< TSharedRef<SWidget> >& Widgets) = 0;
+	virtual void ExtendSelection(const TArray< FWidgetReference >& Selection, TArray< TSharedRef<FDesignerSurfaceElement> >& SurfaceElements)
+	{
+	}
+
+	virtual void Tick(const FGeometry& AllottedGeometry, const double InCurrentTime, const float InDeltaTime)
+	{
+	}
 
 	virtual void Paint(const TSet< FWidgetReference >& Selection, const FGeometry& AllottedGeometry, const FSlateRect& MyClippingRect, FSlateWindowElementList& OutDrawElements, int32 LayerId, const FWidgetStyle& InWidgetStyle, bool bParentEnabled) const
 	{
diff --git a/Engine/Source/Editor/UMGEditor/Public/WidgetReference.h b/Engine/Source/Editor/UMGEditor/Public/WidgetReference.h
--- a/Engine/Source/Editor/UMGEditor/Public/WidgetReference.h
+++ b/Engine/Source/Editor/UMGEditor/Public/WidgetReference.h
@@ -26,6 +26,9 @@ public:
 	/** Constructor for the null widget reference.  Not intended for normal use. */
 	FWidgetReference();
 
+	/** Destructor */
+	~FWidgetReference();
+
 	/** @returns true if both the template and the preview pointers are valid. */
 	bool IsValid() const;
 
@@ -44,6 +47,8 @@ public:
 private:
 	FWidgetReference(TSharedPtr<class FWidgetBlueprintEditor> WidgetEditor, UWidget* TemplateWidget);
 
+	void OnObjectsReplaced(const TMap<UObject*, UObject*>& ReplacementMap);
+
 private:
 
 	TWeakPtr<class FWidgetBlueprintEditor> WidgetEditor;
diff --git a/Engine/Source/Runtime/Slate/Private/Widgets/Layout/SBox.cpp b/Engine/Source/Runtime/Slate/Private/Widgets/Layout/SBox.cpp
--- a/Engine/Source/Runtime/Slate/Private/Widgets/Layout/SBox.cpp
+++ b/Engine/Source/Runtime/Slate/Private/Widgets/Layout/SBox.cpp
@@ -22,7 +22,13 @@ void SBox::Construct( const FArguments& InArgs )
 	];
 }
 
-
+void SBox::SetContent(const TSharedRef< SWidget >& InContent)
+{
+	ChildSlot
+	[
+		InContent
+	];
+}
 
 FVector2D SBox::ComputeDesiredSize() const
 {
diff --git a/Engine/Source/Runtime/Slate/Public/Widgets/Layout/SBox.h b/Engine/Source/Runtime/Slate/Public/Widgets/Layout/SBox.h
--- a/Engine/Source/Runtime/Slate/Public/Widgets/Layout/SBox.h
+++ b/Engine/Source/Runtime/Slate/Public/Widgets/Layout/SBox.h
@@ -60,6 +60,10 @@ class SLATE_API SBox : public SPanel
 		virtual FChildren* GetChildren() override;
 		virtual int32 OnPaint( const FPaintArgs& Args, const FGeometry& AllottedGeometry, const FSlateRect& MyClippingRect, FSlateWindowElementList& OutDrawElements, int32 LayerId, const FWidgetStyle& InWidgetStyle, bool bParentEnabled ) const override;
 
+		/**
+		 * See the Content slot.
+		 */
+		void SetContent(const TSharedRef< SWidget >& InContent);
 private:
 
 		FBoxSlot ChildSlot;
diff --git a/Engine/Source/Runtime/UMG/Private/Components/Border.cpp b/Engine/Source/Runtime/UMG/Private/Components/Border.cpp
--- a/Engine/Source/Runtime/UMG/Private/Components/Border.cpp
+++ b/Engine/Source/Runtime/UMG/Private/Components/Border.cpp
@@ -13,9 +13,6 @@ UBorder::UBorder(const FPostConstructInitializeProperties& PCIP)
 {
 	bIsVariable = false;
 
-	HorizontalAlignment = HAlign_Fill;
-	VerticalAlignment = VAlign_Fill;
-
 	ContentScale = FVector2D(1.0f, 1.0f);
 	ContentColorAndOpacity = FLinearColor::White;
 
@@ -24,8 +21,6 @@ UBorder::UBorder(const FPostConstructInitializeProperties& PCIP)
 	ForegroundColor = FLinearColor::Black;
 
 	SBorder::FArguments BorderDefaults;
-
-	ContentPadding = BorderDefaults._Padding.Get();
 	bShowEffectWhenDisabled = BorderDefaults._ShowEffectWhenDisabled.Get();
 }
 
@@ -42,7 +37,7 @@ TSharedRef<SWidget> UBorder::RebuildWidget()
 	
 	if ( GetChildrenCount() > 0 )
 	{
-		MyBorder->SetContent(GetContentSlot()->Content ? GetContentSlot()->Content->TakeWidget() : SNullWidget::NullWidget);
+		Cast<UBorderSlot>(GetContentSlot())->BuildSlot(MyBorder.ToSharedRef());
 	}
 
 	return MyBorder.ToSharedRef();
@@ -51,10 +46,6 @@ TSharedRef<SWidget> UBorder::RebuildWidget()
 void UBorder::SyncronizeProperties()
 {
 	Super::SyncronizeProperties();
-
-	MyBorder->SetHAlign(HorizontalAlignment);
-	MyBorder->SetVAlign(VerticalAlignment);
-	MyBorder->SetPadding(ContentPadding);
 	
 	MyBorder->SetBorderBackgroundColor(BrushColor);
 	MyBorder->SetColorAndOpacity(ContentColorAndOpacity);
@@ -73,12 +64,17 @@ void UBorder::SyncronizeProperties()
 	MyBorder->SetOnMouseDoubleClick(BIND_UOBJECT_DELEGATE(FPointerEventHandler, HandleMouseDoubleClick));
 }
 
+UClass* UBorder::GetSlotClass() const
+{
+	return UBorderSlot::StaticClass();
+}
+
 void UBorder::OnSlotAdded(UPanelSlot* Slot)
 {
-	// Add the child to the live canvas if it already exists
+	// Add the child to the live slot if it already exists
 	if ( MyBorder.IsValid() )
 	{
-		MyBorder->SetContent(Slot->Content ? Slot->Content->TakeWidget() : SNullWidget::NullWidget);
+		Cast<UBorderSlot>(Slot)->BuildSlot(MyBorder.ToSharedRef());
 	}
 }
 
@@ -109,15 +105,6 @@ void UBorder::SetForegroundColor(FLinearColor InForegroundColor)
 	}
 }
 
-void UBorder::SetContentPadding(FMargin InContentPadding)
-{
-	ContentPadding = InContentPadding;
-	if ( MyBorder.IsValid() )
-	{
-		MyBorder->SetPadding(InContentPadding);
-	}
-}
-
 const FSlateBrush* UBorder::GetBorderBrush() const
 {
 	if ( Brush == NULL )
@@ -169,6 +156,27 @@ FReply UBorder::HandleMouseDoubleClick(const FGeometry& Geometry, const FPointer
 	return FReply::Unhandled();
 }
 
+void UBorder::PostLoad()
+{
+	Super::PostLoad();
+
+	if ( GetChildrenCount() > 0 )
+	{
+		//TODO UMG Pre-Release Upgrade, now have slots of their own.  Convert existing slot to new slot.
+		if ( UPanelSlot* PanelSlot = GetContentSlot() )
+		{
+			UBorderSlot* BorderSlot = Cast<UBorderSlot>(PanelSlot);
+			if ( BorderSlot == NULL )
+			{
+				BorderSlot = ConstructObject<UBorderSlot>(UBorderSlot::StaticClass(), this);
+				BorderSlot->Content = GetContentSlot()->Content;
+				BorderSlot->Content->Slot = BorderSlot;
+				Slots[0] = BorderSlot;
+			}
+		}
+	}
+}
+
 #if WITH_EDITOR
 
 const FSlateBrush* UBorder::GetEditorIcon()
diff --git a/Engine/Source/Runtime/UMG/Private/Components/BorderSlot.cpp b/Engine/Source/Runtime/UMG/Private/Components/BorderSlot.cpp
new file mode 100644
--- /dev/null
+++ b/Engine/Source/Runtime/UMG/Private/Components/BorderSlot.cpp
@@ -0,0 +1,71 @@
+// Copyright 1998-2014 Epic Games, Inc. All Rights Reserved.
+
+#include "UMGPrivatePCH.h"
+
+/////////////////////////////////////////////////////
+// UBorderSlot
+
+UBorderSlot::UBorderSlot(const FPostConstructInitializeProperties& PCIP)
+	: Super(PCIP)
+{
+	Padding = FMargin(4, 2);
+
+	HorizontalAlignment = HAlign_Fill;
+	VerticalAlignment = VAlign_Fill;
+
+	SBorder::FArguments BorderDefaults;
+
+	Padding = BorderDefaults._Padding.Get();
+}
+
+void UBorderSlot::ReleaseNativeWidget()
+{
+	Super::ReleaseNativeWidget();
+
+	Border.Reset();
+}
+
+void UBorderSlot::BuildSlot(TSharedRef<SBorder> InBorder)
+{
+	Border = InBorder;
+
+	Border->SetPadding(Padding);
+	Border->SetHAlign(HorizontalAlignment);
+	Border->SetVAlign(VerticalAlignment);
+
+	Border->SetContent(Content ? Content->TakeWidget() : SNullWidget::NullWidget);
+}
+
+void UBorderSlot::SetPadding(FMargin InPadding)
+{
+	Padding = InPadding;
+	if ( Border.IsValid() )
+	{
+		Border->SetPadding(InPadding);
+	}
+}
+
+void UBorderSlot::SetHorizontalAlignment(EHorizontalAlignment InHorizontalAlignment)
+{
+	HorizontalAlignment = InHorizontalAlignment;
+	if ( Border.IsValid() )
+	{
+		Border->SetHAlign(InHorizontalAlignment);
+	}
+}
+
+void UBorderSlot::SetVerticalAlignment(EVerticalAlignment InVerticalAlignment)
+{
+	VerticalAlignment = InVerticalAlignment;
+	if ( Border.IsValid() )
+	{
+		Border->SetVAlign(InVerticalAlignment);
+	}
+}
+
+void UBorderSlot::SyncronizeProperties()
+{
+	SetPadding(Padding);
+	SetHorizontalAlignment(HorizontalAlignment);
+	SetVerticalAlignment(VerticalAlignment);
+}
diff --git a/Engine/Source/Runtime/UMG/Private/Components/Button.cpp b/Engine/Source/Runtime/UMG/Private/Components/Button.cpp
--- a/Engine/Source/Runtime/UMG/Private/Components/Button.cpp
+++ b/Engine/Source/Runtime/UMG/Private/Components/Button.cpp
@@ -151,13 +151,16 @@ void UButton::PostLoad()
 	if ( GetChildrenCount() > 0 )
 	{
 		//TODO UMG Pre-Release Upgrade, now buttons have slots of their own.  Convert existing slot to new slot.
-		UButtonSlot* ButtonSlot = Cast<UButtonSlot>(GetContentSlot());
-		if ( ButtonSlot == NULL )
+		if ( UPanelSlot* PanelSlot = GetContentSlot() )
 		{
-			ButtonSlot = ConstructObject<UButtonSlot>(UButtonSlot::StaticClass(), this);
-			ButtonSlot->Content = GetContentSlot()->Content;
-			ButtonSlot->Content->Slot = ButtonSlot;
-			Slots[0] = ButtonSlot;
+			UButtonSlot* ButtonSlot = Cast<UButtonSlot>(PanelSlot);
+			if ( ButtonSlot == NULL )
+			{
+				ButtonSlot = ConstructObject<UButtonSlot>(UButtonSlot::StaticClass(), this);
+				ButtonSlot->Content = GetContentSlot()->Content;
+				ButtonSlot->Content->Slot = ButtonSlot;
+				Slots[0] = ButtonSlot;
+			}
 		}
 	}
 }
diff --git a/Engine/Source/Runtime/UMG/Private/Components/CanvasPanel.cpp b/Engine/Source/Runtime/UMG/Private/Components/CanvasPanel.cpp
--- a/Engine/Source/Runtime/UMG/Private/Components/CanvasPanel.cpp
+++ b/Engine/Source/Runtime/UMG/Private/Components/CanvasPanel.cpp
@@ -2,40 +2,6 @@
 
 #include "UMGPrivatePCH.h"
 
-/////////////////////////////////////////////////////
-// SFixedSizeConstraintCanvas
-
-class SFixedSizeConstraintCanvas : public SConstraintCanvas
-{
-	SLATE_BEGIN_ARGS(SFixedSizeConstraintCanvas)
-		{
-			_Visibility = EVisibility::SelfHitTestInvisible;
-		}
-	SLATE_END_ARGS()
-
-	FVector2D CanvasSize;
-
-	void Construct(const FArguments& InArgs, FVector2D InSize)
-	{
-		SConstraintCanvas::FArguments ParentArgs;
-		SConstraintCanvas::Construct(ParentArgs);
-
-		CanvasSize = InSize;
-	}
-
-	void SetDesiredSize(FVector2D InSize)
-	{
-		CanvasSize = InSize;
-	}
-
-	// SWidget interface
-	virtual FVector2D ComputeDesiredSize() const override
-	{
-		return CanvasSize;
-	}
-	// End of SWidget interface
-};
-
 /////////////////////////////////////////////////////
 // UCanvasPanel
 
@@ -43,7 +9,6 @@ UCanvasPanel::UCanvasPanel(const FPostConstructInitializeProperties& PCIP)
 	: Super(PCIP)
 {
 	bIsVariable = false;
-	DesiredCanvasSize = FVector2D(128.0f, 128.0f);
 
 	SConstraintCanvas::FArguments Defaults;
 	Visiblity = UWidget::ConvertRuntimeToSerializedVisiblity(Defaults._Visibility.Get());
@@ -85,7 +50,7 @@ void UCanvasPanel::OnSlotRemoved(UPanelSlot* Slot)
 
 TSharedRef<SWidget> UCanvasPanel::RebuildWidget()
 {
-	MyCanvas = SNew(SFixedSizeConstraintCanvas, DesiredCanvasSize);
+	MyCanvas = SNew(SConstraintCanvas);
 
 	for ( UPanelSlot* Slot : Slots )
 	{
@@ -102,8 +67,6 @@ TSharedRef<SWidget> UCanvasPanel::RebuildWidget()
 void UCanvasPanel::SyncronizeProperties()
 {
 	Super::SyncronizeProperties();
-
-	MyCanvas->SetDesiredSize(DesiredCanvasSize);
 }
 
 TSharedPtr<SConstraintCanvas> UCanvasPanel::GetCanvasWidget() const
diff --git a/Engine/Source/Runtime/UMG/Private/UserWidget.cpp b/Engine/Source/Runtime/UMG/Private/UserWidget.cpp
--- a/Engine/Source/Runtime/UMG/Private/UserWidget.cpp
+++ b/Engine/Source/Runtime/UMG/Private/UserWidget.cpp
@@ -256,26 +256,26 @@ TSharedRef<SWidget> UUserWidget::RebuildWidget()
 	TSharedPtr<SWidget> UserRootWidget;
 
 	// Add the first component to the root of the widget surface.
-	if ( Components.Num() > 0 )
+	if ( Components.Num() > 0 && Components[0] != NULL )
 	{
 		UserRootWidget = Components[0]->TakeWidget();
+
+		// Place all of our top-level children Slate wrapped components into the overlay
+		for ( int32 ComponentIndex = 0; ComponentIndex < Components.Num(); ++ComponentIndex )
+		{
+			UWidget* Handle = Components[ComponentIndex];
+			TSharedPtr<SWidget> Widget = Handle->GetCachedWidget();
+			if ( Widget.IsValid() )
+			{
+				WidgetToComponent.Add(Widget, Handle);
+			}
+		}
 	}
 	else
 	{
 		UserRootWidget = SNew(SSpacer);
 	}
 
-	// Place all of our top-level children Slate wrapped components into the overlay
-	for ( int32 ComponentIndex = 0; ComponentIndex < Components.Num(); ++ComponentIndex )
-	{
-		UWidget* Handle = Components[ComponentIndex];
-		TSharedPtr<SWidget> Widget = Handle->GetCachedWidget();
-		if ( Widget.IsValid() )
-		{
-			WidgetToComponent.Add(Widget, Handle);
-		}
-	}
-
 	if ( !IsDesignTime() )
 	{
 		// Notify the widget that it has been constructed.
diff --git a/Engine/Source/Runtime/UMG/Public/Blueprint/UserWidget.h b/Engine/Source/Runtime/UMG/Public/Blueprint/UserWidget.h
--- a/Engine/Source/Runtime/UMG/Public/Blueprint/UserWidget.h
+++ b/Engine/Source/Runtime/UMG/Public/Blueprint/UserWidget.h
@@ -319,14 +319,16 @@ public:
 	UPROPERTY(EditAnywhere, Category=Appearance)
 	TEnumAsByte<EVerticalAlignment> VerticalAlignment;
 
-	UPROPERTY()
+	/** The components contained in this user widget. */
+	UPROPERTY(Transient)
 	TArray<UWidget*> Components;
 
-	UPROPERTY()
+	/** The widget tree contained inside this user widget initialized by the blueprint */
+	UPROPERTY(Transient)
 	class UWidgetTree* WidgetTree;
 
 	/** All the sequence players currently playing */
-	UPROPERTY(transient)
+	UPROPERTY(Transient)
 	TArray<UUMGSequencePlayer*> ActiveSequencePlayers;
 
 protected:
diff --git a/Engine/Source/Runtime/UMG/Public/Components/Border.h b/Engine/Source/Runtime/UMG/Public/Components/Border.h
--- a/Engine/Source/Runtime/UMG/Public/Components/Border.h
+++ b/Engine/Source/Runtime/UMG/Public/Components/Border.h
@@ -15,18 +15,6 @@ class UMG_API UBorder : public UContentWidget
 {
 	GENERATED_UCLASS_BODY()
 
-	/** Horizontal positioning of the content within the border */
-	UPROPERTY(EditDefaultsOnly, Category="Content Layout")
-	TEnumAsByte<EHorizontalAlignment> HorizontalAlignment;
-
-	/** Vertical positioning of the content within the border */
-	UPROPERTY(EditDefaultsOnly, Category="Content Layout")
-	TEnumAsByte<EVerticalAlignment> VerticalAlignment;
-
-	/** The padding to add around the inner content. */
-	UPROPERTY(EditDefaultsOnly, Category="Content Layout")
-	FMargin ContentPadding;
-
 	/** Color and opacity multiplier of content in the border */
 	UPROPERTY(EditDefaultsOnly, Category="Content Layout")
 	FLinearColor ContentColorAndOpacity;
@@ -83,14 +71,14 @@ class UMG_API UBorder : public UContentWidget
 	UFUNCTION(BlueprintCallable, Category="Appearance")
 	void SetForegroundColor(FLinearColor InForegroundColor);
 
-	/**  */
-	UFUNCTION(BlueprintCallable, Category="Appearance")
-	void SetContentPadding(FMargin InContentPadding);
-
 	// UWidget interface
 	virtual void SyncronizeProperties() override;
 	// End of UWidget interface
 
+	// Begin UObject interface
+	virtual void PostLoad() override;
+	// End of UObject interface
+
 	virtual void ReleaseNativeWidget() override;
 
 #if WITH_EDITOR
@@ -100,6 +88,7 @@ class UMG_API UBorder : public UContentWidget
 protected:
 
 	// UPanelWidget
+	virtual UClass* GetSlotClass() const override;
 	virtual void OnSlotAdded(UPanelSlot* Slot) override;
 	virtual void OnSlotRemoved(UPanelSlot* Slot) override;
 	// End UPanelWidget
diff --git a/Engine/Source/Runtime/UMG/Public/Components/BorderSlot.h b/Engine/Source/Runtime/UMG/Public/Components/BorderSlot.h
new file mode 100644
--- /dev/null
+++ b/Engine/Source/Runtime/UMG/Public/Components/BorderSlot.h
@@ -0,0 +1,49 @@
+// Copyright 1998-2014 Epic Games, Inc. All Rights Reserved.
+
+#pragma once
+
+#include "SlateWrapperTypes.h"
+
+#include "BorderSlot.generated.h"
+
+/** The Slot for the UBorderSlot, contains the widget displayed in a button's single slot */
+UCLASS()
+class UMG_API UBorderSlot : public UPanelSlot
+{
+	GENERATED_UCLASS_BODY()
+	
+	/** The padding area between the slot and the content it contains. */
+	UPROPERTY(EditDefaultsOnly, Category=Layout)
+	FMargin Padding;
+
+	/** The alignment of the object horizontally. */
+	UPROPERTY(EditDefaultsOnly, Category=Layout)
+	TEnumAsByte<EHorizontalAlignment> HorizontalAlignment;
+
+	/** The alignment of the object vertically. */
+	UPROPERTY(EditDefaultsOnly, Category=Layout)
+	TEnumAsByte<EVerticalAlignment> VerticalAlignment;
+
+	UFUNCTION(BlueprintCallable, Category="Appearance")
+	void SetPadding(FMargin InPadding);
+
+	UFUNCTION(BlueprintCallable, Category="Appearance")
+	void SetHorizontalAlignment(EHorizontalAlignment InHorizontalAlignment);
+
+	UFUNCTION(BlueprintCallable, Category="Appearance")
+	void SetVerticalAlignment(EVerticalAlignment InVerticalAlignment);
+
+	// UPanelSlot interface
+	virtual void SyncronizeProperties() override;
+	// End of UPanelSlot interface
+
+	/** Builds the underlying slot for the slate border. */
+	void BuildSlot(TSharedRef<SBorder> InBorder);
+
+	virtual void ReleaseNativeWidget() override;
+
+private:
+
+	/** A pointer to the button to allow us to adjust the size, padding...etc at runtime. */
+	TSharedPtr<SBorder> Border;
+};
diff --git a/Engine/Source/Runtime/UMG/Public/Components/CanvasPanel.h b/Engine/Source/Runtime/UMG/Public/Components/CanvasPanel.h
--- a/Engine/Source/Runtime/UMG/Public/Components/CanvasPanel.h
+++ b/Engine/Source/Runtime/UMG/Public/Components/CanvasPanel.h
@@ -9,10 +9,6 @@ class UMG_API UCanvasPanel : public UPanelWidget
 {
 	GENERATED_UCLASS_BODY()
 
-	/** The desired size of the canvas */
-	UPROPERTY(EditAnywhere, Category=Appearance)
-	FVector2D DesiredCanvasSize;
-
 	/** Gets the underlying native canvas widget if it has been constructed */
 	TSharedPtr<class SConstraintCanvas> GetCanvasWidget() const;
 
@@ -44,7 +40,7 @@ protected:
 
 protected:
 
-	TSharedPtr<class SFixedSizeConstraintCanvas> MyCanvas;
+	TSharedPtr<class SConstraintCanvas> MyCanvas;
 
 protected:
 	// UWidget interface
diff --git a/Engine/Source/Runtime/UMG/Public/UMG.h b/Engine/Source/Runtime/UMG/Public/UMG.h
--- a/Engine/Source/Runtime/UMG/Public/UMG.h
+++ b/Engine/Source/Runtime/UMG/Public/UMG.h
@@ -54,11 +54,13 @@
 #include "WidgetSwitcherSlot.h"
 #include "WidgetSwitcher.h"
 
+#include "BorderSlot.h"
+#include "Border.h"
+
 #include "ButtonSlot.h"
 #include "Button.h"
 
 #include "TextBlock.h"
-#include "Border.h"
 #include "CheckBox.h"
 #include "Image.h"
 
