UMG - Adding missing design time utility class.

diff --git a/Engine/Source/Editor/UMGEditor/Private/Designer/DesignTimeUtils.cpp b/Engine/Source/Editor/UMGEditor/Private/Designer/DesignTimeUtils.cpp
new file mode 100644
--- /dev/null
+++ b/Engine/Source/Editor/UMGEditor/Private/Designer/DesignTimeUtils.cpp
@@ -0,0 +1,63 @@
+// Copyright 1998-2014 Epic Games, Inc. All Rights Reserved.
+
+#include "UMGEditorPrivatePCH.h"
+
+#define LOCTEXT_NAMESPACE "UMGEditor"
+
+bool FDesignTimeUtils::GetArrangedWidget(TSharedRef<SWidget> Widget, FArrangedWidget& ArrangedWidget)
+{
+	TSharedPtr<SWindow> WidgetWindow = FSlateApplication::Get().FindWidgetWindow(Widget);
+	if ( !WidgetWindow.IsValid() )
+	{
+		return false;
+	}
+
+	TSharedRef<SWindow> CurrentWindowRef = WidgetWindow.ToSharedRef();
+
+	FWidgetPath WidgetPath;
+	if ( FSlateApplication::Get().GeneratePathToWidgetUnchecked(Widget, WidgetPath) )
+	{
+		ArrangedWidget = WidgetPath.FindArrangedWidget(Widget);
+		return true;
+	}
+
+	return false;
+}
+
+bool FDesignTimeUtils::GetArrangedWidgetRelativeToWindow(TSharedRef<SWidget> Widget, FArrangedWidget& ArrangedWidget)
+{
+	TSharedPtr<SWindow> WidgetWindow = FSlateApplication::Get().FindWidgetWindow(Widget);
+	if ( !WidgetWindow.IsValid() )
+	{
+		return false;
+	}
+
+	TSharedRef<SWindow> CurrentWindowRef = WidgetWindow.ToSharedRef();
+
+	FWidgetPath WidgetPath;
+	if ( FSlateApplication::Get().GeneratePathToWidgetUnchecked(Widget, WidgetPath) )
+	{
+		ArrangedWidget = WidgetPath.FindArrangedWidget(Widget);
+		ArrangedWidget.Geometry.AbsolutePosition -= CurrentWindowRef->GetPositionInScreen();
+		return true;
+	}
+
+	return false;
+}
+
+bool FDesignTimeUtils::GetArrangedWidgetRelativeToParent(TSharedRef<const SWidget> Widget, TSharedRef<const SWidget> Parent, FArrangedWidget& ArrangedWidget)
+{
+	FWidgetPath WidgetPath;
+	if ( FSlateApplication::Get().GeneratePathToWidgetUnchecked(Widget, WidgetPath) )
+	{
+		FArrangedWidget ArrangedDesigner = WidgetPath.FindArrangedWidget(Parent);
+
+		ArrangedWidget = WidgetPath.FindArrangedWidget(Widget);
+		ArrangedWidget.Geometry.AbsolutePosition -= ArrangedDesigner.Geometry.AbsolutePosition;
+		return true;
+	}
+
+	return false;
+}
+
+#undef LOCTEXT_NAMESPACE
diff --git a/Engine/Source/Editor/UMGEditor/Private/Designer/DesignTimeUtils.h b/Engine/Source/Editor/UMGEditor/Private/Designer/DesignTimeUtils.h
new file mode 100644
--- /dev/null
+++ b/Engine/Source/Editor/UMGEditor/Private/Designer/DesignTimeUtils.h
@@ -0,0 +1,16 @@
+// Copyright 1998-2014 Epic Games, Inc. All Rights Reserved.
+
+#pragma once
+
+/**
+ * A set of utility functions used at design time for the widget blueprint editor.
+ */
+class FDesignTimeUtils
+{
+public:
+	static bool GetArrangedWidget(TSharedRef<SWidget> Widget, FArrangedWidget& ArrangedWidget);
+
+	static bool GetArrangedWidgetRelativeToWindow(TSharedRef<SWidget> Widget, FArrangedWidget& ArrangedWidget);
+
+	static bool GetArrangedWidgetRelativeToParent(TSharedRef<const SWidget> Widget, TSharedRef<const SWidget> Parent, FArrangedWidget& ArrangedWidget);
+};
