UMG - Fixing a build error with SZoomPan, not respecting the new protection level of Widget in the slot.

diff --git a/Engine/Source/Editor/UMGEditor/Private/Designer/SZoomPan.cpp b/Engine/Source/Editor/UMGEditor/Private/Designer/SZoomPan.cpp
--- a/Engine/Source/Editor/UMGEditor/Private/Designer/SZoomPan.cpp
+++ b/Engine/Source/Editor/UMGEditor/Private/Designer/SZoomPan.cpp
@@ -32,7 +32,7 @@ void SZoomPan::OnArrangeChildren(const FGeometry& AllottedGeometry, FArrangedChi
 		ArrangedChildren.AddWidget( ChildVisibility, AllottedGeometry.MakeChild(
 				ChildSlot.GetWidget(),
 				FVector2D(XResult.Offset, YResult.Offset) - ViewOffset.Get(),
-				ChildSlot.Widget->GetDesiredSize(),
+				ChildSlot.GetWidget()->GetDesiredSize(),
 				ZoomAmount.Get()
 		) );
 	}
