UMG - We now only update selection in the designer if the pick hit a widget.

diff --git a/Engine/Source/Editor/UMGEditor/Private/Designer/SDesignerView.cpp b/Engine/Source/Editor/UMGEditor/Private/Designer/SDesignerView.cpp
--- a/Engine/Source/Editor/UMGEditor/Private/Designer/SDesignerView.cpp
+++ b/Engine/Source/Editor/UMGEditor/Private/Designer/SDesignerView.cpp
@@ -450,15 +450,15 @@ FReply SDesignerView::OnMouseButtonDown(const FGeometry& MyGeometry, const FPoin
 
 	//TODO UMG Undoable Selection
 	FArrangedWidget ArrangedWidget(SNullWidget::NullWidget, FGeometry());
-	SelectedWidget = GetWidgetAtCursor(MyGeometry, MouseEvent, ArrangedWidget);
+	FWidgetReference NewSelectedWidget = GetWidgetAtCursor(MyGeometry, MouseEvent, ArrangedWidget);
 	SelectedWidgetContextMenuLocation = ArrangedWidget.Geometry.AbsoluteToLocal(MouseEvent.GetScreenSpacePosition());
 
-	if ( SelectedWidget.IsValid() )
+	if ( NewSelectedWidget.IsValid() )
 	{
 		//@TODO UMG primary FBlueprintEditor needs to be inherited and selection control needs to be centralized.
 		// Set the template as selected in the details panel
 		TSet<FWidgetReference> SelectedTemplates;
-		SelectedTemplates.Add(SelectedWidget);
+		SelectedTemplates.Add(NewSelectedWidget);
 		BlueprintEditor.Pin()->SelectWidgets(SelectedTemplates);
 
 		if ( MouseEvent.GetEffectingButton() == EKeys::LeftMouseButton )
