UMG - Root widgets should now be wrappable.

diff --git a/Engine/Source/Editor/UMGEditor/Private/WidgetBlueprintEditorUtils.cpp b/Engine/Source/Editor/UMGEditor/Private/WidgetBlueprintEditorUtils.cpp
--- a/Engine/Source/Editor/UMGEditor/Private/WidgetBlueprintEditorUtils.cpp
+++ b/Engine/Source/Editor/UMGEditor/Private/WidgetBlueprintEditorUtils.cpp
@@ -255,10 +255,17 @@ void FWidgetBlueprintEditorUtils::WrapWidgets(UWidgetBlueprint* BP, TSet<FWidget
 		UPanelWidget* CurrentParent = BP->WidgetTree->FindWidgetParent(Item.GetTemplate(), OutIndex);
 		if ( CurrentParent )
 		{
+			CurrentParent->Modify();
 			CurrentParent->ReplaceChildAt(OutIndex, NewWrapperWidget);
+		}
+		else
+		{
+			BP->WidgetTree->Modify();
 
-			NewWrapperWidget->AddChild(Item.GetTemplate());
+			BP->WidgetTree->RootWidget = NewWrapperWidget;
 		}
+
+		NewWrapperWidget->AddChild(Item.GetTemplate());
 	}
 
 	FBlueprintEditorUtils::MarkBlueprintAsStructurallyModified(BP);
