UMG - Adding additional checks to prevent the preview widget returned by the blueprint editor from being the bogus one after a compile.

diff --git a/Engine/Source/Editor/UMGEditor/Private/WidgetBlueprintEditor.cpp b/Engine/Source/Editor/UMGEditor/Private/WidgetBlueprintEditor.cpp
--- a/Engine/Source/Editor/UMGEditor/Private/WidgetBlueprintEditor.cpp
+++ b/Engine/Source/Editor/UMGEditor/Private/WidgetBlueprintEditor.cpp
@@ -290,7 +290,7 @@ void FWidgetBlueprintEditor::Tick(float DeltaTime)
 
 	// Note: The weak ptr can become stale if the actor is reinstanced due to a Blueprint change, etc. In that case we 
 	//       look to see if we can find the new instance in the preview world and then update the weak ptr.
-	if ( PreviewWidgetActorPtr.IsStale(true) )
+	if ( PreviewWidgetPtr.IsStale(true) )
 	{
 		UpdatePreview(GetWidgetBlueprintObj(), true);
 	}
@@ -405,7 +405,12 @@ UWidgetBlueprint* FWidgetBlueprintEditor::GetWidgetBlueprintObj() const
 
 UUserWidget* FWidgetBlueprintEditor::GetPreview() const
 {
-	return PreviewWidgetActorPtr.Get();
+	if ( PreviewWidgetPtr.IsStale(true) )
+	{
+		return NULL;
+	}
+
+	return PreviewWidgetPtr.Get();
 }
 
 TSharedPtr<ISequencer>& FWidgetBlueprintEditor::GetSequencer()
@@ -505,7 +510,7 @@ void FWidgetBlueprintEditor::UpdatePreview(UBlueprint* InBlueprint, bool bInForc
 		}
 
 		// Store a reference to the preview actor.
-		PreviewWidgetActorPtr = PreviewActor;
+		PreviewWidgetPtr = PreviewActor;
 	}
 
 	OnWidgetPreviewUpdated.Broadcast();
diff --git a/Engine/Source/Editor/UMGEditor/Private/WidgetBlueprintEditor.h b/Engine/Source/Editor/UMGEditor/Private/WidgetBlueprintEditor.h
--- a/Engine/Source/Editor/UMGEditor/Private/WidgetBlueprintEditor.h
+++ b/Engine/Source/Editor/UMGEditor/Private/WidgetBlueprintEditor.h
@@ -122,7 +122,7 @@ private:
 	TSet<FWidgetReference> SelectedWidgets;
 
 	/** The preview GUI object */
-	mutable TWeakObjectPtr<UUserWidget> PreviewWidgetActorPtr;
+	mutable TWeakObjectPtr<UUserWidget> PreviewWidgetPtr;
 
 	/** Notification for when the preview widget has been updated */
 	FOnWidgetPreviewUpdated OnWidgetPreviewUpdated;
