UMG - The designer now does zoom-to-fit on load.  There is now a button to perform a zoom-to-fit at any time now.

diff --git a/Engine/Source/Editor/UMGEditor/Private/Designer/SDesignSurface.cpp b/Engine/Source/Editor/UMGEditor/Private/Designer/SDesignSurface.cpp
--- a/Engine/Source/Editor/UMGEditor/Private/Designer/SDesignSurface.cpp
+++ b/Engine/Source/Editor/UMGEditor/Private/Designer/SDesignSurface.cpp
@@ -123,12 +123,58 @@ void SDesignSurface::Construct(const FArguments& InArgs)
 	ZoomLevelGraphFade = FCurveSequence(0.0f, 0.5f);
 	ZoomLevelGraphFade.Play();
 
+	bDeferredZoomingToFit = false;
+	bDeferredZoomToExtents = false;
+
+	bAllowContinousZoomInterpolation = false;
+	bTeleportInsteadOfScrollingWhenZoomingToFit = false;
+
+	ZoomTargetTopLeft = FVector2D::ZeroVector;
+	ZoomTargetBottomRight = FVector2D::ZeroVector;
+
 	ChildSlot
 	[
 		InArgs._Content.Widget
 	];
 }
 
+void SDesignSurface::Tick(const FGeometry& AllottedGeometry, const double InCurrentTime, const float InDeltaTime)
+{
+	FSlateRect Bounds = ComputeAreaBounds();
+
+	// Zoom to extents
+	if ( bDeferredZoomToExtents )
+	{
+		bDeferredZoomToExtents = false;
+		//ZoomPadding = NodePanelDefs::DefaultZoomPadding;
+		ZoomTargetTopLeft = FVector2D(Bounds.Left, Bounds.Top);
+		ZoomTargetBottomRight = FVector2D(Bounds.Right, Bounds.Bottom);
+		//bDeferredZoomToSelection = false;
+		bDeferredZoomingToFit = true;
+	}
+
+	if ( bDeferredZoomingToFit )
+	{
+		const FVector2D DesiredViewCenter = ( ZoomTargetTopLeft + ZoomTargetBottomRight ) * 0.5f;
+		const bool bDoneScrolling = ScrollToLocation(AllottedGeometry, DesiredViewCenter, bTeleportInsteadOfScrollingWhenZoomingToFit ? 1000.0f : InDeltaTime);
+		bool bDoneZooming = ZoomToLocation(AllottedGeometry.Size, ZoomTargetBottomRight - ZoomTargetTopLeft, bDoneScrolling);
+
+		if ( bDoneZooming && bDoneScrolling )
+		{
+			// One final push to make sure we centered in the end
+			ViewOffset = DesiredViewCenter - ( 0.5f * AllottedGeometry.Scale * AllottedGeometry.Size / GetZoomAmount() );
+			// Reset ZoomPadding
+			//ZoomPadding = NodePanelDefs::DefaultZoomPadding;
+			ZoomTargetTopLeft = FVector2D::ZeroVector;
+			ZoomTargetBottomRight = FVector2D::ZeroVector;
+
+			bDeferredZoomingToFit = false;
+		}
+	}
+	
+	SCompoundWidget::Tick(AllottedGeometry, InCurrentTime, InDeltaTime);
+}
+
 int32 SDesignSurface::OnPaint(const FPaintArgs& Args, const FGeometry& AllottedGeometry, const FSlateRect& MyClippingRect, FSlateWindowElementList& OutDrawElements, int32 LayerId, const FWidgetStyle& InWidgetStyle, bool bParentEnabled) const
 {
 	const FSlateBrush* BackgroundImage = FEditorStyle::GetBrush(TEXT("Graph.Panel.SolidBackground"));
@@ -246,7 +292,7 @@ void SDesignSurface::ChangeZoomLevel(int32 ZoomLevelDelta, const FVector2D& Widg
 
 		// Re-center the screen so that it feels like zooming around the cursor.
 		{
-			FSlateRect GraphBounds = ComputeSensibleGraphBounds();
+			FSlateRect GraphBounds = ComputeSensibleBounds();
 
 			// Make sure we are not zooming into/out into emptiness; otherwise the user will get lost..
 			const FVector2D ClampedPointToMaintainGraphSpace(
@@ -259,34 +305,107 @@ void SDesignSurface::ChangeZoomLevel(int32 ZoomLevelDelta, const FVector2D& Widg
 	}
 }
 
-FSlateRect SDesignSurface::ComputeSensibleGraphBounds() const
+FSlateRect SDesignSurface::ComputeSensibleBounds() const
 {
-	float Left = 0.0f;
-	float Top = 0.0f;
-	float Right = 1920.f;
-	float Bottom = 1080.0f;
+	// Pad it out in every direction, to roughly account for nodes being of non-zero extent
+	const float Padding = 100.0f;
 
-	// Find the bounds of the node positions
-	//for ( int32 ChildIndex = 0; ChildIndex < Children.Num(); ++ChildIndex )
-	//{
-	//	const TSharedRef<SNode>& SomeChild = Children[ChildIndex];
+	FSlateRect Bounds = ComputeAreaBounds();
+	Bounds.Left -= Padding;
+	Bounds.Top -= Padding;
+	Bounds.Right -= Padding;
+	Bounds.Bottom -= Padding;
 
-	//	FVector2D ChildPos = SomeChild->GetPosition();
+	return Bounds;
+}
 
-	//	Left = FMath::Min(Left, ChildPos.X);
-	//	Right = FMath::Max(Right, ChildPos.X);
-	//	Top = FMath::Min(Top, ChildPos.Y);
-	//	Bottom = FMath::Max(Bottom, ChildPos.Y);
-	//}
+void SDesignSurface::PostChangedZoom()
+{
+}
 
-	// Pad it out in every direction, to roughly account for nodes being of non-zero extent
-	const float Padding = 100.0f;
+bool SDesignSurface::ScrollToLocation(const FGeometry& MyGeometry, FVector2D DesiredCenterPosition, const float InDeltaTime)
+{
+	const FVector2D HalfOFScreenInGraphSpace = 0.5f * MyGeometry.Size / GetZoomAmount();
+	FVector2D CurrentPosition = ViewOffset + HalfOFScreenInGraphSpace;
 
-	return FSlateRect(Left - Padding, Top - Padding, Right + Padding, Bottom + Padding);
+	FVector2D NewPosition = FMath::Vector2DInterpTo(CurrentPosition, DesiredCenterPosition, InDeltaTime, 10.f);
+	ViewOffset = NewPosition - HalfOFScreenInGraphSpace;
+
+	// If within 1 pixel of target, stop interpolating
+	return ( ( NewPosition - DesiredCenterPosition ).Size() < 1.f );
 }
 
-void SDesignSurface::PostChangedZoom()
+bool SDesignSurface::ZoomToLocation(const FVector2D& CurrentSizeWithoutZoom, const FVector2D& DesiredSize, bool bDoneScrolling)
+{
+	if ( bAllowContinousZoomInterpolation && ZoomLevelGraphFade.IsPlaying() )
+	{
+		return false;
+	}
+
+	const int32 DefaultZoomLevel = ZoomLevels->GetDefaultZoomLevel();
+	const int32 NumZoomLevels = ZoomLevels->GetNumZoomLevels();
+	int32 DesiredZoom = DefaultZoomLevel;
+
+	// Find lowest zoom level that will display all nodes
+	for ( int32 Zoom = 0; Zoom < DefaultZoomLevel; ++Zoom )
+	{
+		const FVector2D SizeWithZoom = CurrentSizeWithoutZoom / ZoomLevels->GetZoomAmount(Zoom);
+		const FVector2D LeftOverSize = SizeWithZoom - DesiredSize;
+
+		if ( ( DesiredSize.X > SizeWithZoom.X ) || ( DesiredSize.Y > SizeWithZoom.Y ) )
+		{
+			// Use the previous zoom level, this one is too tight
+			DesiredZoom = FMath::Max<int32>(0, Zoom - 1);
+			break;
+		}
+	}
+
+	if ( DesiredZoom != ZoomLevel )
+	{
+		if ( bAllowContinousZoomInterpolation )
+		{
+			// Animate to it
+			PreviousZoomLevel = ZoomLevel;
+			ZoomLevel = FMath::Clamp(DesiredZoom, 0, NumZoomLevels - 1);
+			ZoomLevelGraphFade.Play();
+			return false;
+		}
+		else
+		{
+			// Do it instantly, either first or last
+			if ( DesiredZoom < ZoomLevel )
+			{
+				// Zooming out; do it instantly
+				ZoomLevel = PreviousZoomLevel = DesiredZoom;
+				ZoomLevelFade.Play();
+			}
+			else
+			{
+				// Zooming in; do it last
+				if ( bDoneScrolling )
+				{
+					ZoomLevel = PreviousZoomLevel = DesiredZoom;
+					ZoomLevelFade.Play();
+				}
+			}
+		}
+
+		PostChangedZoom();
+	}
+
+	return true;
+}
+
+void SDesignSurface::ZoomToFit(bool bInstantZoom)
+{
+	bDeferredZoomingToFit = true;
+	bDeferredZoomToExtents = true;
+	bTeleportInsteadOfScrollingWhenZoomingToFit = bInstantZoom;
+}
+
+FSlateRect SDesignSurface::ComputeAreaBounds() const
 {
+	return FSlateRect(0, 0, 0, 0);
 }
 
 FVector2D SDesignSurface::GetViewOffset() const
diff --git a/Engine/Source/Editor/UMGEditor/Private/Designer/SDesignSurface.h b/Engine/Source/Editor/UMGEditor/Private/Designer/SDesignSurface.h
--- a/Engine/Source/Editor/UMGEditor/Private/Designer/SDesignSurface.h
+++ b/Engine/Source/Editor/UMGEditor/Private/Designer/SDesignSurface.h
@@ -22,6 +22,7 @@ public:
 	void Construct(const FArguments& InArgs);
 
 	// SWidget interface
+	virtual void Tick(const FGeometry& AllottedGeometry, const double InCurrentTime, const float InDeltaTime);
 	virtual FReply OnMouseButtonDown(const FGeometry& MyGeometry, const FPointerEvent& MouseEvent) override;
 	virtual FReply OnMouseButtonUp(const FGeometry& MyGeometry, const FPointerEvent& MouseEvent) override;
 	virtual FReply OnMouseMove(const FGeometry& MyGeometry, const FPointerEvent& MouseEvent) override;
@@ -39,12 +40,20 @@ protected:
 	float GetZoomAmount() const;
 
 	void ChangeZoomLevel(int32 ZoomLevelDelta, const FVector2D& WidgetSpaceZoomOrigin, bool bOverrideZoomLimiting);
-
+	
 	void PostChangedZoom();
 
+	bool ScrollToLocation(const FGeometry& MyGeometry, FVector2D DesiredCenterPosition, const float InDeltaTime);
+
+	bool ZoomToLocation(const FVector2D& CurrentSizeWithoutZoom, const FVector2D& DesiredSize, bool bDoneScrolling);
+
+	void ZoomToFit(bool bInstantZoom);
+
 	FVector2D GetViewOffset() const;
 
-	FSlateRect ComputeSensibleGraphBounds() const;
+	virtual FSlateRect ComputeAreaBounds() const;
+
+	FSlateRect ComputeSensibleBounds() const;
 
 	FVector2D GraphCoordToPanelCoord(const FVector2D& GraphSpaceCoordinate) const;
 	FVector2D PanelCoordToGraphCoord(const FVector2D& PanelSpaceCoordinate) const;
@@ -75,4 +84,17 @@ protected:
 
 	// The interface for mapping ZoomLevel values to actual node scaling values
 	TScopedPointer<FZoomLevelsContainer> ZoomLevels;
+
+	// A flag noting if we have a pending zoom to fit operation to perform next tick.
+	bool bDeferredZoomingToFit;
+
+	// A flag noting if we have a pending zoom to extents operation to perform next tick.
+	bool bDeferredZoomToExtents;
+
+	bool bAllowContinousZoomInterpolation;
+
+	bool bTeleportInsteadOfScrollingWhenZoomingToFit;
+
+	FVector2D ZoomTargetTopLeft;
+	FVector2D ZoomTargetBottomRight;
 };
diff --git a/Engine/Source/Editor/UMGEditor/Private/Designer/SDesignerView.cpp b/Engine/Source/Editor/UMGEditor/Private/Designer/SDesignerView.cpp
--- a/Engine/Source/Editor/UMGEditor/Private/Designer/SDesignerView.cpp
+++ b/Engine/Source/Editor/UMGEditor/Private/Designer/SDesignerView.cpp
@@ -173,6 +173,19 @@ void SDesignerView::Construct(const FArguments& InArgs, TSharedPtr<FWidgetBluepr
 			.VAlign(VAlign_Top)
 			[
 				SNew(SHorizontalBox)
+
+				+ SHorizontalBox::Slot()
+				.AutoWidth()
+				[
+					SNew(SButton)
+					.ToolTipText(LOCTEXT("ZoomToFit_ToolTip", "ZoomToFit"))
+					.OnClicked(this, &SDesignerView::HandleZoomToFitClicked)
+					[
+						//TODO UMG Replace with image.
+						SNew(STextBlock)
+						.Text(LOCTEXT("ZoomToFit", "Zoom To Fit"))
+					]
+				]
 				
 				+ SHorizontalBox::Slot()
 				.AutoWidth()
@@ -205,6 +218,8 @@ void SDesignerView::Construct(const FArguments& InArgs, TSharedPtr<FWidgetBluepr
 	);
 
 	BlueprintEditor.Pin()->OnSelectedWidgetsChanged.AddRaw(this, &SDesignerView::OnEditorSelectionChanged);
+
+	ZoomToFit(/*bInstantZoom*/ true);
 }
 
 SDesignerView::~SDesignerView()
@@ -231,6 +246,11 @@ FOptionalSize SDesignerView::GetPreviewHeight() const
 	return (float)PreviewHeight;
 }
 
+FSlateRect SDesignerView::ComputeAreaBounds() const
+{
+	return FSlateRect(0, 0, PreviewWidth, PreviewHeight);
+}
+
 void SDesignerView::OnEditorSelectionChanged()
 {
 	TSet<FWidgetReference> PendingSelectedWidgets = BlueprintEditor.Pin()->GetSelectedWidgets();
@@ -976,7 +996,7 @@ void SDesignerView::Tick(const FGeometry& AllottedGeometry, const double InCurre
 		Ext->Tick(AllottedGeometry, InCurrentTime, InDeltaTime);
 	}
 
-	SCompoundWidget::Tick(AllottedGeometry, InCurrentTime, InDeltaTime);
+	SDesignSurface::Tick(AllottedGeometry, InCurrentTime, InDeltaTime);
 }
 
 FReply SDesignerView::OnDragDetected(const FGeometry& MyGeometry, const FPointerEvent& MouseEvent)
@@ -1338,4 +1358,10 @@ void SDesignerView::EndTransaction()
 	}
 }
 
+FReply SDesignerView::HandleZoomToFitClicked()
+{
+	ZoomToFit(/*bInstantZoom*/ false);
+	return FReply::Handled();
+}
+
 #undef LOCTEXT_NAMESPACE
diff --git a/Engine/Source/Editor/UMGEditor/Private/Designer/SDesignerView.h b/Engine/Source/Editor/UMGEditor/Private/Designer/SDesignerView.h
--- a/Engine/Source/Editor/UMGEditor/Private/Designer/SDesignerView.h
+++ b/Engine/Source/Editor/UMGEditor/Private/Designer/SDesignerView.h
@@ -53,6 +53,8 @@ private:
 
 	FOptionalSize GetPreviewHeight() const;
 
+	virtual FSlateRect ComputeAreaBounds() const override;
+
 	/** Updates the designer to display the latest preview widget */
 	void UpdatePreviewWidget();
 
@@ -83,6 +85,9 @@ private:
 	void BeginTransaction(const FText& SessionName);
 	void EndTransaction();
 
+private:
+	FReply HandleZoomToFitClicked();
+
 private:
 	enum DragHandle
 	{
