UMG - Adding a placeholder zoom to fit icon.

diff --git a/Engine/Source/Editor/EditorStyle/Private/SlateEditorStyle.cpp b/Engine/Source/Editor/EditorStyle/Private/SlateEditorStyle.cpp
--- a/Engine/Source/Editor/EditorStyle/Private/SlateEditorStyle.cpp
+++ b/Engine/Source/Editor/EditorStyle/Private/SlateEditorStyle.cpp
@@ -5414,6 +5414,8 @@ void FSlateEditorStyle::FStyle::SetupAutomationStyles()
 
 void FSlateEditorStyle::FStyle::SetupUMGEditorStyles()
 {
+	Set("UMGEditor.ZoomToFit", new IMAGE_BRUSH("Icons/UMG/Fit_16x", Icon16x16));
+
 	Set("UMGEditor.AnchorCenter", new IMAGE_BRUSH("Icons/umg_anchor_center", Icon16x16));
 	Set("UMGEditor.AnchorTopBottom", new IMAGE_BRUSH("Icons/umg_anchor_top_bottom", FVector2D(16, 32)));
 	Set("UMGEditor.AnchorLeftRight", new IMAGE_BRUSH("Icons/umg_anchor_left_right", FVector2D(32, 16)));
diff --git a/Engine/Source/Editor/UMGEditor/Private/Designer/SDesignerView.cpp b/Engine/Source/Editor/UMGEditor/Private/Designer/SDesignerView.cpp
--- a/Engine/Source/Editor/UMGEditor/Private/Designer/SDesignerView.cpp
+++ b/Engine/Source/Editor/UMGEditor/Private/Designer/SDesignerView.cpp
@@ -178,12 +178,12 @@ void SDesignerView::Construct(const FArguments& InArgs, TSharedPtr<FWidgetBluepr
 				.AutoWidth()
 				[
 					SNew(SButton)
-					.ToolTipText(LOCTEXT("ZoomToFit_ToolTip", "ZoomToFit"))
+					.ButtonStyle(FEditorStyle::Get(), "ToggleButton")
+					.ToolTipText(LOCTEXT("ZoomToFit_ToolTip", "Zoom To Fit"))
 					.OnClicked(this, &SDesignerView::HandleZoomToFitClicked)
 					[
-						//TODO UMG Replace with image.
-						SNew(STextBlock)
-						.Text(LOCTEXT("ZoomToFit", "Zoom To Fit"))
+						SNew(SImage)
+						.Image(FEditorStyle::GetBrush("UMGEditor.ZoomToFit"))
 					]
 				]
 				
diff --git a/Engine/Source/Editor/UMGEditor/Private/Templates/WidgetTemplateClass.cpp b/Engine/Source/Editor/UMGEditor/Private/Templates/WidgetTemplateClass.cpp
--- a/Engine/Source/Editor/UMGEditor/Private/Templates/WidgetTemplateClass.cpp
+++ b/Engine/Source/Editor/UMGEditor/Private/Templates/WidgetTemplateClass.cpp
@@ -50,7 +50,7 @@ const FSlateBrush* FWidgetTemplateClass::GetIcon() const
 
 TSharedRef<IToolTip> FWidgetTemplateClass::GetToolTip() const
 {
-	return IDocumentation::Get()->CreateToolTip(WidgetClass->GetDisplayNameText(), nullptr, FString(TEXT("Shared/Types/")) + WidgetClass->GetName(), TEXT("Class"));
+	return IDocumentation::Get()->CreateToolTip(FText::FromString(WidgetClass->GetDescription()), nullptr, FString(TEXT("Shared/Types/")) + WidgetClass->GetName(), TEXT("Class"));
 }
 
 #undef LOCTEXT_NAMESPACE
