UMG - Fix incorrect usage of Slate font style

diff --git a/Engine/Source/Editor/UMGEditor/Private/Hierarchy/SHierarchyViewItem.cpp b/Engine/Source/Editor/UMGEditor/Private/Hierarchy/SHierarchyViewItem.cpp
--- a/Engine/Source/Editor/UMGEditor/Private/Hierarchy/SHierarchyViewItem.cpp
+++ b/Engine/Source/Editor/UMGEditor/Private/Hierarchy/SHierarchyViewItem.cpp
@@ -171,8 +171,8 @@ FSlateFontInfo SHierarchyViewItem::GetItemFont() const
 	}
 	else
 	{
-		// TODO UMG Hacky move into style area
-		return FSlateFontInfo(FPaths::EngineContentDir() / TEXT("Slate/Fonts/Roboto-Regular.ttf"), 9);
+		static FName NormalFont("NormalFont");
+		return FCoreStyle::Get().GetFontStyle(NormalFont);
 	}
 }
 
