UMG - Hiding the viewport related properties, now using functions to drive them.  Exposing simple anchors.  Adding the Grid widget.

diff --git a/Engine/Source/Editor/UMGEditor/Private/Designer/SDesignerView.cpp b/Engine/Source/Editor/UMGEditor/Private/Designer/SDesignerView.cpp
--- a/Engine/Source/Editor/UMGEditor/Private/Designer/SDesignerView.cpp
+++ b/Engine/Source/Editor/UMGEditor/Private/Designer/SDesignerView.cpp
@@ -129,6 +129,7 @@ void SDesignerView::Construct(const FArguments& InArgs, TSharedPtr<FWidgetBluepr
 	Register(MakeShareable(new FHorizontalSlotExtension()));
 	Register(MakeShareable(new FCanvasSlotExtension()));
 	Register(MakeShareable(new FUniformGridSlotExtension()));
+	Register(MakeShareable(new FGridSlotExtension()));
 
 	UWidgetBlueprint* Blueprint = GetBlueprint();
 	Blueprint->OnChanged().AddSP(this, &SDesignerView::OnBlueprintChanged);
diff --git a/Engine/Source/Editor/UMGEditor/Private/Extensions/GridSlotExtension.cpp b/Engine/Source/Editor/UMGEditor/Private/Extensions/GridSlotExtension.cpp
new file mode 100644
--- /dev/null
+++ b/Engine/Source/Editor/UMGEditor/Private/Extensions/GridSlotExtension.cpp
@@ -0,0 +1,106 @@
+﻿// Copyright 1998-2014 Epic Games, Inc. All Rights Reserved.
+
+#include "UMGEditorPrivatePCH.h"
+
+#define LOCTEXT_NAMESPACE "UMG"
+
+FGridSlotExtension::FGridSlotExtension()
+{
+	ExtensionId = FName(TEXT("GridSlot"));
+}
+
+bool FGridSlotExtension::CanExtendSelection(const TArray< FWidgetReference >& Selection) const
+{
+	for ( const FWidgetReference& Widget : Selection )
+	{
+		if ( !Widget.GetTemplate()->Slot || !Widget.GetTemplate()->Slot->IsA(UGridSlot::StaticClass()) )
+		{
+			return false;
+		}
+	}
+
+	return Selection.Num() == 1;
+}
+
+void FGridSlotExtension::ExtendSelection(const TArray< FWidgetReference >& Selection, TArray< TSharedRef<FDesignerSurfaceElement> >& SurfaceElements)
+{
+	SelectionCache = Selection;
+
+	TSharedRef<SButton> UpArrow = SNew(SButton)
+		.Text(LOCTEXT("UpArrow", "↑"))
+		.ContentPadding(FMargin(6, 2))
+		.OnClicked(this, &FGridSlotExtension::HandleShiftRow, -1);
+
+	TSharedRef<SButton> DownArrow = SNew(SButton)
+		.Text(LOCTEXT("DownArrow", "↓"))
+		.ContentPadding(FMargin(6, 2))
+		.OnClicked(this, &FGridSlotExtension::HandleShiftRow, 1);
+
+	TSharedRef<SButton> LeftArrow = SNew(SButton)
+		.Text(LOCTEXT("LeftArrow", "←"))
+		.ContentPadding(FMargin(2, 6))
+		.OnClicked(this, &FGridSlotExtension::HandleShiftColumn, -1);
+
+	TSharedRef<SButton> RightArrow = SNew(SButton)
+		.Text(LOCTEXT("RightArrow", "→"))
+		.ContentPadding(FMargin(2, 6))
+		.OnClicked(this, &FGridSlotExtension::HandleShiftColumn, 1);
+
+	UpArrow->SlatePrepass();
+	DownArrow->SlatePrepass();
+	LeftArrow->SlatePrepass();
+	RightArrow->SlatePrepass();
+
+	SurfaceElements.Add(MakeShareable(new FDesignerSurfaceElement(LeftArrow, EExtensionLayoutLocation::CenterLeft, FVector2D(-LeftArrow->GetDesiredSize().X, LeftArrow->GetDesiredSize().Y * -0.5f))));
+	SurfaceElements.Add(MakeShareable(new FDesignerSurfaceElement(RightArrow, EExtensionLayoutLocation::CenterRight, FVector2D(0, RightArrow->GetDesiredSize().Y * -0.5f))));
+	SurfaceElements.Add(MakeShareable(new FDesignerSurfaceElement(UpArrow, EExtensionLayoutLocation::TopCenter, FVector2D(UpArrow->GetDesiredSize().X * -0.5f, -UpArrow->GetDesiredSize().Y))));
+	SurfaceElements.Add(MakeShareable(new FDesignerSurfaceElement(DownArrow, EExtensionLayoutLocation::BottomCenter, FVector2D(DownArrow->GetDesiredSize().X * -0.5f, 0))));
+}
+
+FReply FGridSlotExtension::HandleShiftRow(int32 ShiftAmount)
+{
+	BeginTransaction(LOCTEXT("MoveWidget", "Move Widget"));
+
+	for ( FWidgetReference& Selection : SelectionCache )
+	{
+		ShiftRow(Selection.GetPreview(), ShiftAmount);
+		ShiftRow(Selection.GetTemplate(), ShiftAmount);
+	}
+
+	EndTransaction();
+
+	FBlueprintEditorUtils::MarkBlueprintAsModified(Blueprint);
+
+	return FReply::Handled();
+}
+
+FReply FGridSlotExtension::HandleShiftColumn(int32 ShiftAmount)
+{
+	BeginTransaction(LOCTEXT("MoveWidget", "Move Widget"));
+
+	for ( FWidgetReference& Selection : SelectionCache )
+	{
+		ShiftColumn(Selection.GetPreview(), ShiftAmount);
+		ShiftColumn(Selection.GetTemplate(), ShiftAmount);
+	}
+
+	EndTransaction();
+
+	FBlueprintEditorUtils::MarkBlueprintAsModified(Blueprint);
+
+	return FReply::Handled();
+}
+
+void FGridSlotExtension::ShiftRow(UWidget* Widget, int32 ShiftAmount)
+{
+	UGridSlot* Slot = Cast<UGridSlot>(Widget->Slot);
+	Slot->SetRow(FMath::Max(Slot->Row + ShiftAmount, 0));
+}
+
+void FGridSlotExtension::ShiftColumn(UWidget* Widget, int32 ShiftAmount)
+{
+	UGridSlot* Slot = Cast<UGridSlot>(Widget->Slot);
+	Slot->SetColumn(FMath::Max(Slot->Column + ShiftAmount, 0));
+}
+
+#undef LOCTEXT_NAMESPACE
diff --git a/Engine/Source/Editor/UMGEditor/Private/Extensions/GridSlotExtension.h b/Engine/Source/Editor/UMGEditor/Private/Extensions/GridSlotExtension.h
new file mode 100644
--- /dev/null
+++ b/Engine/Source/Editor/UMGEditor/Private/Extensions/GridSlotExtension.h
@@ -0,0 +1,25 @@
+// Copyright 1998-2014 Epic Games, Inc. All Rights Reserved.
+
+#pragma once
+
+#include "DesignerExtension.h"
+
+class FGridSlotExtension : public FDesignerExtension
+{
+public:
+	FGridSlotExtension();
+
+	virtual ~FGridSlotExtension() {}
+
+	virtual bool CanExtendSelection(const TArray< FWidgetReference >& Selection) const override;
+	
+	virtual void ExtendSelection(const TArray< FWidgetReference >& Selection, TArray< TSharedRef<FDesignerSurfaceElement> >& SurfaceElements) override;
+
+private:
+
+	FReply HandleShiftRow(int32 ShiftAmount);
+	FReply HandleShiftColumn(int32 ShiftAmount);
+
+	void ShiftRow(UWidget* Widget, int32 ShiftAmount);
+	void ShiftColumn(UWidget* Widget, int32 ShiftAmount);
+};
\ No newline at end of file
diff --git a/Engine/Source/Editor/UMGEditor/Private/UMGEditorPrivatePCH.h b/Engine/Source/Editor/UMGEditor/Private/UMGEditorPrivatePCH.h
--- a/Engine/Source/Editor/UMGEditor/Private/UMGEditorPrivatePCH.h
+++ b/Engine/Source/Editor/UMGEditor/Private/UMGEditorPrivatePCH.h
@@ -36,6 +36,7 @@
 #include "HorizontalSlotExtension.h"
 #include "CanvasSlotExtension.h"
 #include "UniformGridSlotExtension.h"
+#include "GridSlotExtension.h"
 
 #include "DetailCustomizations.h"
 #include "CanvasSlotCustomization.h"
diff --git a/Engine/Source/Runtime/Slate/Private/Widgets/Layout/SGridPanel.cpp b/Engine/Source/Runtime/Slate/Private/Widgets/Layout/SGridPanel.cpp
--- a/Engine/Source/Runtime/Slate/Private/Widgets/Layout/SGridPanel.cpp
+++ b/Engine/Source/Runtime/Slate/Private/Widgets/Layout/SGridPanel.cpp
@@ -11,6 +11,21 @@ SGridPanel::FSlot& SGridPanel::AddSlot( int32 Column, int32 Row, SGridPanel::Lay
 	return InsertSlot( new FSlot( Column, Row, InLayer.TheLayer ) );
 }
 
+
+bool SGridPanel::RemoveSlot(const TSharedRef<SWidget>& SlotWidget)
+{
+	for ( int32 SlotIdx = 0; SlotIdx < Slots.Num(); ++SlotIdx )
+	{
+		if ( SlotWidget == Slots[SlotIdx].GetWidget() )
+		{
+			Slots.RemoveAt(SlotIdx);
+			return true;
+		}
+	}
+
+	return false;
+}
+
 void SGridPanel::ClearChildren()
 {
 	Columns.Empty();
@@ -342,7 +357,6 @@ void SGridPanel::NotifySlotChanged(SGridPanel::FSlot* InSlot)
 	}
 }
 
-
 void SGridPanel::ComputeDesiredCellSizes( TArray<float>& OutColumns, TArray<float>& OutRows ) const
 {
 	FMemory::Memzero( OutColumns.GetData(), OutColumns.Num() * sizeof(float) );
diff --git a/Engine/Source/Runtime/Slate/Public/Widgets/Layout/SGridPanel.h b/Engine/Source/Runtime/Slate/Public/Widgets/Layout/SGridPanel.h
--- a/Engine/Source/Runtime/Slate/Public/Widgets/Layout/SGridPanel.h
+++ b/Engine/Source/Runtime/Slate/Public/Widgets/Layout/SGridPanel.h
@@ -57,6 +57,18 @@ public:
 			/** Offset this slot's content by some amount; positive values offset to lower right*/
 			FVector2D NudgeParam;
 
+			/**  */
+			FSlot& Column(int32 Column)
+			{
+				ColumnParam = FMath::Max(0, Column);
+
+				if ( Panel.IsValid() )
+				{
+					Panel.Pin()->NotifySlotChanged(this);
+				}
+
+				return *this;
+			}
 
 			/** How many columns this slot spans over */
 			FSlot& ColumnSpan( int32 ColumnSpan )
@@ -71,6 +83,19 @@ public:
 				return *this;
 			}
 
+			/**  */
+			FSlot& Row(int32 Row)
+			{
+				RowParam = FMath::Max(0, Row);
+
+				if ( Panel.IsValid() )
+				{
+					Panel.Pin()->NotifySlotChanged(this);
+				}
+
+				return *this;
+			}
+
 			/** How many rows this this slot spans over */
 			FSlot& RowSpan( int32 RowSpan )
 			{
@@ -107,6 +132,14 @@ public:
 	 */
 	FSlot& AddSlot( int32 Column, int32 Row, Layer InLayer = Layer(0) );
 
+	/**
+	* Removes a slot from this panel which contains the specified SWidget
+	*
+	* @param SlotWidget The widget to match when searching through the slots
+	* @returns The true if the slot was removed and false if no slot was found matching the widget
+	*/
+	bool RemoveSlot(const TSharedRef<SWidget>& SlotWidget);
+
 	SLATE_BEGIN_ARGS( SGridPanel )
 		{
 			_Visibility = EVisibility::SelfHitTestInvisible;
diff --git a/Engine/Source/Runtime/UMG/Private/Components/GridPanel.cpp b/Engine/Source/Runtime/UMG/Private/Components/GridPanel.cpp
new file mode 100644
--- /dev/null
+++ b/Engine/Source/Runtime/UMG/Private/Components/GridPanel.cpp
@@ -0,0 +1,95 @@
+// Copyright 1998-2014 Epic Games, Inc. All Rights Reserved.
+
+#include "UMGPrivatePCH.h"
+
+#define LOCTEXT_NAMESPACE "UMG"
+
+/////////////////////////////////////////////////////
+// UGridPanel
+
+UGridPanel::UGridPanel(const FPostConstructInitializeProperties& PCIP)
+	: Super(PCIP)
+{
+	bIsVariable = false;
+
+	SGridPanel::FArguments Defaults;
+	Visiblity = UWidget::ConvertRuntimeToSerializedVisiblity(Defaults._Visibility.Get());
+}
+
+void UGridPanel::ReleaseNativeWidget()
+{
+	Super::ReleaseNativeWidget();
+
+	MyGridPanel.Reset();
+}
+
+UClass* UGridPanel::GetSlotClass() const
+{
+	return UGridSlot::StaticClass();
+}
+
+void UGridPanel::OnSlotAdded(UPanelSlot* Slot)
+{
+	// Add the child to the live canvas if it already exists
+	if ( MyGridPanel.IsValid() )
+	{
+		Cast<UGridSlot>(Slot)->BuildSlot(MyGridPanel.ToSharedRef());
+	}
+}
+
+void UGridPanel::OnSlotRemoved(UPanelSlot* Slot)
+{
+	// Remove the widget from the live slot if it exists.
+	if ( MyGridPanel.IsValid() )
+	{
+		TSharedPtr<SWidget> Widget = Slot->Content->GetCachedWidget();
+		if ( Widget.IsValid() )
+		{
+			MyGridPanel->RemoveSlot(Widget.ToSharedRef());
+		}
+	}
+}
+
+TSharedRef<SWidget> UGridPanel::RebuildWidget()
+{
+	MyGridPanel = SNew(SGridPanel);
+
+	for ( UPanelSlot* Slot : Slots )
+	{
+		if ( UGridSlot* TypedSlot = Cast<UGridSlot>(Slot) )
+		{
+			TypedSlot->Parent = this;
+			TypedSlot->BuildSlot( MyGridPanel.ToSharedRef() );
+		}
+	}
+
+	return BuildDesignTimeWidget( MyGridPanel.ToSharedRef() );
+}
+
+void UGridPanel::SyncronizeProperties()
+{
+	Super::SyncronizeProperties();
+
+	for ( int ColumnIndex = 0; ColumnIndex < ColumnFill.Num(); ColumnIndex++ )
+	{
+		MyGridPanel->SetColumnFill(ColumnIndex, ColumnFill[ColumnIndex]);
+	}
+
+	for ( int RowIndex = 0; RowIndex < RowFill.Num(); RowIndex++ )
+	{
+		MyGridPanel->SetRowFill(RowIndex, RowFill[RowIndex]);
+	}
+}
+
+#if WITH_EDITOR
+
+const FSlateBrush* UGridPanel::GetEditorIcon()
+{
+	return FUMGStyle::Get().GetBrush("Widget.Grid");
+}
+
+#endif
+
+/////////////////////////////////////////////////////
+
+#undef LOCTEXT_NAMESPACE
diff --git a/Engine/Source/Runtime/UMG/Private/Components/GridSlot.cpp b/Engine/Source/Runtime/UMG/Private/Components/GridSlot.cpp
new file mode 100644
--- /dev/null
+++ b/Engine/Source/Runtime/UMG/Private/Components/GridSlot.cpp
@@ -0,0 +1,96 @@
+// Copyright 1998-2014 Epic Games, Inc. All Rights Reserved.
+
+#include "UMGPrivatePCH.h"
+
+/////////////////////////////////////////////////////
+// UGridSlot
+
+UGridSlot::UGridSlot(const FPostConstructInitializeProperties& PCIP)
+	: Super(PCIP)
+	, Slot(NULL)
+{
+	HorizontalAlignment = HAlign_Fill;
+	VerticalAlignment = VAlign_Fill;
+
+	Layer = 0;
+	Nudge = FVector2D(0, 0);
+}
+
+void UGridSlot::ReleaseNativeWidget()
+{
+	Super::ReleaseNativeWidget();
+
+	Slot = NULL;
+}
+
+void UGridSlot::BuildSlot(TSharedRef<SGridPanel> GridPanel)
+{
+	Slot = &GridPanel->AddSlot(Column, Row)
+		[
+			Content == NULL ? SNullWidget::NullWidget : Content->TakeWidget()
+		];
+}
+
+void UGridSlot::SetRow(int32 InRow)
+{
+	Row = InRow;
+	if ( Slot )
+	{
+		Slot->Row(InRow);
+	}
+}
+
+void UGridSlot::SetRowSpan(int32 InRowSpan)
+{
+	RowSpan = InRowSpan;
+	if ( Slot )
+	{
+		Slot->RowSpan(InRowSpan);
+	}
+}
+
+void UGridSlot::SetColumn(int32 InColumn)
+{
+	Column = InColumn;
+	if ( Slot )
+	{
+		Slot->Column(InColumn);
+	}
+}
+
+void UGridSlot::SetColumnSpan(int32 InColumnSpan)
+{
+	ColumnSpan = InColumnSpan;
+	if ( Slot )
+	{
+		Slot->ColumnSpan(InColumnSpan);
+	}
+}
+
+void UGridSlot::SetHorizontalAlignment(EHorizontalAlignment InHorizontalAlignment)
+{
+	HorizontalAlignment = InHorizontalAlignment;
+	if ( Slot )
+	{
+		Slot->HAlignment = InHorizontalAlignment;
+	}
+}
+
+void UGridSlot::SetVerticalAlignment(EVerticalAlignment InVerticalAlignment)
+{
+	VerticalAlignment = InVerticalAlignment;
+	if ( Slot )
+	{
+		Slot->VAlignment = InVerticalAlignment;
+	}
+}
+
+void UGridSlot::SyncronizeProperties()
+{
+	SetRow(Row);
+	SetRowSpan(RowSpan);
+	SetColumn(Column);
+	SetColumnSpan(ColumnSpan);
+	SetHorizontalAlignment(HorizontalAlignment);
+	SetVerticalAlignment(VerticalAlignment);
+}
diff --git a/Engine/Source/Runtime/UMG/Private/Components/VerticalBox.cpp b/Engine/Source/Runtime/UMG/Private/Components/VerticalBox.cpp
--- a/Engine/Source/Runtime/UMG/Private/Components/VerticalBox.cpp
+++ b/Engine/Source/Runtime/UMG/Private/Components/VerticalBox.cpp
@@ -48,9 +48,9 @@ void UVerticalBox::OnSlotRemoved(UPanelSlot* Slot)
 	}
 }
 
-UVerticalBoxSlot* UVerticalBox::AddChildToVerticalBox(UWidget* Content)
+UVerticalBoxSlot* UVerticalBox::AddSlot(UWidget* Content)
 {
-	return Cast<UVerticalBoxSlot>( Super::AddChild(Content) );
+	return Cast<UVerticalBoxSlot>(Super::AddChild(Content));
 }
 
 TSharedRef<SWidget> UVerticalBox::RebuildWidget()
diff --git a/Engine/Source/Runtime/UMG/Private/UserWidget.cpp b/Engine/Source/Runtime/UMG/Private/UserWidget.cpp
--- a/Engine/Source/Runtime/UMG/Private/UserWidget.cpp
+++ b/Engine/Source/Runtime/UMG/Private/UserWidget.cpp
@@ -351,6 +351,7 @@ TSharedRef<SWidget> UUserWidget::MakeViewportWidget(bool bAbsoluteLayout, bool b
 
 			+ SConstraintCanvas::Slot()
 			.Offset(BIND_UOBJECT_ATTRIBUTE(FMargin, GetFullScreenOffset))
+			.Anchors(BIND_UOBJECT_ATTRIBUTE(FAnchors, GetViewportAnchors))
 			.Alignment(BIND_UOBJECT_ATTRIBUTE(FVector2D, GetFullScreenAlignment))
 			.ZOrder(BIND_UOBJECT_ATTRIBUTE(int32, GetFullScreenZOrder))
 			[
@@ -478,9 +479,36 @@ APlayerController* UUserWidget::GetPlayerController() const
 	return PlayerContext.IsValid() ? PlayerContext.GetPlayerController() : NULL;
 }
 
+void UUserWidget::SetOffsetInViewport(FVector2D DesiredOffset)
+{
+	ViewportOffsets.Left = DesiredOffset.X;
+	ViewportOffsets.Top = DesiredOffset.Y;
+}
+
+void UUserWidget::SetDesiredSizeInViewport(FVector2D DesiredSize)
+{
+	ViewportOffsets.Right = DesiredSize.X;
+	ViewportOffsets.Bottom = DesiredSize.Y;
+}
+
+void UUserWidget::SetAnchorsInViewport(FVector2D Anchors)
+{
+	ViewportAnchors = Anchors;
+}
+
+void UUserWidget::SetAlignmentInViewport(FVector2D Alignment)
+{
+	ViewportAlignment = Alignment;
+}
+
+void UUserWidget::SetZOrderInViewport(int32 ZOrder)
+{
+	ViewportZOrder = ZOrder;
+}
+
 FMargin UUserWidget::GetFullScreenOffset() const
 {
-	FVector2D FinalSize = FullScreenSize;
+	FVector2D FinalSize = FVector2D(ViewportOffsets.Right, ViewportOffsets.Bottom);
 	if ( FinalSize.IsZero() )
 	{
 		TSharedPtr<SWidget> CachedWidget = GetCachedWidget();
@@ -490,17 +518,22 @@ FMargin UUserWidget::GetFullScreenOffset() const
 		}
 	}
 
-	return FMargin(FullScreenPosition.X, FullScreenPosition.Y, FinalSize.X, FinalSize.Y);
+	return FMargin(ViewportOffsets.Left, ViewportOffsets.Top, FinalSize.X, FinalSize.Y);
+}
+
+FAnchors UUserWidget::GetViewportAnchors() const
+{
+	return FAnchors(ViewportAnchors.X, ViewportAnchors.Y, ViewportAnchors.X, ViewportAnchors.Y);
 }
 
 FVector2D UUserWidget::GetFullScreenAlignment() const
 {
-	return FullScreenAlignment;
+	return ViewportAlignment;
 }
 
 int32 UUserWidget::GetFullScreenZOrder() const
 {
-	return FullScreenZOrder;
+	return ViewportZOrder;
 }
 
 #if WITH_EDITOR
diff --git a/Engine/Source/Runtime/UMG/Public/Blueprint/UserWidget.h b/Engine/Source/Runtime/UMG/Public/Blueprint/UserWidget.h
--- a/Engine/Source/Runtime/UMG/Public/Blueprint/UserWidget.h
+++ b/Engine/Source/Runtime/UMG/Public/Blueprint/UserWidget.h
@@ -113,13 +113,34 @@ public:
 	// End of UVisual interface
 
 	/*  */
-	UFUNCTION(BlueprintCallable, Category="Appearance")
+	UFUNCTION(BlueprintCallable, Category="User Interface|Viewport")
 	void AddToViewport(bool bAbsoluteLayout = false, bool bModal = false, bool bShowCursor = false);
 
 	/*  */
-	UFUNCTION(BlueprintCallable, Category="Appearance")
+	UFUNCTION(BlueprintCallable, Category="User Interface|Viewport")
 	void RemoveFromViewport();
 
+	/*  */
+	UFUNCTION(BlueprintCallable, Category="User Interface|Viewport")
+	void SetOffsetInViewport(FVector2D DesiredOffset);
+
+	/*  */
+	UFUNCTION(BlueprintCallable, Category="User Interface|Viewport")
+	void SetDesiredSizeInViewport(FVector2D DesiredSize);
+
+	/*  */
+	UFUNCTION(BlueprintCallable, Category="User Interface|Viewport")
+	void SetAnchorsInViewport(FVector2D Anchors);
+
+	/*  */
+	UFUNCTION(BlueprintCallable, Category="User Interface|Viewport")
+	void SetAlignmentInViewport(FVector2D Alignment);
+
+	/*  */
+	UFUNCTION(BlueprintCallable, Category="User Interface|Viewport")
+	void SetZOrderInViewport(int32 ZOrder);
+
+	/*  */
 	UFUNCTION(BlueprintPure, Category="Appearance")
 	bool GetIsVisible();
 
@@ -273,22 +294,6 @@ public:
 	UPROPERTY(EditDefaultsOnly, Category=Appearance)
 	FSlateChildSize Size;
 
-	/** The position on the screen to place the UI. */
-	UPROPERTY(EditAnywhere, BlueprintReadWrite, Category=Appearance)
-	FVector2D FullScreenPosition;
-
-	/** The size on the screen the UI should be. */
-	UPROPERTY(EditAnywhere, BlueprintReadWrite, Category=Appearance)
-	FVector2D FullScreenSize;
-
-	/** The normalized UI alignment on the screen, 0..1.  Think of this as the pivot point. */
-	UPROPERTY(EditAnywhere, BlueprintReadWrite, Category=Appearance)
-	FVector2D FullScreenAlignment;
-
-	/** The Z-Order when the UI is fullscreen. */
-	UPROPERTY(EditAnywhere, BlueprintReadWrite, Category=Appearance)
-	int32 FullScreenZOrder;
-
 	/**
 	* Horizontal pivot position
 	*  Given a top aligned slot, where '+' represents the
@@ -336,10 +341,16 @@ protected:
 	virtual TSharedRef<SWidget> RebuildWidget() override;
 
 	FMargin GetFullScreenOffset() const;
+	FAnchors GetViewportAnchors() const;
 	FVector2D GetFullScreenAlignment() const;
 	int32 GetFullScreenZOrder() const;
 
 private:
+	FVector2D ViewportAnchors;
+	FMargin ViewportOffsets;
+	FVector2D ViewportAlignment;
+	int32 ViewportZOrder;
+
 	TWeakPtr<SWidget> FullScreenWidget;
 
 	FLocalPlayerContext PlayerContext;
diff --git a/Engine/Source/Runtime/UMG/Public/Components/GridPanel.h b/Engine/Source/Runtime/UMG/Public/Components/GridPanel.h
new file mode 100644
--- /dev/null
+++ b/Engine/Source/Runtime/UMG/Public/Components/GridPanel.h
@@ -0,0 +1,51 @@
+// Copyright 1998-2014 Epic Games, Inc. All Rights Reserved.
+
+#pragma once
+
+#include "GridPanel.generated.h"
+
+/** A panel that evenly divides up available space between all of its children. */
+UCLASS(meta=( Category="Panel" ), ClassGroup=UserInterface)
+class UMG_API UGridPanel : public UPanelWidget
+{
+	GENERATED_UCLASS_BODY()
+
+public:
+
+	/** Fill */
+	UPROPERTY(EditDefaultsOnly, Category="Fill Rules")
+	TArray<float> ColumnFill;
+
+	/** Fill */
+	UPROPERTY(EditDefaultsOnly, Category="Fill Rules")
+	TArray<float> RowFill;
+
+public:
+
+	// UWidget interface
+	virtual void SyncronizeProperties() override;
+	// End of UWidget interface
+
+	virtual void ReleaseNativeWidget() override;
+
+#if WITH_EDITOR
+	virtual const FSlateBrush* GetEditorIcon() override;
+#endif
+
+protected:
+
+	// UPanelWidget
+	virtual UClass* GetSlotClass() const override;
+	virtual void OnSlotAdded(UPanelSlot* Slot) override;
+	virtual void OnSlotRemoved(UPanelSlot* Slot) override;
+	// End UPanelWidget
+
+protected:
+
+	TSharedPtr<SGridPanel> MyGridPanel;
+
+protected:
+	// UWidget interface
+	virtual TSharedRef<SWidget> RebuildWidget() override;
+	// End of UWidget interface
+};
diff --git a/Engine/Source/Runtime/UMG/Public/Components/GridSlot.h b/Engine/Source/Runtime/UMG/Public/Components/GridSlot.h
new file mode 100644
--- /dev/null
+++ b/Engine/Source/Runtime/UMG/Public/Components/GridSlot.h
@@ -0,0 +1,90 @@
+// Copyright 1998-2014 Epic Games, Inc. All Rights Reserved.
+
+#pragma once
+
+#include "SlateWrapperTypes.h"
+
+#include "GridSlot.generated.h"
+
+/**
+ * A slot for UGridPanel, these slots all share the same size as the largest slot
+ * in the grid.
+ */
+UCLASS()
+class UMG_API UGridSlot : public UPanelSlot
+{
+	GENERATED_UCLASS_BODY()
+
+public:
+
+	/** The alignment of the object horizontally. */
+	UPROPERTY(EditDefaultsOnly, Category=Layout)
+	TEnumAsByte<EHorizontalAlignment> HorizontalAlignment;
+
+	/** The alignment of the object vertically. */
+	UPROPERTY(EditDefaultsOnly, Category=Layout)
+	TEnumAsByte<EVerticalAlignment> VerticalAlignment;
+	
+	/** The row index of the cell this slot is in */
+	UPROPERTY(EditDefaultsOnly, meta=( UIMin = "0" ), Category=Layout)
+	int32 Row;
+	
+	/**  */
+	UPROPERTY(EditDefaultsOnly, Category=Layout)
+	int32 RowSpan;
+
+	/** The column index of the cell this slot is in */
+	UPROPERTY(EditDefaultsOnly, meta=( UIMin = "0" ), Category=Layout)
+	int32 Column;
+
+	/**  */
+	UPROPERTY(EditDefaultsOnly, Category=Layout)
+	int32 ColumnSpan;
+
+	/** Positive values offset this cell to be hit-tested and drawn on top of others. Default is 0; i.e. no offset. */
+	UPROPERTY(EditDefaultsOnly, Category=Layout)
+	int32 Layer;
+
+	/** Offset this slot's content by some amount; positive values offset to lower right*/
+	UPROPERTY(EditDefaultsOnly, Category=Layout)
+	FVector2D Nudge;
+
+	/** Sets the row index of the slot, this determines what cell the slot is in the panel */
+	UFUNCTION(BlueprintCallable, Category=Layout)
+	void SetRow(int32 InRow);
+
+	/** How many rows this this slot spans over */
+	UFUNCTION(BlueprintCallable, Category=Layout)
+	void SetRowSpan(int32 InRowSpan);
+
+	/** Sets the column index of the slot, this determines what cell the slot is in the panel */
+	UFUNCTION(BlueprintCallable, Category=Layout)
+	void SetColumn(int32 InColumn);
+
+	/** How many columns this slot spans over */
+	UFUNCTION(BlueprintCallable, Category=Layout)
+	void SetColumnSpan(int32 InColumnSpan);
+
+	/**  */
+	UFUNCTION(BlueprintCallable, Category=Layout)
+	void SetHorizontalAlignment(EHorizontalAlignment InHorizontalAlignment);
+
+	/**  */
+	UFUNCTION(BlueprintCallable, Category=Layout)
+	void SetVerticalAlignment(EVerticalAlignment InVerticalAlignment);
+
+	// UPanelSlot interface
+	virtual void SyncronizeProperties() override;
+	// End of UPanelSlot interface
+
+	/** Builds the underlying FSlot for the Slate layout panel. */
+	void BuildSlot(TSharedRef<SGridPanel> GridPanel);
+
+	virtual void ReleaseNativeWidget() override;
+
+private:
+	//TODO UMG Slots should hold weak or shared refs to slots.
+
+	/** A raw pointer to the slot to allow us to adjust the size, padding...etc at runtime. */
+	SGridPanel::FSlot* Slot;
+};
\ No newline at end of file
diff --git a/Engine/Source/Runtime/UMG/Public/Components/VerticalBox.h b/Engine/Source/Runtime/UMG/Public/Components/VerticalBox.h
--- a/Engine/Source/Runtime/UMG/Public/Components/VerticalBox.h
+++ b/Engine/Source/Runtime/UMG/Public/Components/VerticalBox.h
@@ -16,8 +16,8 @@ class UMG_API UVerticalBox : public UPanelWidget
 public:
 
 	/**  */
-	UFUNCTION(BlueprintCallable, Category="Widget")
-	UVerticalBoxSlot* AddChildToVerticalBox(UWidget* Content);
+	UFUNCTION(BlueprintCallable, Category="Panel")
+	UVerticalBoxSlot* AddSlot(UWidget* Content);
 
 #if WITH_EDITOR
 	// UWidget interface
diff --git a/Engine/Source/Runtime/UMG/Public/UMG.h b/Engine/Source/Runtime/UMG/Public/UMG.h
--- a/Engine/Source/Runtime/UMG/Public/UMG.h
+++ b/Engine/Source/Runtime/UMG/Public/UMG.h
@@ -45,6 +45,9 @@
 #include "UniformGridSlot.h"
 #include "UniformGridPanel.h"
 
+#include "GridSlot.h"
+#include "GridPanel.h"
+
 #include "OverlaySlot.h"
 #include "Overlay.h"
 
