UMG - Fixing a crash when sniffing functions for bindability, no longer requires all functions found to return a value.

diff --git a/Engine/Source/Editor/UMGEditor/Private/Customizations/DetailCustomizations.cpp b/Engine/Source/Editor/UMGEditor/Private/Customizations/DetailCustomizations.cpp
--- a/Engine/Source/Editor/UMGEditor/Private/Customizations/DetailCustomizations.cpp
+++ b/Engine/Source/Editor/UMGEditor/Private/Customizations/DetailCustomizations.cpp
@@ -262,12 +262,15 @@ void FBlueprintWidgetCustomization::RefreshBlueprintMemberCache(const UFunction*
 	for ( TFieldIterator<UProperty> PropIt(SkeletonClass, EFieldIteratorFlags::ExcludeSuper); PropIt; ++PropIt )
 	{
 		UProperty* Prop = *PropIt;
-		
-		if ( DelegateSignature->GetReturnProperty()->SameType(Prop) )
+
+		if ( UProperty* ReturnProperty = DelegateSignature->GetReturnProperty() )
 		{
-			if ( Prop->HasAnyPropertyFlags(UP::BlueprintReadWrite) )
+			if ( ReturnProperty->SameType(Prop) )
 			{
-				BlueprintPropertyCache.Add(Prop);
+				if ( Prop->HasAnyPropertyFlags(UP::BlueprintReadWrite) )
+				{
+					BlueprintPropertyCache.Add(Prop);
+				}
 			}
 		}
 	}
