UMG - Lots of improvements to the details panel in the designer.  Muticast events are now visible and creatable, function bindable events are now moved to the bottom and always visible.  Function bindings for events no longer are always created as pure.  Generated names for bindings are improved to include the variable name of the widget if available.

diff --git a/Engine/Source/Editor/UMGEditor/Private/Customizations/DetailCustomizations.cpp b/Engine/Source/Editor/UMGEditor/Private/Customizations/DetailCustomizations.cpp
--- a/Engine/Source/Editor/UMGEditor/Private/Customizations/DetailCustomizations.cpp
+++ b/Engine/Source/Editor/UMGEditor/Private/Customizations/DetailCustomizations.cpp
@@ -12,7 +12,7 @@
 
 #define LOCTEXT_NAMESPACE "UMG"
 
-void FBlueprintWidgetCustomization::CreateDelegateCustomization( IDetailLayoutBuilder& DetailLayout, UDelegateProperty* Property )
+void FBlueprintWidgetCustomization::CreateDelegateCustomization( IDetailLayoutBuilder& DetailLayout, UDelegateProperty* Property, UWidget* Widget )
 {
 	FString ConstPropertyName = Property->GetName();
 	ConstPropertyName.RemoveFromEnd(TEXT("Delegate"));
@@ -46,8 +46,8 @@ void FBlueprintWidgetCustomization::CreateDelegateCustomization( IDetailLayoutBu
 	PropertyRow.GetDefaultWidgets(NameWidget, ValueWidget, Row);
 
 	TSharedRef<SComboButton> DelegateButton = SNew(SComboButton)
-		.OnGetMenuContent(this, &FBlueprintWidgetCustomization::OnGenerateDelegateMenu, ConstPropertyHandle, Property->SignatureFunction)
-		.ContentPadding(0)
+		.OnGetMenuContent(this, &FBlueprintWidgetCustomization::OnGenerateDelegateMenu, Widget, ConstPropertyHandle, Property->SignatureFunction, true)
+		.ContentPadding(1)
 		.ButtonContent()
 		[
 			SNew(SHorizontalBox)
@@ -112,7 +112,7 @@ void FBlueprintWidgetCustomization::CreateDelegateCustomization( IDetailLayoutBu
 		];
 }
 
-void FBlueprintWidgetCustomization::CreateEventCustomization( IDetailLayoutBuilder& DetailLayout, UDelegateProperty* Property )
+void FBlueprintWidgetCustomization::CreateEventCustomization( IDetailLayoutBuilder& DetailLayout, UDelegateProperty* Property, UWidget* Widget )
 {
 	if(!Property->GetName().EndsWith(TEXT("Event")))
 	{
@@ -127,7 +127,7 @@ void FBlueprintWidgetCustomization::CreateEventCustomization( IDetailLayoutBuild
 		return;
 	}
 
-	IDetailCategoryBuilder& PropertyCategory = DetailLayout.EditCategory(FObjectEditorUtils::GetCategoryFName(Property));
+	IDetailCategoryBuilder& PropertyCategory = DetailLayout.EditCategory(FObjectEditorUtils::GetCategoryFName(Property), TEXT(""), ECategoryPriority::Uncommon);
 
 	IDetailPropertyRow& PropertyRow = PropertyCategory.AddProperty(DelegatePropertyHandle);
 
@@ -137,8 +137,8 @@ void FBlueprintWidgetCustomization::CreateEventCustomization( IDetailLayoutBuild
 	PropertyRow.GetDefaultWidgets(NameWidget, ValueWidget, Row);
 
 	TSharedRef<SComboButton> DelegateButton = SNew(SComboButton)
-		.OnGetMenuContent(this, &FBlueprintWidgetCustomization::OnGenerateDelegateMenu, DelegatePropertyHandle, Property->SignatureFunction)
-		.ContentPadding(0)
+		.OnGetMenuContent(this, &FBlueprintWidgetCustomization::OnGenerateDelegateMenu, Widget, DelegatePropertyHandle, Property->SignatureFunction, false)
+		.ContentPadding(1)
 		.ButtonContent()
 		[
 			SNew(STextBlock)
@@ -146,17 +146,38 @@ void FBlueprintWidgetCustomization::CreateEventCustomization( IDetailLayoutBuild
 			.Font(IDetailLayoutBuilder::GetDetailFont())
 		];
 
+	FString LabelStr = Property->GetName();
+	LabelStr.RemoveFromEnd(TEXT("Event"));
+
+	FText Label = FText::FromString(LabelStr);
+
 	const bool bShowChildren = true;
 	PropertyRow.CustomWidget(bShowChildren)
 		.NameContent()
 		.MinDesiredWidth(Row.NameWidget.MinWidth)
 		.MaxDesiredWidth(Row.NameWidget.MaxWidth)
 		[
-			NameWidget.ToSharedRef()
+			SNew(SHorizontalBox)
+
+			+ SHorizontalBox::Slot()
+			.AutoWidth()
+			.VAlign(VAlign_Center)
+			.Padding(0,0,5,0)
+			[
+				SNew(SImage)
+				.Image(FEditorStyle::GetBrush("GraphEditor.Event_16x"))
+			]
+
+			+ SHorizontalBox::Slot()
+			.VAlign(VAlign_Center)
+			[
+				SNew(STextBlock)
+				.Text(Label)
+			]
 		]
-	.ValueContent()
-		.MinDesiredWidth(Row.ValueWidget.MinWidth)
-		.MaxDesiredWidth(Row.ValueWidget.MaxWidth)
+		.ValueContent()
+		.MinDesiredWidth(200)
+		.MaxDesiredWidth(200)
 		[
 			SNew(SHorizontalBox)
 
@@ -168,22 +189,134 @@ void FBlueprintWidgetCustomization::CreateEventCustomization( IDetailLayoutBuild
 			]
 
 			+ SHorizontalBox::Slot()
-				.AutoWidth()
+			.AutoWidth()
+			[
+				SNew(SButton)
+				.ButtonStyle(FEditorStyle::Get(), "HoverHintOnly")
+				.Visibility(this, &FBlueprintWidgetCustomization::GetGotoBindingVisibility, DelegatePropertyHandle)
+				.OnClicked(this, &FBlueprintWidgetCustomization::HandleGotoBindingClicked, DelegatePropertyHandle)
+				.VAlign(VAlign_Center)
+				.ToolTipText(LOCTEXT("GotoFunction", "Goto Function"))
 				[
-					SNew(SButton)
-					.ButtonStyle(FEditorStyle::Get(), "HoverHintOnly")
-					.Visibility(this, &FBlueprintWidgetCustomization::GetGotoBindingVisibility, DelegatePropertyHandle)
-					.OnClicked(this, &FBlueprintWidgetCustomization::HandleGotoBindingClicked, DelegatePropertyHandle)
-					.VAlign(VAlign_Center)
-					.ToolTipText(LOCTEXT("GotoFunction", "Goto Function"))
-					[
-						SNew(SImage)
-						.Image(FEditorStyle::GetBrush("PropertyWindow.Button_Browse"))
-					]
+					SNew(SImage)
+					.Image(FEditorStyle::GetBrush("PropertyWindow.Button_Browse"))
 				]
+			]
 		];
 }
 
+void FBlueprintWidgetCustomization::CreateMulticastEventCustomization(IDetailLayoutBuilder& DetailLayout, FName ThisComponentName, UClass* PropertyClass, UMulticastDelegateProperty* DelegateProperty)
+{
+	const FString AddString = FString(TEXT("Add "));
+	const FString ViewString = FString(TEXT("View "));
+
+	const UEdGraphSchema_K2* K2Schema = GetDefault<UEdGraphSchema_K2>();
+
+	if ( !K2Schema->CanUserKismetAccessVariable(DelegateProperty, PropertyClass, UEdGraphSchema_K2::MustBeDelegate) )
+	{
+		return;
+	}
+
+	FString PropertyTooltip = DelegateProperty->GetToolTipText().ToString();
+	if ( PropertyTooltip.IsEmpty() )
+	{
+		PropertyTooltip = DelegateProperty->GetName();
+	}
+
+	// Add on category for delegate property
+	const FString EventCategory = FObjectEditorUtils::GetCategory(DelegateProperty);
+
+	UObjectProperty* ComponentProperty = FindField<UObjectProperty>(Blueprint->SkeletonGeneratedClass, ThisComponentName);
+
+	if ( !ComponentProperty )
+	{
+		return;
+	}
+
+	const UK2Node* EventNode = FKismetEditorUtilities::FindBoundEventForComponent(Blueprint, DelegateProperty->GetFName(), ComponentProperty->GetFName());
+
+	TSharedPtr<FEdGraphSchemaAction> ClickAction;
+
+	if ( EventNode )
+	{
+		TSharedPtr<FEdGraphSchemaAction_K2ViewNode> NewDelegateNode = TSharedPtr<FEdGraphSchemaAction_K2ViewNode>(new FEdGraphSchemaAction_K2ViewNode(EventCategory, FText::FromString(ViewString + DelegateProperty->GetName()), PropertyTooltip, K2Schema->AG_LevelReference));
+		NewDelegateNode->NodePtr = EventNode;
+
+		ClickAction = NewDelegateNode;
+	}
+	else
+	{
+		TSharedPtr<FEdGraphSchemaAction_K2NewNode> NewDelegateNode = TSharedPtr<FEdGraphSchemaAction_K2NewNode>(new FEdGraphSchemaAction_K2NewNode(EventCategory, FText::FromString(AddString + DelegateProperty->GetName()), PropertyTooltip, K2Schema->AG_LevelReference));
+
+		UK2Node_ComponentBoundEvent* NewComponentEvent = NewObject<UK2Node_ComponentBoundEvent>(Blueprint, UK2Node_ComponentBoundEvent::StaticClass());
+		NewComponentEvent->InitializeComponentBoundEventParams(ComponentProperty, DelegateProperty);
+		NewDelegateNode->NodeTemplate = NewComponentEvent;
+		NewDelegateNode->bGotoNode = true;
+
+		ClickAction = NewDelegateNode;
+	}
+
+	TSharedRef<IPropertyHandle> DelegatePropertyHandle = DetailLayout.GetProperty(DelegateProperty->GetFName(), CastChecked<UClass>(DelegateProperty->GetOuter()));
+
+	IDetailCategoryBuilder& PropertyCategory = DetailLayout.EditCategory("Events", LOCTEXT("Events", "Events").ToString(), ECategoryPriority::Uncommon);
+
+	PropertyCategory.AddCustomRow(LOCTEXT("FindInLevelScript_Filter", "Find in Level Script").ToString())
+	.NameContent()
+	[
+		SNew(SHorizontalBox)
+
+		+ SHorizontalBox::Slot()
+		.AutoWidth()
+		.VAlign(VAlign_Center)
+		.Padding(0,0,5,0)
+		[
+			SNew(SImage)
+			.Image(FEditorStyle::GetBrush("GraphEditor.Event_16x"))
+		]
+
+		+ SHorizontalBox::Slot()
+		.VAlign(VAlign_Center)
+		[
+			SNew(STextBlock)
+			.Text(DelegateProperty->GetName())
+		]
+	]
+	.ValueContent()
+	.MinDesiredWidth(200)
+	.MaxDesiredWidth(200)
+	[
+		SNew(SButton)
+		.ToolTipText(ClickAction->TooltipDescription)
+		.OnClicked(this, &FBlueprintWidgetCustomization::AddOrViewEventBinding, ClickAction)
+		.HAlign(HAlign_Center)
+		[
+			SNew(STextBlock)
+			.Text(ClickAction->MenuDescription)
+		]
+	];
+}
+
+FReply FBlueprintWidgetCustomization::AddOrViewEventBinding(TSharedPtr<FEdGraphSchemaAction> Action)
+{
+	UEdGraph* TargetGraph = NULL;
+	if ( Blueprint->UbergraphPages.Num() > 0 )
+	{
+		TargetGraph = Blueprint->UbergraphPages[0]; // Just use the first graph
+	}
+
+	if ( TargetGraph != NULL )
+	{
+		Editor.Pin()->SetCurrentMode(FWidgetBlueprintApplicationModes::GraphMode);
+
+		// Figure out a decent place to stick the node
+		const FVector2D NewNodePos = TargetGraph->GetGoodPlaceForNewNode();
+
+		Action->PerformAction(TargetGraph, NULL, NewNodePos);
+	}
+
+	return FReply::Handled();
+}
+
 void FBlueprintWidgetCustomization::CustomizeDetails( IDetailLayoutBuilder& DetailLayout )
 {
 	TArray< TWeakObjectPtr<UObject> > OutObjects;
@@ -191,22 +324,28 @@ void FBlueprintWidgetCustomization::CustomizeDetails( IDetailLayoutBuilder& Deta
 
 	if ( OutObjects.Num() == 1 )
 	{
+		UWidget* Widget = Cast<UWidget>(OutObjects[0].Get());
 		UClass* PropertyClass = OutObjects[0].Get()->GetClass();
 
 		for ( TFieldIterator<UProperty> PropertyIt(PropertyClass, EFieldIteratorFlags::IncludeSuper); PropertyIt; ++PropertyIt )
 		{
 			UProperty* Property = *PropertyIt;
+
 			if ( UDelegateProperty* DelegateProperty = Cast<UDelegateProperty>(*PropertyIt) )
 			{
 				if( DelegateProperty->GetName().EndsWith(TEXT("Delegate")) )
 				{
-					CreateDelegateCustomization( DetailLayout, DelegateProperty );
+					CreateDelegateCustomization( DetailLayout, DelegateProperty, Widget );
 				}
 				else if( DelegateProperty->GetName().EndsWith(TEXT("Event")) )
 				{
-					CreateEventCustomization( DetailLayout, DelegateProperty );
+					CreateEventCustomization( DetailLayout, DelegateProperty, Widget );
 				}
 			}
+			else if ( UMulticastDelegateProperty* DelegateProperty = Cast<UMulticastDelegateProperty>(Property) )
+			{
+				CreateMulticastEventCustomization(DetailLayout, OutObjects[0].Get()->GetFName(), PropertyClass, DelegateProperty);
+			}
 		}
 	}
 }
@@ -276,7 +415,7 @@ void FBlueprintWidgetCustomization::RefreshBlueprintMemberCache(const UFunction*
 	}
 }
 
-TSharedRef<SWidget> FBlueprintWidgetCustomization::OnGenerateDelegateMenu(TSharedRef<IPropertyHandle> PropertyHandle, UFunction* DelegateSignature)
+TSharedRef<SWidget> FBlueprintWidgetCustomization::OnGenerateDelegateMenu(UWidget* Widget, TSharedRef<IPropertyHandle> PropertyHandle, UFunction* DelegateSignature, bool bIsPure)
 {
 	RefreshBlueprintMemberCache(DelegateSignature);
 
@@ -310,7 +449,7 @@ TSharedRef<SWidget> FBlueprintWidgetCustomization::OnGenerateDelegateMenu(TShare
 			LOCTEXT("CreateBinding", "Create Binding"),
 			LOCTEXT("CreateBindingToolTip", "Creates a new function for this property to be bound to"),
 			FSlateIcon(),
-			FUIAction(FExecuteAction::CreateSP(this, &FBlueprintWidgetCustomization::HandleCreateAndAddBinding, PropertyHandle, DelegateSignature))
+			FUIAction(FExecuteAction::CreateSP(this, &FBlueprintWidgetCustomization::HandleCreateAndAddBinding, Widget, PropertyHandle, DelegateSignature, bIsPure))
 			);
 	}
 	MenuBuilder.EndSection(); //CreateBinding
@@ -534,24 +673,37 @@ void FBlueprintWidgetCustomization::HandleAddPropertyBinding(TSharedRef<IPropert
 	FBlueprintEditorUtils::MarkBlueprintAsStructurallyModified(Blueprint);
 }
 
-void FBlueprintWidgetCustomization::HandleCreateAndAddBinding(TSharedRef<IPropertyHandle> PropertyHandle, UFunction* DelegateSignature)
+void FBlueprintWidgetCustomization::HandleCreateAndAddBinding(UWidget* Widget, TSharedRef<IPropertyHandle> PropertyHandle, UFunction* DelegateSignature, bool bIsPure)
 {
 	const FScopedTransaction Transaction(LOCTEXT("CreateDelegate", "Create Binding"));
 
 	Blueprint->Modify();
 
+	FString Pre = bIsPure ? FString(TEXT("Get")) : FString(TEXT("On"));
+
+	FString WidgetName;
+	if ( Widget && Widget->bIsVariable )
+	{
+		WidgetName = TEXT("_") + Widget->GetName() + TEXT("_");
+	}
+
+	FString Post = PropertyHandle->GetProperty()->GetName();
+	Post.RemoveFromStart(TEXT("On"));
+	Post.RemoveFromEnd(TEXT("Event"));
+
 	// Create the function graph.
-	FString FunctionName = FString(TEXT("Get")) + PropertyHandle->GetProperty()->GetName();
+	FString FunctionName = Pre + WidgetName + Post;
 	UEdGraph* FunctionGraph = FBlueprintEditorUtils::CreateNewGraph(Blueprint, FBlueprintEditorUtils::FindUniqueKismetName(Blueprint, FunctionName), UEdGraph::StaticClass(), UEdGraphSchema_K2::StaticClass());
 
 	bool bUserCreated = true;
 	FBlueprintEditorUtils::AddFunctionGraph(Blueprint, FunctionGraph, bUserCreated, DelegateSignature);
 
-	//TODO UMG Once we start binding event functions, pure needs to be variable, so maybe pass it in.
-	
-	// All bound delegates should be pure.
-	const UEdGraphSchema_K2* Schema_K2 = Cast<UEdGraphSchema_K2>(FunctionGraph->GetSchema());
-	Schema_K2->AddExtraFunctionFlags(FunctionGraph, FUNC_BlueprintPure);
+	// Only mark bindings as pure that need to be.
+	if ( bIsPure )
+	{
+		const UEdGraphSchema_K2* Schema_K2 = Cast<UEdGraphSchema_K2>(FunctionGraph->GetSchema());
+		Schema_K2->AddExtraFunctionFlags(FunctionGraph, FUNC_BlueprintPure);
+	}
 
 	// Add the binding to the blueprint
 	TSharedPtr<FunctionInfo> SelectedFunction = MakeShareable(new FunctionInfo());
diff --git a/Engine/Source/Editor/UMGEditor/Private/Customizations/DetailCustomizations.h b/Engine/Source/Editor/UMGEditor/Private/Customizations/DetailCustomizations.h
--- a/Engine/Source/Editor/UMGEditor/Private/Customizations/DetailCustomizations.h
+++ b/Engine/Source/Editor/UMGEditor/Private/Customizations/DetailCustomizations.h
@@ -32,7 +32,7 @@ private:
 
 	void RefreshBlueprintMemberCache(const UFunction* DelegateSignature);
 
-	TSharedRef<SWidget> OnGenerateDelegateMenu(TSharedRef<IPropertyHandle> PropertyHandle, UFunction* DelegateSignature);
+	TSharedRef<SWidget> OnGenerateDelegateMenu(UWidget* Widget, TSharedRef<IPropertyHandle> PropertyHandle, UFunction* DelegateSignature, bool bIsPure);
 
 	const FSlateBrush* GetCurrentBindingImage(TSharedRef<IPropertyHandle> PropertyHandle) const;
 	FText GetCurrentBindingText(TSharedRef<IPropertyHandle> PropertyHandle) const;
@@ -43,16 +43,20 @@ private:
 	void HandleAddFunctionBinding(TSharedRef<IPropertyHandle> PropertyHandle, TSharedPtr<FunctionInfo> SelectedFunction);
 	void HandleAddPropertyBinding(TSharedRef<IPropertyHandle> PropertyHandle, UProperty* SelectedProperty);
 
-	void HandleCreateAndAddBinding(TSharedRef<IPropertyHandle> PropertyHandle, UFunction* DelegateSignature);
+	void HandleCreateAndAddBinding(UWidget* Widget, TSharedRef<IPropertyHandle> PropertyHandle, UFunction* DelegateSignature, bool bIsPure);
 	void GotoFunction(UEdGraph* FunctionGraph);
 
 	EVisibility GetGotoBindingVisibility(TSharedRef<IPropertyHandle> PropertyHandle) const;
 
 	FReply HandleGotoBindingClicked(TSharedRef<IPropertyHandle> PropertyHandle);
 
-	void CreateDelegateCustomization( IDetailLayoutBuilder& DetailLayout, UDelegateProperty* Property );
+	void CreateDelegateCustomization( IDetailLayoutBuilder& DetailLayout, UDelegateProperty* Property, UWidget* Widget );
 
-	void CreateEventCustomization( IDetailLayoutBuilder& DetailLayout, UDelegateProperty* Property );
+	void CreateEventCustomization( IDetailLayoutBuilder& DetailLayout, UDelegateProperty* Property, UWidget* Widget );
+
+	void CreateMulticastEventCustomization(IDetailLayoutBuilder& DetailLayout, FName ThisComponentName, UClass* PropertyClass, UMulticastDelegateProperty* Property);
+
+	FReply AddOrViewEventBinding(TSharedPtr<FEdGraphSchemaAction> Action);
 private:
 
 	TWeakPtr<FWidgetBlueprintEditor> Editor;
