UMG - We now attempt to migrate edit condition changes.

diff --git a/Engine/Source/Editor/UMGEditor/Private/WidgetBlueprintEditor.cpp b/Engine/Source/Editor/UMGEditor/Private/WidgetBlueprintEditor.cpp
--- a/Engine/Source/Editor/UMGEditor/Private/WidgetBlueprintEditor.cpp
+++ b/Engine/Source/Editor/UMGEditor/Private/WidgetBlueprintEditor.cpp
@@ -8,6 +8,8 @@
 #include "Editor/Sequencer/Public/ISequencerModule.h"
 #include "Animation/UMGSequencerObjectBindingManager.h"
 
+#include "PropertyCustomizationHelpers.h"
+
 #include "WidgetBlueprintApplicationModes.h"
 //#include "WidgetDefafaultsApplicationMode.h"
 #include "WidgetDesignerApplicationMode.h"
@@ -356,6 +358,14 @@ static bool MigratePropertyValue(UObject* SourceObject, UObject* DestinationObje
 		}
 		else
 		{
+			// Check to see if there's an edit condition property we also need to migrate.
+			bool bDummyNegate = false;
+			UBoolProperty* EditConditionProperty = PropertyCustomizationHelpers::GetEditConditionProperty(MemberProperty, bDummyNegate);
+			if ( EditConditionProperty != NULL )
+			{
+				MigratePropertyValue(SourceObject, DestinationObject, EditConditionProperty);
+			}
+
 			return MigratePropertyValue(SourceObject, DestinationObject, MemberProperty);
 		}
 	}
