UMG - Animation setup overhaul.  You can now create multiple animations.  Each animation will create a variable which you can use to Play/Stop animations at runtime in a widget blueprint graph

diff --git a/Engine/Source/Editor/EditorStyle/Private/SlateEditorStyle.cpp b/Engine/Source/Editor/EditorStyle/Private/SlateEditorStyle.cpp
--- a/Engine/Source/Editor/EditorStyle/Private/SlateEditorStyle.cpp
+++ b/Engine/Source/Editor/EditorStyle/Private/SlateEditorStyle.cpp
@@ -5500,6 +5500,17 @@ void FSlateEditorStyle::FStyle::SetupUMGEditorStyles()
 	Set("VerticalAlignment_Center", new IMAGE_BRUSH("Icons/UMG/Alignment/Vertical_Center", Icon20x20));
 	Set("VerticalAlignment_Bottom", new IMAGE_BRUSH("Icons/UMG/Alignment/Vertical_Bottom", Icon20x20));
 	Set("VerticalAlignment_Fill", new IMAGE_BRUSH("Icons/UMG/Alignment/Vertical_Fill", Icon20x20));
+
+	const FTextBlockStyle NoAnimationFont =
+		FTextBlockStyle(NormalText)
+		.SetFont(TTF_CORE_FONT("Fonts/Roboto-Regular", 20))
+		.SetColorAndOpacity(FSlateColor::UseForeground())
+		.SetShadowOffset(FVector2D(1, 1))
+		.SetShadowColorAndOpacity(FLinearColor::Black);
+
+
+	Set("UMGEditor.AddAnimationIcon", new IMAGE_BRUSH("Icons/PlusSymbol_12x", Icon12x12, FLinearColor(.05,.05,.05) ) );
+	Set("UMGEditor.NoAnimationFont", NoAnimationFont);
 }
 
 void FSlateEditorStyle::FStyle::SetupTranslationEditorStyles()
diff --git a/Engine/Source/Editor/Sequencer/Private/Sequencer.cpp b/Engine/Source/Editor/Sequencer/Private/Sequencer.cpp
--- a/Engine/Source/Editor/Sequencer/Private/Sequencer.cpp
+++ b/Engine/Source/Editor/Sequencer/Private/Sequencer.cpp
@@ -190,6 +190,30 @@ UMovieScene* FSequencer::GetFocusedMovieScene() const
 	return MovieSceneStack.Top()->GetMovieScene();
 }
 
+void FSequencer::ResetToNewRootMovieScene( UMovieScene& NewRoot, TSharedRef<ISequencerObjectBindingManager> NewObjectBindingManager )
+{
+	DestroySpawnablesForAllMovieScenes();
+
+	//@todo Sequencer - Encapsulate this better
+	MovieSceneStack.Empty();
+	SelectedSections.Empty();
+	SelectedKeys.Empty();
+	FilteringShots.Empty();
+	UnfilterableSections.Empty();
+	UnfilterableObjects.Empty();
+	MovieSceneSectionToInstanceMap.Empty();
+
+	NewRoot.SetFlags(RF_Transactional);
+
+	ObjectBindingManager = NewObjectBindingManager;
+
+	// Focusing the initial movie scene needs to be done before the first time GetFocusedMovieSceneInstane or GetRootMovieSceneInstance is used
+	RootMovieSceneInstance = MakeShareable(new FMovieSceneInstance(NewRoot));
+	MovieSceneStack.Add(RootMovieSceneInstance.ToSharedRef());
+
+	NotifyMovieSceneDataChanged();
+}
+
 TSharedRef<FMovieSceneInstance> FSequencer::GetRootMovieSceneInstance() const
 {
 	return MovieSceneStack[0];
diff --git a/Engine/Source/Editor/Sequencer/Private/Sequencer.h b/Engine/Source/Editor/Sequencer/Private/Sequencer.h
--- a/Engine/Source/Editor/Sequencer/Private/Sequencer.h
+++ b/Engine/Source/Editor/Sequencer/Private/Sequencer.h
@@ -60,8 +60,9 @@ public:
 
 	/** ISequencer interface */
 	virtual TSharedRef<SWidget> GetSequencerWidget() const override { return SequencerWidget.ToSharedRef(); }
-	virtual UMovieScene* GetRootMovieScene() const;
+	virtual UMovieScene* GetRootMovieScene() const override;
 	virtual UMovieScene* GetFocusedMovieScene() const override;
+	virtual void ResetToNewRootMovieScene( UMovieScene& NewRoot, TSharedRef<ISequencerObjectBindingManager> NewObjectBindingManager );
 	virtual TSharedRef<FMovieSceneInstance> GetRootMovieSceneInstance() const override;
 	virtual TSharedRef<FMovieSceneInstance> GetFocusedMovieSceneInstance() const override;
 	virtual void FocusSubMovieScene( TSharedRef<FMovieSceneInstance> SubMovieSceneInstance ) override;
diff --git a/Engine/Source/Editor/Sequencer/Public/ISequencer.h b/Engine/Source/Editor/Sequencer/Public/ISequencer.h
--- a/Engine/Source/Editor/Sequencer/Public/ISequencer.h
+++ b/Engine/Source/Editor/Sequencer/Public/ISequencer.h
@@ -17,6 +17,9 @@ public:
 	/** @return Widget used to display the sequencer */
 	virtual TSharedRef<SWidget> GetSequencerWidget() const = 0;
 
+	/** @return The root movie scene being used */
+	virtual UMovieScene* GetRootMovieScene() const = 0;
+
 	/** @return Returns the MovieScene that is currently focused for editing by the sequencer.  This can change at any time. */
 	virtual UMovieScene* GetFocusedMovieScene() const = 0;
 
@@ -25,6 +28,9 @@ public:
 	 */
 	virtual TSharedRef<FMovieSceneInstance> GetFocusedMovieSceneInstance() const = 0;
 
+	/** Resets sequencer with a new RootMovieScene */
+	virtual void ResetToNewRootMovieScene( UMovieScene& NewRoot, TSharedRef<class ISequencerObjectBindingManager> NewObjectBindingManager ) = 0;
+
 	/**
 	 * Focuses a sub-movie scene (MovieScene within a MovieScene) in the sequencer
 	 * 
diff --git a/Engine/Source/Editor/UMGEditor/Private/Animation/UMGSequencerObjectBindingManager.cpp b/Engine/Source/Editor/UMGEditor/Private/Animation/UMGSequencerObjectBindingManager.cpp
--- a/Engine/Source/Editor/UMGEditor/Private/Animation/UMGSequencerObjectBindingManager.cpp
+++ b/Engine/Source/Editor/UMGEditor/Private/Animation/UMGSequencerObjectBindingManager.cpp
@@ -6,8 +6,8 @@
 #include "ISequencer.h"
 #include "MovieScene.h"
 
-FUMGSequencerObjectBindingManager::FUMGSequencerObjectBindingManager( FWidgetBlueprintEditor& InWidgetBlueprintEditor, UMovieScene& InMovieScene )
-	: MovieScene( &InMovieScene )
+FUMGSequencerObjectBindingManager::FUMGSequencerObjectBindingManager( FWidgetBlueprintEditor& InWidgetBlueprintEditor, UWidgetAnimation& InAnimation )
+	: WidgetAnimation( &InAnimation )
 	, WidgetBlueprintEditor( InWidgetBlueprintEditor )
 {
 	WidgetBlueprintEditor.GetOnWidgetPreviewUpdated().AddRaw( this, &FUMGSequencerObjectBindingManager::OnWidgetPreviewUpdated );
@@ -16,6 +16,7 @@ FUMGSequencerObjectBindingManager::FUMGSequencerObjectBindingManager( FWidgetBlu
 
 FUMGSequencerObjectBindingManager::~FUMGSequencerObjectBindingManager()
 {
+	WidgetBlueprintEditor.GetOnWidgetPreviewUpdated().RemoveAll( this );
 }
 
 FGuid FUMGSequencerObjectBindingManager::FindGuidForObject( const UMovieScene& MovieScene, UObject& Object ) const
@@ -37,9 +38,6 @@ void FUMGSequencerObjectBindingManager::BindPossessableObject( const FGuid& Poss
 {
 	PreviewObjectToGuidMap.Add( &PossessedObject, PossessableGuid );
 
-	UWidgetBlueprint* WidgetBlueprint = WidgetBlueprintEditor.GetWidgetBlueprintObj();
-	FWidgetAnimation* WidgetAnimation = WidgetBlueprint->FindAnimationDataForMovieScene(*MovieScene);
-
 	FWidgetAnimationBinding NewBinding;
 	NewBinding.WidgetName = PossessedObject.GetFName();
 	if( PossessedObject.IsA<UPanelSlot>() )
@@ -64,7 +62,6 @@ void FUMGSequencerObjectBindingManager::UnbindPossessableObjects( const FGuid& P
 	}
 
 	UWidgetBlueprint* WidgetBlueprint = WidgetBlueprintEditor.GetWidgetBlueprintObj();
-	FWidgetAnimation* WidgetAnimation = WidgetBlueprint->FindAnimationDataForMovieScene(*MovieScene);
 
 	WidgetAnimation->AnimationBindings.RemoveAll( [&]( const FWidgetAnimationBinding& Binding ) { return Binding.AnimationGuid == PossessableGuid; } );
 
@@ -82,41 +79,45 @@ void FUMGSequencerObjectBindingManager::GetRuntimeObjects( const TSharedRef<FMov
 	}
 }
 
-void FUMGSequencerObjectBindingManager::InitPreviewObjects()
+bool FUMGSequencerObjectBindingManager::HasValidWidgetAnimation() const
 {
 	UWidgetBlueprint* WidgetBlueprint = WidgetBlueprintEditor.GetWidgetBlueprintObj();
+	return WidgetAnimation.IsValid() && WidgetBlueprint->Animations.Contains( WidgetAnimation.Get() );
+}
 
+void FUMGSequencerObjectBindingManager::InitPreviewObjects()
+{
 	// Populate preview object to guid map
 	UUserWidget* PreviewWidget = WidgetBlueprintEditor.GetPreview();
 
-	UWidgetTree* WidgetTree = PreviewWidget->WidgetTree;
-
-	const FWidgetAnimation* WidgetAnimation = WidgetBlueprint->FindAnimationDataForMovieScene( *MovieScene );
-	check( WidgetAnimation );
-
-	for( const FWidgetAnimationBinding& Binding : WidgetAnimation->AnimationBindings )
+	if( PreviewWidget )
 	{
-		FName ObjectName = Binding.WidgetName;
-		FName SlotWidgetName = Binding.SlotWidgetName;
+		UWidgetTree* WidgetTree = PreviewWidget->WidgetTree;
 
-		FName NameToSearchFor = SlotWidgetName != NAME_None ? SlotWidgetName : ObjectName;
-
-		UObject* FoundObject = FindObject<UObject>(WidgetTree, *NameToSearchFor.ToString());
-		if(FoundObject)
+		for(const FWidgetAnimationBinding& Binding : WidgetAnimation->AnimationBindings)
 		{
-			if( SlotWidgetName == NAME_None )
-			{
-				PreviewObjectToGuidMap.Add(FoundObject, Binding.AnimationGuid);
-			}
-			else
+			FName ObjectName = Binding.WidgetName;
+			FName SlotWidgetName = Binding.SlotWidgetName;
+
+			FName NameToSearchFor = SlotWidgetName != NAME_None ? SlotWidgetName : ObjectName;
+
+			UObject* FoundObject = FindObject<UObject>(WidgetTree, *NameToSearchFor.ToString());
+			if(FoundObject)
 			{
-				FoundObject = FindObject<UObject>( FoundObject, *ObjectName.ToString() );
-				if( FoundObject )
+				if(SlotWidgetName == NAME_None)
 				{
 					PreviewObjectToGuidMap.Add(FoundObject, Binding.AnimationGuid);
 				}
+				else
+				{
+					FoundObject = FindObject<UObject>(FoundObject, *ObjectName.ToString());
+					if(FoundObject)
+					{
+						PreviewObjectToGuidMap.Add(FoundObject, Binding.AnimationGuid);
+					}
+				}
+
 			}
-		
 		}
 	}
 }
diff --git a/Engine/Source/Editor/UMGEditor/Private/Animation/UMGSequencerObjectBindingManager.h b/Engine/Source/Editor/UMGEditor/Private/Animation/UMGSequencerObjectBindingManager.h
--- a/Engine/Source/Editor/UMGEditor/Private/Animation/UMGSequencerObjectBindingManager.h
+++ b/Engine/Source/Editor/UMGEditor/Private/Animation/UMGSequencerObjectBindingManager.h
@@ -7,7 +7,7 @@
 class FUMGSequencerObjectBindingManager : public ISequencerObjectBindingManager
 {
 public:
-	FUMGSequencerObjectBindingManager( FWidgetBlueprintEditor& InWidgetBlueprintEditor, UMovieScene& InMovieScene );
+	FUMGSequencerObjectBindingManager( FWidgetBlueprintEditor& InWidgetBlueprintEditor, UWidgetAnimation& InAnimation );
 	~FUMGSequencerObjectBindingManager();
 
 	/** ISequencerObjectBindingManager interface */
@@ -20,6 +20,12 @@ public:
 	virtual void UnbindPossessableObjects( const FGuid& PossessableGuid ) override;
 	virtual void GetRuntimeObjects( const TSharedRef<FMovieSceneInstance>& MovieSceneInstance, const FGuid& ObjectGuid, TArray<UObject*>& OutRuntimeObjects ) const override;
 	
+	/** @return true if the current animation is valid */
+	bool HasValidWidgetAnimation() const;
+
+	/** @return the current widget animation */
+	UWidgetAnimation* GetWidgetAnimation() { return WidgetAnimation.Get(); }
+
 	/** Rebuilds mappings to live preview objects */
 	void InitPreviewObjects();
 
@@ -29,7 +35,7 @@ private:
 	/** Mapping of preview objects to sequencer guids */
 	TMap< TWeakObjectPtr<UObject>, FGuid> PreviewObjectToGuidMap;
 
-	TWeakObjectPtr<UMovieScene> MovieScene;
+	TWeakObjectPtr<UWidgetAnimation> WidgetAnimation;
 
 	FWidgetBlueprintEditor& WidgetBlueprintEditor;
 };
diff --git a/Engine/Source/Editor/UMGEditor/Private/BlueprintModes/WidgetDesignerApplicationMode.cpp b/Engine/Source/Editor/UMGEditor/Private/BlueprintModes/WidgetDesignerApplicationMode.cpp
--- a/Engine/Source/Editor/UMGEditor/Private/BlueprintModes/WidgetDesignerApplicationMode.cpp
+++ b/Engine/Source/Editor/UMGEditor/Private/BlueprintModes/WidgetDesignerApplicationMode.cpp
@@ -20,6 +20,7 @@
 #include "DesignerTabSummoner.h"
 #include "SequencerTabSummoner.h"
 #include "DetailsTabSummoner.h"
+#include "AnimationTabSummoner.h"
 
 /////////////////////////////////////////////////////
 // FWidgetDesignerApplicationMode
@@ -43,6 +44,7 @@ FWidgetDesignerApplicationMode::FWidgetDesignerApplicationMode(TSharedPtr<FWidge
 		(
 			FTabManager::NewSplitter()
 			->SetOrientation(Orient_Horizontal)
+			->SetSizeCoefficient( 0.70f )
 			->Split
 			(
 				FTabManager::NewSplitter()
@@ -75,17 +77,30 @@ FWidgetDesignerApplicationMode::FWidgetDesignerApplicationMode(TSharedPtr<FWidge
 				->Split
 				(
 					FTabManager::NewStack()
-					->SetSizeCoefficient( 0.5f )
+					->SetSizeCoefficient( 0.35f )
 					->AddTab( FDetailsTabSummoner::TabID, ETabState::OpenedTab )
 				)
 			)
 		)
 		->Split
 		(
-			FTabManager::NewStack()
-			->SetSizeCoefficient(0.2f)
-			->SetHideTabWell(true)
-			->AddTab( FSequencerTabSummoner::TabID, ETabState::OpenedTab )
+			FTabManager::NewSplitter()
+			->SetOrientation(Orient_Horizontal)
+			->SetSizeCoefficient( 0.30f )
+			->Split
+			(
+				FTabManager::NewStack()
+				->SetSizeCoefficient(.15f)
+				->AddTab(FAnimationTabSummoner::TabID, ETabState::OpenedTab)
+			)
+			->Split
+			(
+				FTabManager::NewStack()
+				->SetSizeCoefficient(.85f)
+				->SetHideTabWell(true)
+				->AddTab(FSequencerTabSummoner::TabID, ETabState::OpenedTab)
+
+			)
 		)
 	);
 
@@ -97,6 +112,7 @@ FWidgetDesignerApplicationMode::FWidgetDesignerApplicationMode(TSharedPtr<FWidge
 	TabFactories.RegisterFactory(MakeShareable(new FHierarchyTabSummoner(InWidgetEditor)));
 	TabFactories.RegisterFactory(MakeShareable(new FPaletteTabSummoner(InWidgetEditor)));
 	TabFactories.RegisterFactory(MakeShareable(new FSequencerTabSummoner(InWidgetEditor)));
+	TabFactories.RegisterFactory(MakeShareable(new FAnimationTabSummoner(InWidgetEditor)));
 
 	// setup toolbar - clear existing toolbar extender from the BP mode
 	//@TODO: Keep this in sync with BlueprintEditorModes.cpp
diff --git a/Engine/Source/Editor/UMGEditor/Private/TabFactory/AnimationTabSummoner.cpp b/Engine/Source/Editor/UMGEditor/Private/TabFactory/AnimationTabSummoner.cpp
new file mode 100644
--- /dev/null
+++ b/Engine/Source/Editor/UMGEditor/Private/TabFactory/AnimationTabSummoner.cpp
@@ -0,0 +1,481 @@
+// Copyright 1998-2014 Epic Games, Inc. All Rights Reserved.
+
+#include "UMGEditorPrivatePCH.h"
+
+#include "AnimationTabSummoner.h"
+#include "WidgetBlueprintEditor.h"
+#include "ISequencer.h"
+
+#define LOCTEXT_NAMESPACE "UMG"
+
+const FName FAnimationTabSummoner::TabID(TEXT("Animations"));
+
+FAnimationTabSummoner::FAnimationTabSummoner(TSharedPtr<class FWidgetBlueprintEditor> InBlueprintEditor)
+		: FWorkflowTabFactory(TabID, InBlueprintEditor)
+		, BlueprintEditor(InBlueprintEditor)
+{
+	TabLabel = LOCTEXT("AnimationsTabLabel", "Animations");
+
+	bIsSingleton = true;
+
+	ViewMenuDescription = LOCTEXT("Animations_ViewMenu_Desc", "Animations");
+	ViewMenuTooltip = LOCTEXT("Animations_ViewMenu_ToolTip", "Opens a tab to manage animations");
+}
+
+struct FWidgetAnimationListItem
+{
+	FWidgetAnimationListItem(UWidgetAnimation* InAnimation, bool bInRenameRequestPending = false, bool bInNewAnimation = false )
+		: Animation(InAnimation)
+		, bRenameRequestPending(bInRenameRequestPending)
+		, bNewAnimation( bInNewAnimation )
+	{}
+
+	UWidgetAnimation* Animation;
+	bool bRenameRequestPending;
+	bool bNewAnimation;
+};
+
+
+typedef SListView<TSharedPtr<FWidgetAnimationListItem> > SWidgetAnimationListView;
+
+class SWidgetAnimationListItem : public STableRow<TSharedPtr<FWidgetAnimationListItem> >
+{
+public:
+	SLATE_BEGIN_ARGS( SWidgetAnimationListItem ){}
+	SLATE_END_ARGS()
+
+	void Construct( const FArguments& InArgs, const TSharedRef<STableViewBase>& InOwnerTableView, TSharedPtr<FWidgetBlueprintEditor> InBlueprintEditor, TSharedPtr<FWidgetAnimationListItem> InListItem )
+	{
+		ListItem = InListItem;
+		BlueprintEditor = InBlueprintEditor;
+
+		STableRow<TSharedPtr<FWidgetAnimationListItem>>::Construct(
+			STableRow<TSharedPtr<FWidgetAnimationListItem>>::FArguments()
+			.Padding( FMargin( 3.0f, 2.0f) )
+			.Content()
+			[
+				SAssignNew(InlineTextBlock, SInlineEditableTextBlock)
+				.Font(FCoreStyle::Get().GetFontStyle("NormalFont"))
+				.Text(this, &SWidgetAnimationListItem::GetMovieSceneText)
+				//.HighlightText(InArgs._HighlightText)
+				.OnVerifyTextChanged(this, &SWidgetAnimationListItem::OnVerifyNameTextChanged)
+				.OnTextCommitted(this, &SWidgetAnimationListItem::OnNameTextCommited)
+				.IsSelected(this, &SWidgetAnimationListItem::IsSelectedExclusively)
+			],
+			InOwnerTableView);
+	}
+
+	void BeginRename()
+	{
+		InlineTextBlock->EnterEditingMode();
+	}
+
+private:
+	FText GetMovieSceneText() const
+	{
+		if( ListItem.IsValid() )
+		{
+			return FText::FromString( ListItem.Pin()->Animation->GetName() );
+		}
+
+		return FText::GetEmpty();
+	}
+
+	bool OnVerifyNameTextChanged(const FText& InText, FText& OutErrorMessage)
+	{
+		UWidgetAnimation* Animation = ListItem.Pin()->Animation;
+
+		FString NewName = InText.ToString();
+
+		UWidgetBlueprint* Blueprint = BlueprintEditor.Pin()->GetWidgetBlueprintObj();
+		if( Animation->GetName() != NewName && FindObject<UWidgetAnimation>( Blueprint, *NewName, true ) )
+		{
+			OutErrorMessage = LOCTEXT("ExistingMovieSceneError", "An animation with this name already exists");
+			return false;
+		}
+
+		return true;
+	}
+
+	void OnNameTextCommited(const FText& InText, ETextCommit::Type CommitInfo)
+	{
+		UWidgetAnimation* WidgetAnimation = ListItem.Pin()->Animation;
+		UWidgetBlueprint* Blueprint = BlueprintEditor.Pin()->GetWidgetBlueprintObj();
+
+		FName CurrentName = WidgetAnimation->GetFName();
+
+		if(!CurrentName.ToString().Equals(InText.ToString()) && !InText.IsEmpty())
+		{
+			bool bNewAnimation = ListItem.Pin()->bNewAnimation;
+			FText TransactionName = bNewAnimation ? LOCTEXT("NewAnimation", "New Animation") : LOCTEXT("RenameAnimation", "Rename Animation");
+
+			{
+				const FScopedTransaction Transaction(TransactionName);
+				WidgetAnimation->Modify();
+				WidgetAnimation->MovieScene->Modify();
+
+				WidgetAnimation->Rename(*InText.ToString());
+				WidgetAnimation->MovieScene->Rename(*InText.ToString());
+
+				if(bNewAnimation)
+				{
+					Blueprint->Modify();
+					Blueprint->Animations.Add(WidgetAnimation);
+					ListItem.Pin()->bNewAnimation = false;
+				}
+			}
+
+			FBlueprintEditorUtils::ReplaceVariableReferences(Blueprint, CurrentName, FName(*InText.ToString()));
+
+			FBlueprintEditorUtils::MarkBlueprintAsStructurallyModified(Blueprint);
+		}
+	}
+private:
+	TWeakPtr<FWidgetAnimationListItem> ListItem;
+	TWeakPtr<FWidgetBlueprintEditor> BlueprintEditor;
+	TSharedPtr<SInlineEditableTextBlock> InlineTextBlock;
+};
+
+
+class SUMGAnimationList : public SCompoundWidget
+{
+public:
+	SLATE_BEGIN_ARGS( SUMGAnimationList	) {}
+	SLATE_END_ARGS()
+
+	void Construct( const FArguments& InArgs, TSharedPtr<FWidgetBlueprintEditor> InBlueprintEditor )
+	{
+		BlueprintEditor = InBlueprintEditor;
+
+		InBlueprintEditor->GetOnWidgetBlueprintTransaction().AddSP( this, &SUMGAnimationList::OnWidgetBlueprintTransaction );
+
+		SAssignNew(AnimationListView, SWidgetAnimationListView)
+			.ItemHeight(20.0f)
+			.SelectionMode(ESelectionMode::Single)
+			.OnGenerateRow(this, &SUMGAnimationList::OnGenerateWidgetForMovieScene)
+			.OnItemScrolledIntoView(this, &SUMGAnimationList::OnItemScrolledIntoView)
+			.OnSelectionChanged(this, &SUMGAnimationList::OnSelectionChanged)
+			.OnContextMenuOpening(this, &SUMGAnimationList::OnContextMenuOpening)
+			.ListItemsSource(&Animations);
+
+		ChildSlot
+		[
+			SNew(SBorder)
+			.BorderImage(FEditorStyle::GetBrush("ToolPanel.GroupBorder"))
+			[
+				SNew(SVerticalBox)
+				+ SVerticalBox::Slot()
+				.Padding( 2 )
+				.AutoHeight()
+				[
+					SNew( SHorizontalBox )
+					+ SHorizontalBox::Slot()
+					.Padding(0)
+					.VAlign( VAlign_Center )
+					.AutoWidth()
+					[
+						SNew(SButton)
+						.ContentPadding(FMargin(2.0f, 1.0f))
+						.OnClicked( this, &SUMGAnimationList::OnNewAnimationClicked )
+						.HAlign(HAlign_Center)
+						.VAlign(VAlign_Center)
+						[
+							SNew( SHorizontalBox )
+							+ SHorizontalBox::Slot()
+							.VAlign(VAlign_Center)
+							.AutoWidth()
+							[
+								SNew( SImage )
+								.Image( FEditorStyle::GetBrush("UMGEditor.AddAnimationIcon")) 
+							]
+							+ SHorizontalBox::Slot()
+							.Padding( 2.0f, 0.0f )
+							[
+								SNew( STextBlock )
+								.Text( LOCTEXT("NewAnimationButtonText", "New") )
+							]
+						]
+					]
+					+ SHorizontalBox::Slot()
+					.Padding(2.0f, 0.0f)
+					.VAlign( VAlign_Center )
+					[
+						SNew(SSearchBox)
+						.HintText(LOCTEXT("Search Animations", "Search Animations"))
+						.OnTextChanged(this, &SUMGAnimationList::OnSearchChanged)
+					]
+				]
+				+ SVerticalBox::Slot()
+				.FillHeight(1.0f)
+				[
+					SNew(SScrollBorder, AnimationListView.ToSharedRef())
+					[
+						AnimationListView.ToSharedRef()
+					]
+				]
+			]
+		];
+
+		UpdateAnimationList();
+
+		CreateCommandList();
+	}
+
+	virtual FReply OnKeyDown( const FGeometry& MyGeometry, const FKeyboardEvent& InKeyboardEvent )
+	{
+		FReply Reply = FReply::Unhandled();
+		if( CommandList->ProcessCommandBindings( InKeyboardEvent ) )
+		{
+			Reply = FReply::Handled();
+		}
+
+		return Reply;
+	}
+
+private:
+
+	void UpdateAnimationList()
+	{
+		Animations.Empty();
+
+		const TArray<UWidgetAnimation*>& WidgetAnimations = BlueprintEditor.Pin()->GetWidgetBlueprintObj()->Animations;
+
+		for( UWidgetAnimation* Animation : WidgetAnimations )
+		{
+			Animations.Add( MakeShareable( new FWidgetAnimationListItem( Animation ) ) );
+		}
+		
+		AnimationListView->RequestListRefresh();
+	}
+
+	void OnWidgetBlueprintTransaction()
+	{
+		UpdateAnimationList();
+
+		
+		TSharedPtr<FWidgetBlueprintEditor> WidgetBlueprintEditorPin = BlueprintEditor.Pin();
+		const UWidgetAnimation* ViewedAnim = WidgetBlueprintEditorPin->RefreshCurrentAnimation();
+
+		if(ViewedAnim)
+		{
+			TSharedPtr<FWidgetAnimationListItem> FoundListItem = *Animations.FindByPredicate( [&](const TSharedPtr<FWidgetAnimationListItem>& ListItem ) { return ListItem->Animation == ViewedAnim; } );
+			
+			AnimationListView->SetSelection(FoundListItem);
+		}
+
+	}
+
+	void OnItemScrolledIntoView( TSharedPtr<FWidgetAnimationListItem> InListItem, const TSharedPtr<ITableRow>& InWidget ) const
+	{
+		if( InListItem->bRenameRequestPending )
+		{
+			StaticCastSharedPtr<SWidgetAnimationListItem>( InWidget )->BeginRename();
+			InListItem->bRenameRequestPending = false;
+		}
+	}
+
+	FReply OnNewAnimationClicked()
+	{
+		UWidgetBlueprint* WidgetBlueprint = BlueprintEditor.Pin()->GetWidgetBlueprintObj();
+
+		UWidgetAnimation* NewAnimation = ConstructObject<UWidgetAnimation>(UWidgetAnimation::StaticClass(), WidgetBlueprint, MakeUniqueObjectName(WidgetBlueprint, UWidgetAnimation::StaticClass(), "NewAnimation"), RF_Transactional);
+		NewAnimation->MovieScene =  ConstructObject<UMovieScene>(UMovieScene::StaticClass(), NewAnimation, NewAnimation->GetFName(), RF_Transactional);
+
+		bool bRequestRename = true;
+		bool bNewAnimation = true;
+
+		int32 NewIndex = Animations.Add( MakeShareable( new FWidgetAnimationListItem( NewAnimation, bRequestRename, bNewAnimation ) ) );
+
+		AnimationListView->RequestScrollIntoView( Animations[NewIndex] );
+
+		return FReply::Handled();
+	}
+
+	void OnSearchChanged( const FText& InSearchText )
+	{
+		const TArray<UWidgetAnimation*>& WidgetAnimations = BlueprintEditor.Pin()->GetWidgetBlueprintObj()->Animations;
+
+		if( !InSearchText.IsEmpty() )
+		{
+			struct Local
+			{
+				static void UpdateFilterStrings(UWidgetAnimation* InAnimation, OUT TArray< FString >& OutFilterStrings)
+				{
+					OutFilterStrings.Add(InAnimation->GetName());
+				}
+			};
+
+			TTextFilter<UWidgetAnimation*> TextFilter(TTextFilter<UWidgetAnimation*>::FItemToStringArray::CreateStatic(&Local::UpdateFilterStrings));
+
+			TextFilter.SetRawFilterText(InSearchText);
+
+			Animations.Reset();
+
+			for(UWidgetAnimation* Animation : WidgetAnimations)
+			{
+				if(TextFilter.PassesFilter(Animation))
+				{
+					Animations.Add(MakeShareable(new FWidgetAnimationListItem(Animation)));
+				}
+			}
+
+			AnimationListView->RequestListRefresh();
+		}
+		else
+		{
+			// Just regenerate the whole list
+			UpdateAnimationList();
+		}
+	}
+
+	void OnSelectionChanged( TSharedPtr<FWidgetAnimationListItem> InSelectedItem, ESelectInfo::Type SelectionInfo )
+	{
+		if( InSelectedItem.IsValid() )
+		{
+			UWidgetBlueprint* WidgetBlueprint = BlueprintEditor.Pin()->GetWidgetBlueprintObj();
+			BlueprintEditor.Pin()->ChangeViewedAnimation( *InSelectedItem->Animation );
+		}
+	}
+
+	TSharedPtr<SWidget> OnContextMenuOpening() const
+	{
+		FMenuBuilder MenuBuilder( true, CommandList.ToSharedRef() );
+
+		MenuBuilder.BeginSection("Edit", LOCTEXT("Edit", "Edit"));
+		{
+			MenuBuilder.AddMenuEntry(FGenericCommands::Get().Rename);
+			MenuBuilder.AddMenuEntry(FGenericCommands::Get().Duplicate);
+			MenuBuilder.AddMenuSeparator();
+
+			MenuBuilder.AddMenuEntry(FGenericCommands::Get().Delete);
+		}
+		MenuBuilder.EndSection();
+
+		return MenuBuilder.MakeWidget();
+	}
+
+	TSharedRef<ITableRow> OnGenerateWidgetForMovieScene( TSharedPtr<FWidgetAnimationListItem> InListItem, const TSharedRef< STableViewBase >& InOwnerTableView )
+	{
+		return SNew( SWidgetAnimationListItem, InOwnerTableView, BlueprintEditor.Pin(), InListItem );
+	}
+
+	void CreateCommandList()
+	{
+		CommandList = MakeShareable(new FUICommandList);
+
+		CommandList->MapAction(
+			FGenericCommands::Get().Duplicate,
+			FExecuteAction::CreateSP(this, &SUMGAnimationList::OnDuplicateAnimation),
+			FCanExecuteAction::CreateSP(this, &SUMGAnimationList::CanExecuteContextMenuAction)
+			);
+
+		CommandList->MapAction(
+			FGenericCommands::Get().Delete,
+			FExecuteAction::CreateSP(this, &SUMGAnimationList::OnDeleteAnimation),
+			FCanExecuteAction::CreateSP(this, &SUMGAnimationList::CanExecuteContextMenuAction)
+			);
+
+		CommandList->MapAction(
+			FGenericCommands::Get().Rename,
+			FExecuteAction::CreateSP(this, &SUMGAnimationList::OnRenameAnimation),
+			FCanExecuteAction::CreateSP(this, &SUMGAnimationList::CanExecuteContextMenuAction)
+			);
+	}
+
+	bool CanExecuteContextMenuAction() const
+	{
+		return AnimationListView->GetNumItemsSelected() == 1 && !BlueprintEditor.Pin()->InDebuggingMode();
+	}
+
+	void OnDuplicateAnimation()
+	{
+		TArray< TSharedPtr<FWidgetAnimationListItem> > SelectedAnimations = AnimationListView->GetSelectedItems();
+		check(SelectedAnimations.Num() == 1);
+
+		TSharedPtr<FWidgetAnimationListItem> SelectedAnimation = SelectedAnimations[0];
+
+		UWidgetBlueprint* WidgetBlueprint = BlueprintEditor.Pin()->GetWidgetBlueprintObj();
+
+		UWidgetAnimation* NewAnimation = 
+			DuplicateObject<UWidgetAnimation>
+			(
+				SelectedAnimation->Animation,
+				WidgetBlueprint, 
+				*MakeUniqueObjectName( WidgetBlueprint, UWidgetAnimation::StaticClass(), SelectedAnimation->Animation->GetFName() ).ToString()
+			);
+	
+		NewAnimation->MovieScene->Rename(*NewAnimation->GetName(), nullptr, REN_DontCreateRedirectors | REN_ForceNoResetLoaders);
+
+
+		bool bRenameRequestPending = true;
+		bool bNewAnimation = true;
+		int32 NewIndex = Animations.Add(MakeShareable(new FWidgetAnimationListItem(NewAnimation,bRenameRequestPending,bNewAnimation)));
+
+		AnimationListView->RequestScrollIntoView(Animations[NewIndex]);
+	}
+
+
+	void OnDeleteAnimation()
+	{
+		TArray< TSharedPtr<FWidgetAnimationListItem> > SelectedAnimations = AnimationListView->GetSelectedItems();
+		check(SelectedAnimations.Num() == 1);
+
+		TSharedPtr<FWidgetAnimationListItem> SelectedAnimation = SelectedAnimations[0];
+
+		TSharedPtr<FWidgetBlueprintEditor> WidgetBlueprintEditorPin = BlueprintEditor.Pin();
+
+		UWidgetBlueprint* WidgetBlueprint = WidgetBlueprintEditorPin->GetWidgetBlueprintObj();
+
+		TArray<UWidgetAnimation*>& WidgetAnimations = WidgetBlueprint->Animations;
+
+		{
+			const FScopedTransaction Transaction(LOCTEXT("DeleteAnimationTransaction", "Delete Animation"));
+			WidgetBlueprint->Modify();
+
+			WidgetAnimations.Remove(SelectedAnimation->Animation);
+
+			UpdateAnimationList();
+		}
+
+		FBlueprintEditorUtils::MarkBlueprintAsStructurallyModified(WidgetBlueprint);
+
+		if(Animations.Num() > 0)
+		{
+			// Ensure we are viewing a valid animation
+			WidgetBlueprintEditorPin->ChangeViewedAnimation(*Animations[0]->Animation);
+			AnimationListView->SetSelection( Animations[0] );
+		}
+		else
+		{
+			WidgetBlueprintEditorPin->ChangeViewedAnimation(*UWidgetAnimation::GetNullAnimation());
+		}
+	}
+
+	void OnRenameAnimation()
+	{
+		TArray< TSharedPtr<FWidgetAnimationListItem> > SelectedAnimations = AnimationListView->GetSelectedItems();
+		check( SelectedAnimations.Num() == 1 );
+
+		TSharedPtr<FWidgetAnimationListItem> SelectedAnimation = SelectedAnimations[0];
+		SelectedAnimation->bRenameRequestPending = true;
+
+		AnimationListView->RequestScrollIntoView( SelectedAnimation );
+	}
+
+private:
+	TSharedPtr<FUICommandList> CommandList;
+	TWeakPtr<FWidgetBlueprintEditor> BlueprintEditor;
+	TSharedPtr<SWidgetAnimationListView> AnimationListView;
+	TArray< TSharedPtr<FWidgetAnimationListItem> > Animations;
+};
+
+
+TSharedRef<SWidget> FAnimationTabSummoner::CreateTabBody(const FWorkflowTabSpawnInfo& Info) const
+{
+	TSharedPtr<FWidgetBlueprintEditor> BlueprintEditorPinned = BlueprintEditor.Pin();
+
+	return SNew( SUMGAnimationList, BlueprintEditorPinned );
+		
+}
+
+#undef LOCTEXT_NAMESPACE 
diff --git a/Engine/Source/Editor/UMGEditor/Private/TabFactory/AnimationTabSummoner.h b/Engine/Source/Editor/UMGEditor/Private/TabFactory/AnimationTabSummoner.h
new file mode 100644
--- /dev/null
+++ b/Engine/Source/Editor/UMGEditor/Private/TabFactory/AnimationTabSummoner.h
@@ -0,0 +1,17 @@
+// Copyright 1998-2014 Epic Games, Inc. All Rights Reserved.
+
+#pragma once
+
+struct FAnimationTabSummoner : public FWorkflowTabFactory
+{
+public:
+	static const FName TabID;
+	
+public:
+	FAnimationTabSummoner(TSharedPtr<class FWidgetBlueprintEditor> InBlueprintEditor);
+	
+	virtual TSharedRef<SWidget> CreateTabBody(const FWorkflowTabSpawnInfo& Info) const override;
+	
+protected:
+	TWeakPtr<class FWidgetBlueprintEditor> BlueprintEditor;
+};
diff --git a/Engine/Source/Editor/UMGEditor/Private/TabFactory/SequencerTabSummoner.cpp b/Engine/Source/Editor/UMGEditor/Private/TabFactory/SequencerTabSummoner.cpp
--- a/Engine/Source/Editor/UMGEditor/Private/TabFactory/SequencerTabSummoner.cpp
+++ b/Engine/Source/Editor/UMGEditor/Private/TabFactory/SequencerTabSummoner.cpp
@@ -15,7 +15,7 @@ FSequencerTabSummoner::FSequencerTabSummoner(TSharedPtr<class FWidgetBlueprintEd
 		, BlueprintEditor(InBlueprintEditor)
 {
 	TabLabel = LOCTEXT("SequencerLabel", "Timeline");
-	//TabIcon = FEditorStyle::GetBrush("Kismet.Tabs.Components");
+
 
 	bIsSingleton = true;
 
@@ -27,11 +27,8 @@ TSharedRef<SWidget> FSequencerTabSummoner::CreateTabBody(const FWorkflowTabSpawn
 {
 	TSharedPtr<FWidgetBlueprintEditor> BlueprintEditorPinned = BlueprintEditor.Pin();
 
-	return SNew(SBox)
-		.Tag(TEXT("Sequencer"))
-		[
-			BlueprintEditorPinned->GetSequencer()->GetSequencerWidget()
-		];
+	return BlueprintEditorPinned->CreateSequencerWidget();
+
 }
 
 #undef LOCTEXT_NAMESPACE 
diff --git a/Engine/Source/Editor/UMGEditor/Private/WidgetBlueprint.cpp b/Engine/Source/Editor/UMGEditor/Private/WidgetBlueprint.cpp
--- a/Engine/Source/Editor/UMGEditor/Private/WidgetBlueprint.cpp
+++ b/Engine/Source/Editor/UMGEditor/Private/WidgetBlueprint.cpp
@@ -2,6 +2,7 @@
 
 #include "UMGEditorPrivatePCH.h"
 #include "Runtime/MovieSceneCore/Classes/MovieScene.h"
+#include "PropertyTag.h"
 
 #define LOCTEXT_NAMESPACE "UMG"
 
@@ -23,7 +24,18 @@ FDelegateRuntimeBinding FDelegateEditorBinding::ToRuntimeBinding(UWidgetBlueprin
 	return Binding;
 }
 
+bool FWidgetAnimation_DEPRECATED::SerializeFromMismatchedTag(struct FPropertyTag const& Tag, FArchive& Ar)
+{
+	static FName AnimationDataName("AnimationData");
+	if(Tag.Type == NAME_StructProperty && Tag.Name == AnimationDataName)
+	{
+		Ar << MovieScene;
+		Ar << AnimationBindings;
+		return true;
+	}
 
+	return false;
+}
 /////////////////////////////////////////////////////
 // UWidgetBlueprint
 
@@ -45,6 +57,28 @@ void UWidgetBlueprint::PostLoad()
 		Widget->ConnectEditorData();
 	}
 
+	if( GetLinkerUE4Version() < VER_UE4_FIXUP_WIDGET_ANIMATION_CLASS )
+	{
+		// Fixup widget animiations.
+		for( auto& OldAnim : AnimationData_DEPRECATED )
+		{
+			FName AnimName = OldAnim.MovieScene->GetFName();
+
+			// Rename the old movie scene so we can reuse the name
+			OldAnim.MovieScene->Rename( *MakeUniqueObjectName( this, UMovieScene::StaticClass(), "MovieScene").ToString(), nullptr, REN_ForceNoResetLoaders | REN_DontCreateRedirectors | REN_DoNotDirty | REN_NonTransactional);
+
+			UWidgetAnimation* NewAnimation = ConstructObject<UWidgetAnimation>( UWidgetAnimation::StaticClass(), this, AnimName, RF_Transactional );
+
+			OldAnim.MovieScene->Rename(*AnimName.ToString(), NewAnimation, REN_ForceNoResetLoaders | REN_DontCreateRedirectors | REN_DoNotDirty | REN_NonTransactional );
+
+			NewAnimation->MovieScene = OldAnim.MovieScene;
+			NewAnimation->AnimationBindings = OldAnim.AnimationBindings;
+			
+			Animations.Add( NewAnimation );
+		}	
+
+		AnimationData_DEPRECATED.Empty();
+	}
 }
 
 UClass* UWidgetBlueprint::RegenerateClass(UClass* ClassToRegenerate, UObject* PreviousCDO, TArray<UObject*>& ObjLoaded)
@@ -64,11 +98,6 @@ bool UWidgetBlueprint::AllowsDynamicBinding() const
 	return true;
 }
 
-FWidgetAnimation* UWidgetBlueprint::FindAnimationDataForMovieScene( UMovieScene& MovieScene )
-{
-	return AnimationData.FindByPredicate( [&](const FWidgetAnimation& Data) { return Data.MovieScene == &MovieScene; } );
-}
-
 bool UWidgetBlueprint::ValidateGeneratedClass(const UClass* InClass)
 {
 	bool Result = Super::ValidateGeneratedClass(InClass);
diff --git a/Engine/Source/Editor/UMGEditor/Private/WidgetBlueprintCompiler.cpp b/Engine/Source/Editor/UMGEditor/Private/WidgetBlueprintCompiler.cpp
--- a/Engine/Source/Editor/UMGEditor/Private/WidgetBlueprintCompiler.cpp
+++ b/Engine/Source/Editor/UMGEditor/Private/WidgetBlueprintCompiler.cpp
@@ -228,7 +228,20 @@ void FWidgetBlueprintCompiler::CreateClassVariablesFromBlueprint()
 	}
 
 	// Add movie scenes variables here
+	for(UWidgetAnimation* Animation : Blueprint->Animations)
+	{
+		FString VariableName = Animation->GetName();
+
+		FEdGraphPinType WidgetPinType(Schema->PC_Object, TEXT(""), Animation->GetClass(), false, true);
+		UProperty* WidgetProperty = CreateVariable(FName(*VariableName), WidgetPinType);
 
+		if(WidgetProperty != NULL)
+		{
+			WidgetProperty->SetMetaData(TEXT("Category"), TEXT("Animations") );
+
+			WidgetProperty->SetPropertyFlags(CPF_BlueprintVisible);
+		}
+	}
 }
 
 void FWidgetBlueprintCompiler::FinishCompilingClass(UClass* Class)
@@ -239,16 +252,17 @@ void FWidgetBlueprintCompiler::FinishCompilingClass(UClass* Class)
 
 	BPGClass->WidgetTree = DuplicateObject<UWidgetTree>(Blueprint->WidgetTree, BPGClass);
 
-	BPGClass->AnimationData.Empty();
-	for ( const FWidgetAnimation& Animation : Blueprint->AnimationData )
-	{
-		UMovieScene* ClonedSequence = DuplicateObject<UMovieScene>(Animation.MovieScene, BPGClass, *Animation.MovieScene->GetName());
-		FWidgetAnimation CompiledAnimation;
+	BPGClass->Animations.Empty();
 
-		CompiledAnimation.MovieScene = ClonedSequence;
-		CompiledAnimation.AnimationBindings = Animation.AnimationBindings;
+	// Dont duplicate animation data on the skeleton class
+	if( Blueprint->SkeletonGeneratedClass != Class )
+	{
+		for(const UWidgetAnimation* Animation : Blueprint->Animations)
+		{
+			UWidgetAnimation* ClonedAnimation = DuplicateObject<UWidgetAnimation>(Animation, BPGClass, *(Animation->GetName()+TEXT("_INST")));
 
-		BPGClass->AnimationData.Add( CompiledAnimation );
+			BPGClass->Animations.Add(ClonedAnimation);
+		}
 	}
 
 	BPGClass->Bindings.Reset();
diff --git a/Engine/Source/Editor/UMGEditor/Private/WidgetBlueprintEditor.cpp b/Engine/Source/Editor/UMGEditor/Private/WidgetBlueprintEditor.cpp
--- a/Engine/Source/Editor/UMGEditor/Private/WidgetBlueprintEditor.cpp
+++ b/Engine/Source/Editor/UMGEditor/Private/WidgetBlueprintEditor.cpp
@@ -34,6 +34,10 @@ FWidgetBlueprintEditor::~FWidgetBlueprintEditor()
 	}
 
 	GEditor->OnObjectsReplaced().RemoveAll(this);
+	
+	Sequencer.Reset();
+
+	SequencerObjectBindingManager.Reset();
 }
 
 void FWidgetBlueprintEditor::InitWidgetBlueprintEditor(const EToolkitMode::Type Mode, const TSharedPtr< IToolkitHost >& InitToolkitHost, const TArray<UBlueprint*>& InBlueprints, bool bShouldOpenInDefaultsMode)
@@ -251,6 +255,27 @@ void FWidgetBlueprintEditor::CutSelectedWidgets()
 	FWidgetBlueprintEditorUtils::CutWidgets(GetWidgetBlueprintObj(), Widgets);
 }
 
+const UWidgetAnimation* FWidgetBlueprintEditor::RefreshCurrentAnimation()
+{
+	if( !SequencerObjectBindingManager->HasValidWidgetAnimation() )
+	{
+		const TArray<UWidgetAnimation*>& Animations = GetWidgetBlueprintObj()->Animations;
+		if( Animations.Num() > 0 )
+		{
+			// Ensure we are viewing a valid animation
+			ChangeViewedAnimation( *Animations[0] );
+			return Animations[0];
+		}
+		else
+		{
+			ChangeViewedAnimation( *UWidgetAnimation::GetNullAnimation() );
+			return nullptr;
+		}
+	}
+
+	return SequencerObjectBindingManager->GetWidgetAnimation();
+}
+
 bool FWidgetBlueprintEditor::CanPasteWidgets()
 {
 	TSet<FWidgetReference> Widgets = GetSelectedWidgets();
@@ -420,6 +445,37 @@ void FWidgetBlueprintEditor::MigrateFromChain(FEditPropertyChain* PropertyThatCh
 	}
 }
 
+void FWidgetBlueprintEditor::PostUndo(bool bSuccessful)
+{
+	FBlueprintEditor::PostUndo(bSuccessful);
+
+	OnWidgetBlueprintTransaction.Broadcast();
+}
+
+void FWidgetBlueprintEditor::PostRedo(bool bSuccessful)
+{
+	FBlueprintEditor::PostRedo(bSuccessful);
+
+	OnWidgetBlueprintTransaction.Broadcast();
+}
+
+TSharedRef<SWidget> FWidgetBlueprintEditor::CreateSequencerWidget()
+{
+	TSharedRef<SOverlay> SequencerOverlayRef =
+		SNew(SOverlay)
+		.Tag(TEXT("Sequencer"))
+		+ SOverlay::Slot()
+		[
+			GetSequencer()->GetSequencerWidget()
+		];
+
+	SequencerOverlay = SequencerOverlayRef;
+
+	RefreshCurrentAnimation();
+
+	return SequencerOverlayRef;
+}
+
 UWidgetBlueprint* FWidgetBlueprintEditor::GetWidgetBlueprintObj() const
 {
 	return Cast<UWidgetBlueprint>(GetBlueprintObj());
@@ -441,23 +497,21 @@ TSharedPtr<ISequencer>& FWidgetBlueprintEditor::GetSequencer()
 	{
 		UWidgetBlueprint* Blueprint = GetWidgetBlueprintObj();
 
-		FWidgetAnimation* WidgetAnimation = nullptr;
+		UWidgetAnimation* WidgetAnimation = nullptr;
 
-		if( Blueprint->AnimationData.Num() )
+		if( Blueprint->Animations.Num() )
 		{
-			WidgetAnimation = &Blueprint->AnimationData[0];
+			WidgetAnimation = Blueprint->Animations[0];
 		}
 		else
 		{
-			UMovieScene* MovieScene = ConstructObject<UMovieScene>( UMovieScene::StaticClass(), Blueprint, MakeUniqueObjectName( Blueprint, UMovieScene::StaticClass(), "DefaultAnimationData"), RF_Transactional );
-			FWidgetAnimation NewAnimation;
-			NewAnimation.MovieScene = MovieScene;
-			int32 Index = Blueprint->AnimationData.Add( NewAnimation );
+			UWidgetAnimation* NewAnimation = ConstructObject<UWidgetAnimation>(UWidgetAnimation::StaticClass(), Blueprint, MakeUniqueObjectName(Blueprint, UWidgetAnimation::StaticClass(), "NewAnimation"), RF_Transactional);
+			NewAnimation->MovieScene =  ConstructObject<UMovieScene>(UMovieScene::StaticClass(), NewAnimation, NAME_None, RF_Transactional);
 
-			WidgetAnimation = &Blueprint->AnimationData[Index];
+			WidgetAnimation = NewAnimation;
 		}
 
-		TSharedRef<FUMGSequencerObjectBindingManager> ObjectBindingManager = MakeShareable(new FUMGSequencerObjectBindingManager(*this, *WidgetAnimation->MovieScene) );
+		TSharedRef<FUMGSequencerObjectBindingManager> ObjectBindingManager = MakeShareable(new FUMGSequencerObjectBindingManager(*this, *WidgetAnimation ) );
 		check( !SequencerObjectBindingManager.IsValid() );
 
 		SequencerObjectBindingManager = ObjectBindingManager;
@@ -472,6 +526,47 @@ TSharedPtr<ISequencer>& FWidgetBlueprintEditor::GetSequencer()
 	return Sequencer;
 }
 
+void FWidgetBlueprintEditor::ChangeViewedAnimation( UWidgetAnimation& InAnimationToView )
+{
+	if( InAnimationToView.MovieScene != Sequencer->GetRootMovieScene() )
+	{
+		TSharedRef<FUMGSequencerObjectBindingManager> NewObjectBindingManager = MakeShareable(new FUMGSequencerObjectBindingManager(*this, InAnimationToView));
+
+		Sequencer->ResetToNewRootMovieScene(*InAnimationToView.MovieScene, NewObjectBindingManager);
+
+		check(SequencerObjectBindingManager.IsUnique());
+
+		SequencerObjectBindingManager = NewObjectBindingManager;
+
+		SequencerObjectBindingManager->InitPreviewObjects();
+
+		TSharedPtr<SOverlay> SequencerOverlayPin = SequencerOverlay.Pin();
+		if( &InAnimationToView == UWidgetAnimation::GetNullAnimation() && SequencerOverlayPin.IsValid() )
+		{
+			Sequencer->GetSequencerWidget()->SetEnabled(false);
+			// Disable sequencer from interaction
+			SequencerOverlayPin->AddSlot(1)
+			.HAlign(HAlign_Center)
+			.VAlign(VAlign_Center)
+			[
+				SNew( STextBlock )
+				.TextStyle( FEditorStyle::Get(), "UMGEditor.NoAnimationFont" )
+				.Text( LOCTEXT("NoAnimationSelected","No Animation Selected") )
+			];
+
+			SequencerOverlayPin->SetVisibility( EVisibility::HitTestInvisible );
+		}
+		else if( SequencerOverlayPin.IsValid() && SequencerOverlayPin->GetNumWidgets() > 1 )
+		{
+			Sequencer->GetSequencerWidget()->SetEnabled(true);
+
+			SequencerOverlayPin->RemoveSlot(1);
+			// Allow sequencer to be interacted with
+			SequencerOverlayPin->SetVisibility( EVisibility::SelfHitTestInvisible );
+		}
+	}
+}
+
 void FWidgetBlueprintEditor::DestroyPreview()
 {
 	UUserWidget* PreviewActor = GetPreview();
diff --git a/Engine/Source/Editor/UMGEditor/Private/WidgetBlueprintEditor.h b/Engine/Source/Editor/UMGEditor/Private/WidgetBlueprintEditor.h
--- a/Engine/Source/Editor/UMGEditor/Private/WidgetBlueprintEditor.h
+++ b/Engine/Source/Editor/UMGEditor/Private/WidgetBlueprintEditor.h
@@ -26,7 +26,10 @@ public:
 
 	void InitWidgetBlueprintEditor(const EToolkitMode::Type Mode, const TSharedPtr< IToolkitHost >& InitToolkitHost, const TArray<UBlueprint*>& InBlueprints, bool bShouldOpenInDefaultsMode);
 
+	/** FBlueprintEditor interface */
 	virtual void Tick(float DeltaTime) override;
+	virtual void PostUndo(bool bSuccessful) override;
+	virtual void PostRedo(bool bSuccessful) override;
 
 	/** FGCObjectInterface */
 	virtual void AddReferencedObjects( FReferenceCollector& Collector ) override;
@@ -40,6 +43,12 @@ public:
 	/** @return The sequencer used to create widget animations */
 	TSharedPtr<ISequencer>& GetSequencer();
 
+	/** Changes the currently viewed animation in Sequencer to the new one*/
+	void ChangeViewedAnimation( UWidgetAnimation& InAnimationToView );
+
+	/** Updates the current animation if it is invalid */
+	const UWidgetAnimation* RefreshCurrentAnimation();
+
 	/** Sets the currently selected set of widgets */
 	void SelectWidgets(const TSet<FWidgetReference>& Widgets);
 
@@ -57,6 +66,12 @@ public:
 	/** Migrate a property change from the preview GUI to the template GUI. */
 	void MigrateFromChain(FEditPropertyChain* PropertyThatChanged, bool bIsModify);
 
+	/** Event called when an undo/redo transaction occurs */
+	DECLARE_EVENT(FWidgetBlueprintEditor, FOnWidgetBlueprintTransaction)
+	FOnWidgetBlueprintTransaction& GetOnWidgetBlueprintTransaction() { return OnWidgetBlueprintTransaction; }
+
+	/** Creates a sequencer widget */
+	TSharedRef<SWidget> CreateSequencerWidget();
 public:
 	/** Fires whenever the selected set of widgets changing */
 	FOnSelectedWidgetsChanged OnSelectedWidgetsChanging;
@@ -101,12 +116,6 @@ private:
 
 	/** Tick the current preview GUI object */
 	void UpdatePreview(UBlueprint* InBlueprint, bool bInForceFullUpdate);
-	
-	/**
-	 * Gets the default movie scene which is used when there is no 
-	 * animation data present on the widget blueprint
-	 */
-	UMovieScene* GetDefaultMovieScene();
 
 private:
 	/** The preview scene that owns the preview GUI */
@@ -115,6 +124,9 @@ private:
 	/** Sequencer for creating and previewing widget animations */
 	TSharedPtr<ISequencer> Sequencer;
 
+	/** Overlay used to display UI on top of sequencer */
+	TWeakPtr<SOverlay> SequencerOverlay;
+
 	/** Manager for handling bindings to sequence animations */
 	TSharedPtr<class FUMGSequencerObjectBindingManager> SequencerObjectBindingManager;
 
@@ -130,6 +142,9 @@ private:
 	/** Notification for when the preview widget has been updated */
 	FOnWidgetPreviewUpdated OnWidgetPreviewUpdated;
 
+	/** Delegate called when a undo/redo transaction happens */
+	FOnWidgetBlueprintTransaction OnWidgetBlueprintTransaction;
+
 	/** The toolbar builder associated with this editor */
 	TSharedPtr<class FWidgetBlueprintEditorToolbar> WidgetToolbar;
 };
diff --git a/Engine/Source/Editor/UMGEditor/Private/WidgetBlueprintEditorUtils.cpp b/Engine/Source/Editor/UMGEditor/Private/WidgetBlueprintEditorUtils.cpp
--- a/Engine/Source/Editor/UMGEditor/Private/WidgetBlueprintEditorUtils.cpp
+++ b/Engine/Source/Editor/UMGEditor/Private/WidgetBlueprintEditorUtils.cpp
@@ -147,9 +147,9 @@ bool FWidgetBlueprintEditorUtils::RenameWidget(TSharedRef<FWidgetBlueprintEditor
 		}
 
 		// Update widget blueprint names
-		for( FWidgetAnimation& WidgetAnimation : Blueprint->AnimationData )
+		for( UWidgetAnimation* WidgetAnimation : Blueprint->Animations )
 		{
-			for( FWidgetAnimationBinding& AnimBinding : WidgetAnimation.AnimationBindings )
+			for( FWidgetAnimationBinding& AnimBinding : WidgetAnimation->AnimationBindings )
 			{
 				if( AnimBinding.WidgetName == OldName )
 				{
diff --git a/Engine/Source/Editor/UMGEditor/Public/WidgetBlueprint.h b/Engine/Source/Editor/UMGEditor/Public/WidgetBlueprint.h
--- a/Engine/Source/Editor/UMGEditor/Public/WidgetBlueprint.h
+++ b/Engine/Source/Editor/UMGEditor/Public/WidgetBlueprint.h
@@ -4,7 +4,7 @@
 
 #include "Engine/Blueprint.h"
 #include "UserWidget.h"
-#include "WidgetBlueprintGeneratedClass.h"
+#include "WidgetAnimation.h"
 
 #include "WidgetBlueprint.generated.h"
 
@@ -43,6 +43,33 @@ struct UMGEDITOR_API FDelegateEditorBinding
 	FDelegateRuntimeBinding ToRuntimeBinding(class UWidgetBlueprint* Blueprint) const;
 };
 
+
+/** Struct used only for loading old animations */
+USTRUCT()
+struct FWidgetAnimation_DEPRECATED
+{
+	GENERATED_USTRUCT_BODY()
+
+	UPROPERTY()
+	UMovieScene* MovieScene;
+
+	UPROPERTY()
+	TArray<FWidgetAnimationBinding> AnimationBindings;
+
+	bool SerializeFromMismatchedTag(struct FPropertyTag const& Tag, FArchive& Ar);
+
+};
+
+template<>
+struct TStructOpsTypeTraits<FWidgetAnimation_DEPRECATED> : public TStructOpsTypeTraitsBase
+{
+	enum
+	{
+		WithSerializeFromMismatchedTag = true,
+	};
+};
+
+
 /**
  * The widget blueprint enables extending UUserWidget the user extensible UWidget.
  */
@@ -60,7 +87,10 @@ public:
 	TArray< FDelegateEditorBinding > Bindings;
 
 	UPROPERTY()
-	TArray<FWidgetAnimation> AnimationData;
+	TArray<FWidgetAnimation_DEPRECATED> AnimationData_DEPRECATED;
+
+	UPROPERTY()
+	TArray<UWidgetAnimation*> Animations;
 
 	/** UObject interface */
 	virtual void PostLoad() override;
@@ -80,7 +110,5 @@ public:
 	virtual void GetReparentingRules(TSet< const UClass* >& AllowedChildrenOfClasses, TSet< const UClass* >& DisallowedChildrenOfClasses) const;
 	// End of UBlueprint interface
 
-	FWidgetAnimation* FindAnimationDataForMovieScene( UMovieScene& MovieScene );
-
 	static bool ValidateGeneratedClass(const UClass* InClass);
 };
\ No newline at end of file
diff --git a/Engine/Source/Runtime/Core/Public/UObject/ObjectVersion.h b/Engine/Source/Runtime/Core/Public/UObject/ObjectVersion.h
--- a/Engine/Source/Runtime/Core/Public/UObject/ObjectVersion.h
+++ b/Engine/Source/Runtime/Core/Public/UObject/ObjectVersion.h
@@ -594,6 +594,8 @@ enum EUnrealEngineObjectUE4Version
 	VER_UE4_ADDED_CURRENCY_CODE_TO_FTEXT,
 	// Added support for C++11 enum classes
 	VER_UE4_ENUM_CLASS_SUPPORT,
+	// Fixup widget animation class
+	VER_UE4_FIXUP_WIDGET_ANIMATION_CLASS,
 
 	// -----<new versions can be added before this line>-------------------------------------------------
 	// - this needs to be the last line (see note below)
diff --git a/Engine/Source/Runtime/UMG/Private/Animation/UMGSequencePlayer.cpp b/Engine/Source/Runtime/UMG/Private/Animation/UMGSequencePlayer.cpp
--- a/Engine/Source/Runtime/UMG/Private/Animation/UMGSequencePlayer.cpp
+++ b/Engine/Source/Runtime/UMG/Private/Animation/UMGSequencePlayer.cpp
@@ -6,18 +6,21 @@
 #include "MovieSceneBindings.h"
 #include "MovieSceneInstance.h"
 #include "MovieScene.h"
+#include "WidgetAnimation.h"
 
 UUMGSequencePlayer::UUMGSequencePlayer(const FPostConstructInitializeProperties& PCIP)
 	: Super(PCIP)
 {
 	PlayerStatus = EMovieScenePlayerStatus::Stopped;
 	TimeCursorPosition = 0.0f;
-	MovieScene = nullptr;
+	Animation = nullptr;
 }
 
-void UUMGSequencePlayer::InitSequencePlayer( const FWidgetAnimation& Animation, UUserWidget& UserWidget )
+void UUMGSequencePlayer::InitSequencePlayer( const UWidgetAnimation& InAnimation, UUserWidget& UserWidget )
 {
-	MovieScene = Animation.MovieScene;
+	Animation = &InAnimation;
+
+	UMovieScene* MovieScene = Animation->MovieScene;
 
 	// Cache the time range of the sequence to determine when we stop
 	TimeRange = MovieScene->GetTimeRange();
@@ -29,7 +32,7 @@ void UUMGSequencePlayer::InitSequencePlayer( const FWidgetAnimation& Animation,
 
 	TMap<FGuid, TArray<UObject*> > GuidToRuntimeObjectMap;
 	// Bind to Runtime Objects
-	for (const FWidgetAnimationBinding& Binding : Animation.AnimationBindings)
+	for (const FWidgetAnimationBinding& Binding : InAnimation.AnimationBindings)
 	{
 		FName ObjectName = Binding.WidgetName;
 		FName SlotWidgetName = Binding.SlotWidgetName;
@@ -96,6 +99,8 @@ void UUMGSequencePlayer::Tick(float DeltaTime)
 
 void UUMGSequencePlayer::Play()
 {
+	UMovieScene* MovieScene = Animation->MovieScene;
+
 	RootMovieSceneInstance = MakeShareable( new FMovieSceneInstance( *MovieScene ) );
 
 	RootMovieSceneInstance->RefreshInstance( *this );
diff --git a/Engine/Source/Runtime/UMG/Private/Animation/WidgetAnimation.cpp b/Engine/Source/Runtime/UMG/Private/Animation/WidgetAnimation.cpp
new file mode 100644
--- /dev/null
+++ b/Engine/Source/Runtime/UMG/Private/Animation/WidgetAnimation.cpp
@@ -0,0 +1,32 @@
+// Copyright 1998-2014 Epic Games, Inc. All Rights Reserved.
+
+#include "UMGPrivatePCH.h"
+#include "MovieScene.h"
+#include "WidgetAnimation.h"
+
+UWidgetAnimation::UWidgetAnimation(const FPostConstructInitializeProperties& PCIP)
+	: Super(PCIP)
+{
+
+}
+
+void UWidgetAnimation::PostDuplicate(bool bDuplicateForPIE)
+{
+	Super::PostDuplicate(bDuplicateForPIE);
+}
+
+#if WITH_EDITOR
+
+UWidgetAnimation* UWidgetAnimation::GetNullAnimation()
+{
+	static UWidgetAnimation* NullAnimation = nullptr;
+	if( !NullAnimation )
+	{
+		NullAnimation = ConstructObject<UWidgetAnimation>( UWidgetAnimation::StaticClass(), GetTransientPackage(), NAME_None, RF_RootSet );
+		NullAnimation->MovieScene = ConstructObject<UMovieScene>( UMovieScene::StaticClass(), NullAnimation, FName("No Animation"), RF_RootSet );
+	}
+
+	return NullAnimation;
+}
+
+#endif
\ No newline at end of file
diff --git a/Engine/Source/Runtime/UMG/Private/UserWidget.cpp b/Engine/Source/Runtime/UMG/Private/UserWidget.cpp
--- a/Engine/Source/Runtime/UMG/Private/UserWidget.cpp
+++ b/Engine/Source/Runtime/UMG/Private/UserWidget.cpp
@@ -4,6 +4,7 @@
 
 #include "UMGSequencePlayer.h"
 #include "SceneViewport.h"
+#include "WidgetAnimation.h"
 
 /**
  * This class holds onto the widget when it's placed into the viewport.
@@ -189,57 +190,48 @@ UWorld* UUserWidget::GetWorld() const
 	return CachedWorld;
 }
 
-void UUserWidget::PlayAnimation(FName AnimationName)
+void UUserWidget::PlayAnimation(const UWidgetAnimation* InAnimation)
 {
-	UWidgetBlueprintGeneratedClass* BPClass = Cast<UWidgetBlueprintGeneratedClass>(GetClass());
-	if (BPClass)
+	if( InAnimation )
 	{
-		const FWidgetAnimation* Animation = BPClass->FindAnimation(AnimationName);
-		if( Animation )
-		{
-			// @todo UMG sequencer - Restart animations which have had Play called on them?
-			UUMGSequencePlayer** FoundPlayer = ActiveSequencePlayers.FindByPredicate( [&](const UUMGSequencePlayer* Player) { return Player->GetMovieScene() == Animation->MovieScene; } );
+		// @todo UMG sequencer - Restart animations which have had Play called on them?
+		UUMGSequencePlayer** FoundPlayer = ActiveSequencePlayers.FindByPredicate( [&](const UUMGSequencePlayer* Player) { return Player->GetAnimation() == InAnimation; } );
 
-			if( !FoundPlayer )
-			{
-				UUMGSequencePlayer* NewPlayer = ConstructObject<UUMGSequencePlayer>( UUMGSequencePlayer::StaticClass(), this );
-				ActiveSequencePlayers.Add( NewPlayer );
+		if( !FoundPlayer )
+		{
+			UUMGSequencePlayer* NewPlayer = ConstructObject<UUMGSequencePlayer>( UUMGSequencePlayer::StaticClass(), this );
+			ActiveSequencePlayers.Add( NewPlayer );
 
-				NewPlayer->OnSequenceFinishedPlaying().AddUObject( this, &UUserWidget::OnAnimationFinishedPlaying );
+			NewPlayer->OnSequenceFinishedPlaying().AddUObject( this, &UUserWidget::OnAnimationFinishedPlaying );
 
-				NewPlayer->InitSequencePlayer( *Animation, *this );
+			NewPlayer->InitSequencePlayer( *InAnimation, *this );
 
-				NewPlayer->Play();
-			}
-			else
-			{
-				(*FoundPlayer)->Play();
-			}
+			NewPlayer->Play();
+		}
+		else
+		{
+			(*FoundPlayer)->Play();
 		}
 	}
 }
 
-void UUserWidget::StopAnimation(FName AnimationName)
+void UUserWidget::StopAnimation(const UWidgetAnimation* InAnimation)
 {
-	UWidgetBlueprintGeneratedClass* BPClass = Cast<UWidgetBlueprintGeneratedClass>(GetClass());
-	if(BPClass)
+	if(InAnimation)
 	{
-		const FWidgetAnimation* Animation = BPClass->FindAnimation(AnimationName);
-		if(Animation)
+		// @todo UMG sequencer - Restart animations which have had Play called on them?
+		UUMGSequencePlayer** FoundPlayer = ActiveSequencePlayers.FindByPredicate([&](const UUMGSequencePlayer* Player) { return Player->GetAnimation() == InAnimation; } );
+
+		if(FoundPlayer)
 		{
-			// @todo UMG sequencer - Restart animations which have had Play called on them?
-			UUMGSequencePlayer** FoundPlayer = ActiveSequencePlayers.FindByPredicate([&](const UUMGSequencePlayer* Player) { return Player->GetMovieScene() == Animation->MovieScene; });
-			if(FoundPlayer)
-			{
-				(*FoundPlayer)->Stop();
-			}
+			(*FoundPlayer)->Stop();
 		}
 	}
 }
 
 void UUserWidget::OnAnimationFinishedPlaying( UUMGSequencePlayer& Player )
 {
-	ActiveSequencePlayers.Remove( &Player );
+	StoppedSequencePlayers.Add( &Player );
 }
 
 UWidget* UUserWidget::GetWidgetHandle(TSharedRef<SWidget> InWidget)
@@ -317,6 +309,14 @@ void UUserWidget::NativeTick(const FGeometry& MyGeometry, float InDeltaTime )
 		Player->Tick( InDeltaTime );
 	}
 
+	// The process of ticking the players above can stop them so we remove them after all players have ticked
+	for( UUMGSequencePlayer* StoppedPlayer : StoppedSequencePlayers )
+	{
+		ActiveSequencePlayers.Remove( StoppedPlayer );	
+	}
+
+	StoppedSequencePlayers.Empty();
+
 	if ( !bDesignTime )
 	{
 		UWorld* World = GetWorld();
diff --git a/Engine/Source/Runtime/UMG/Private/WidgetBlueprintGeneratedClass.cpp b/Engine/Source/Runtime/UMG/Private/WidgetBlueprintGeneratedClass.cpp
--- a/Engine/Source/Runtime/UMG/Private/WidgetBlueprintGeneratedClass.cpp
+++ b/Engine/Source/Runtime/UMG/Private/WidgetBlueprintGeneratedClass.cpp
@@ -63,7 +63,19 @@ void UWidgetBlueprintGeneratedClass::InitializeWidget(UUserWidget* UserWidget) c
 		TArray<UWidget*> ClonedWidgets;
 		ClonedTree->GetAllWidgets(ClonedWidgets);
 
-		// @todo UMG Create Movie Scene vars here
+		for(UWidgetAnimation* Animation : Animations)
+		{
+			if( Animation->MovieScene )
+			{
+				// Find property with the same name as the template and assign the new widget to it.
+				UObjectPropertyBase* Prop = FindField<UObjectPropertyBase>(WidgetBlueprintClass, Animation->MovieScene->GetFName());
+				if(Prop)
+				{
+					Prop->SetObjectPropertyValue_InContainer(UserWidget, Animation);
+				}
+			}
+
+		}
 
 		for ( UWidget* Widget : ClonedWidgets )
 		{
@@ -146,9 +158,4 @@ void UWidgetBlueprintGeneratedClass::InitializeWidget(UUserWidget* UserWidget) c
 	}
 }
 
-const FWidgetAnimation* UWidgetBlueprintGeneratedClass::FindAnimation( FName AnimationName ) const
-{
-	return AnimationData.FindByPredicate( [&](const FWidgetAnimation& Animation) { return Animation.MovieScene->GetFName() == AnimationName; } );
-}
-
 #undef LOCTEXT_NAMESPACE
diff --git a/Engine/Source/Runtime/UMG/Public/Animation/UMGSequencePlayer.h b/Engine/Source/Runtime/UMG/Public/Animation/UMGSequencePlayer.h
--- a/Engine/Source/Runtime/UMG/Public/Animation/UMGSequencePlayer.h
+++ b/Engine/Source/Runtime/UMG/Public/Animation/UMGSequencePlayer.h
@@ -5,7 +5,7 @@
 #include "IMovieScenePlayer.h"
 #include "UMGSequencePlayer.generated.h"
 
-class UMovieScene;
+class UWidgetAnimation;
 class UMovieSceneBindings;
 
 UCLASS(transient)
@@ -14,7 +14,7 @@ class UMG_API UUMGSequencePlayer : public UObject, public IMovieScenePlayer
 	GENERATED_UCLASS_BODY()
 
 public:
-	void InitSequencePlayer( const FWidgetAnimation& Animation, UUserWidget& UserWidget );
+	void InitSequencePlayer( const UWidgetAnimation& InAnimation, UUserWidget& UserWidget );
 
 	/** Updates the running movie */
 	void Tick( float DeltaTime );
@@ -28,7 +28,8 @@ public:
 	/** Pauses a running animation */
 	void Pause();
 
-	const UMovieScene* GetMovieScene() const { return MovieScene; }
+	/** @return The current animation being played */
+	const UWidgetAnimation* GetAnimation() const { return Animation; }
 
 	/** IMovieScenePlayer interface */
 	virtual void GetRuntimeObjects( TSharedRef<FMovieSceneInstance> MovieSceneInstance, const FGuid& ObjectHandle, TArray< UObject* >& OutObjects ) const override;
@@ -42,9 +43,9 @@ public:
 	FOnSequenceFinishedPlaying& OnSequenceFinishedPlaying() { return OnSequenceFinishedPlayingEvent; }
 
 private:
-	/** Sequence being played */
+	/** Animation being played */
 	UPROPERTY()
-	UMovieScene* MovieScene;
+	const UWidgetAnimation* Animation;
 
 	/** Bindings to actual live objects */
 	UPROPERTY()
diff --git a/Engine/Source/Runtime/UMG/Public/Animation/WidgetAnimation.h b/Engine/Source/Runtime/UMG/Public/Animation/WidgetAnimation.h
new file mode 100644
--- /dev/null
+++ b/Engine/Source/Runtime/UMG/Public/Animation/WidgetAnimation.h
@@ -0,0 +1,59 @@
+// Copyright 1998-2014 Epic Games, Inc. All Rights Reserved.
+
+#pragma once
+
+#include "IMovieScenePlayer.h"
+#include "WidgetAnimation.generated.h"
+
+class UMovieScene;
+
+
+/** A single object bound to a umg sequence */
+USTRUCT()
+struct FWidgetAnimationBinding
+{
+	GENERATED_USTRUCT_BODY()
+
+	UPROPERTY()
+	FName WidgetName;
+
+	UPROPERTY()
+	FName SlotWidgetName;
+
+	UPROPERTY()
+	FGuid AnimationGuid;
+
+	bool operator==(const FWidgetAnimationBinding& Other) const
+	{
+		return WidgetName == Other.WidgetName && SlotWidgetName == Other.SlotWidgetName && AnimationGuid == Other.AnimationGuid;
+	}
+
+	friend FArchive& operator<<(FArchive& Ar, FWidgetAnimationBinding& Binding)
+	{
+		Ar << Binding.WidgetName;
+		Ar << Binding.SlotWidgetName;
+		Ar << Binding.AnimationGuid;
+		return Ar;
+	}
+};
+
+
+UCLASS(MinimalAPI)
+class UWidgetAnimation : public UObject
+{
+	GENERATED_UCLASS_BODY()
+
+public:
+	virtual void PostDuplicate(bool bDuplicateForPIE) override;
+
+#if WITH_EDITOR
+	static UMG_API UWidgetAnimation* GetNullAnimation();
+#endif
+
+public:
+	UPROPERTY()
+	UMovieScene* MovieScene;
+
+	UPROPERTY()
+	TArray<FWidgetAnimationBinding> AnimationBindings;
+};
diff --git a/Engine/Source/Runtime/UMG/Public/Blueprint/UserWidget.h b/Engine/Source/Runtime/UMG/Public/Blueprint/UserWidget.h
--- a/Engine/Source/Runtime/UMG/Public/Blueprint/UserWidget.h
+++ b/Engine/Source/Runtime/UMG/Public/Blueprint/UserWidget.h
@@ -3,7 +3,6 @@
 #pragma once
 
 #include "Geometry.h"
-
 #include "UserWidget.generated.h"
 
 static FGeometry NullGeometry;
@@ -11,6 +10,8 @@ static FSlateRect NullRect;
 static FSlateWindowElementList NullElementList;
 static FWidgetStyle NullStyle;
 
+class UWidgetAnimation;
+
 /**
  * The state passed into OnPaint that we can expose as a single painting structure to blueprints to
  * allow script code to override OnPaint behavior.
@@ -232,7 +233,7 @@ public:
 	 * @param The name of the animation to play
 	 */
 	UFUNCTION(BlueprintCallable, Category="User Interface|Animation")
-	void PlayAnimation(FName AnimationName);
+	void PlayAnimation(const UWidgetAnimation* InAnimation);
 
 	/**
 	 * Stops an already running animation in this widget
@@ -240,7 +241,7 @@ public:
 	 * @param The name of the animation to stop
 	 */
 	UFUNCTION(BlueprintCallable, Category="User Interface|Animation")
-	void StopAnimation(FName AnimationName);
+	void StopAnimation(const UWidgetAnimation* InAnimation);
 
 	/** Called when a sequence player is finished playing an animation */
 	void OnAnimationFinishedPlaying(UUMGSequencePlayer& Player );
@@ -324,6 +325,9 @@ public:
 	UPROPERTY(Transient)
 	TArray<UUMGSequencePlayer*> ActiveSequencePlayers;
 
+	/** List of sequence players to cache and clean up when safe */
+	UPROPERTY(Transient)
+	TArray<UUMGSequencePlayer*> StoppedSequencePlayers;
 protected:
 	virtual TSharedRef<SWidget> RebuildWidget() override;
 
diff --git a/Engine/Source/Runtime/UMG/Public/Blueprint/WidgetBlueprintGeneratedClass.h b/Engine/Source/Runtime/UMG/Public/Blueprint/WidgetBlueprintGeneratedClass.h
--- a/Engine/Source/Runtime/UMG/Public/Blueprint/WidgetBlueprintGeneratedClass.h
+++ b/Engine/Source/Runtime/UMG/Public/Blueprint/WidgetBlueprintGeneratedClass.h
@@ -40,39 +40,6 @@ struct FDelegateRuntimeBinding
 	}*/
 };
 
-/** A single object bound to a umg sequence */
-USTRUCT()
-struct FWidgetAnimationBinding
-{
-	GENERATED_USTRUCT_BODY()
-
-	UPROPERTY()
-	FName WidgetName;
-
-	UPROPERTY()
-	FName SlotWidgetName;
-
-	UPROPERTY()
-	FGuid AnimationGuid;
-
-	bool operator==( const FWidgetAnimationBinding& Other ) const
-	{
-		return WidgetName == Other.WidgetName && SlotWidgetName == Other.SlotWidgetName && AnimationGuid == Other.AnimationGuid;
-	}
-};
-
-/** Runtime animation bindings for a single movie scene */
-USTRUCT()
-struct FWidgetAnimation
-{
-	GENERATED_USTRUCT_BODY()
-
-	UPROPERTY()
-	class UMovieScene* MovieScene;
-
-	UPROPERTY()
-	TArray<FWidgetAnimationBinding> AnimationBindings;
-};
 
 /**
  * The widget blueprint generated class allows us to create blueprintable widgets for UMG at runtime.
@@ -93,7 +60,7 @@ class UMG_API UWidgetBlueprintGeneratedClass : public UBlueprintGeneratedClass
 	TArray< FDelegateRuntimeBinding > Bindings;
 
 	UPROPERTY()
-	TArray<FWidgetAnimation> AnimationData;
+	TArray<UWidgetAnimation*> Animations;
 
 	/** This is transient data calculated at link time. */
 	TArray<UStructProperty*> WidgetNodeProperties;
@@ -108,12 +75,4 @@ class UMG_API UWidgetBlueprintGeneratedClass : public UBlueprintGeneratedClass
 	 * binding and wiring nessesary to have the user's widget perform as desired.
 	 */
 	void InitializeWidget(class UUserWidget* UserWidget) const;
-
-	/** 
-	 * Finds an animation by name of the Sequence
-	 * 
-	 * @param AnimationName	Name of the animation
-	 * @return Found animation or nullptr if one is not found
-	 */
-	const FWidgetAnimation* FindAnimation( FName AnimationName ) const;
 };
\ No newline at end of file
