#UE4
Fix Build

diff --git a/Engine/Source/Editor/UMGEditor/Private/TabFactory/AnimationTabSummoner.cpp b/Engine/Source/Editor/UMGEditor/Private/TabFactory/AnimationTabSummoner.cpp
--- a/Engine/Source/Editor/UMGEditor/Private/TabFactory/AnimationTabSummoner.cpp
+++ b/Engine/Source/Editor/UMGEditor/Private/TabFactory/AnimationTabSummoner.cpp
@@ -5,6 +5,8 @@
 #include "AnimationTabSummoner.h"
 #include "WidgetBlueprintEditor.h"
 #include "ISequencer.h"
+#include "WidgetAnimation.h"
+#include "MovieScene.h"
 
 #define LOCTEXT_NAMESPACE "UMG"
 
diff --git a/Engine/Source/Runtime/UMG/Private/WidgetBlueprintGeneratedClass.cpp b/Engine/Source/Runtime/UMG/Private/WidgetBlueprintGeneratedClass.cpp
--- a/Engine/Source/Runtime/UMG/Private/WidgetBlueprintGeneratedClass.cpp
+++ b/Engine/Source/Runtime/UMG/Private/WidgetBlueprintGeneratedClass.cpp
@@ -2,6 +2,7 @@
 
 #include "UMGPrivatePCH.h"
 #include "MovieScene.h"
+#include "WidgetAnimation.h"
 
 #define LOCTEXT_NAMESPACE "UMG"
 
