UMG - Still experimental but removing the experimental flag and always loading UMG.

diff --git a/Engine/Source/Editor/Sequencer/Private/Sequencer.cpp b/Engine/Source/Editor/Sequencer/Private/Sequencer.cpp
--- a/Engine/Source/Editor/Sequencer/Private/Sequencer.cpp
+++ b/Engine/Source/Editor/Sequencer/Private/Sequencer.cpp
@@ -39,12 +39,7 @@ DEFINE_LOG_CATEGORY(LogSequencer);
 
 bool FSequencer::IsSequencerEnabled()
 {
-	if ( GetDefault<UEditorExperimentalSettings>()->bUnrealMotionGraphics )
-	{
-		return true;
-	}
-
-	return FParse::Param(FCommandLine::Get(), TEXT("Sequencer"));
+	return true;
 }
 
 
diff --git a/Engine/Source/Editor/UMGEditor/Private/UMGEditorModule.cpp b/Engine/Source/Editor/UMGEditor/Private/UMGEditorModule.cpp
--- a/Engine/Source/Editor/UMGEditor/Private/UMGEditorModule.cpp
+++ b/Engine/Source/Editor/UMGEditor/Private/UMGEditorModule.cpp
@@ -37,20 +37,13 @@ public:
 		sd.BindRaw(this, &FUMGEditorModule::CompileWidgetBlueprint);
 		KismetCompilerModule.GetCompilers().Add(sd);
 
-		bUMGEnabled = GetDefault<UEditorExperimentalSettings>()->bUnrealMotionGraphics;
+		// Register asset types
+		IAssetTools& AssetTools = FModuleManager::LoadModuleChecked<FAssetToolsModule>("AssetTools").Get();
+		RegisterAssetTypeAction(AssetTools, MakeShareable(new FAssetTypeActions_WidgetBlueprint()));
 
-		// In the event we were loaded by an uncooked game, we need to check if we actually should be enabled.
-		// Because the UMG runtime will attempt to start our module no matter what, if WITH_EDITOR is enabled.
-		if ( bUMGEnabled )
-		{
-			// Register asset types
-			IAssetTools& AssetTools = FModuleManager::LoadModuleChecked<FAssetToolsModule>("AssetTools").Get();
-			RegisterAssetTypeAction(AssetTools, MakeShareable(new FAssetTypeActions_WidgetBlueprint()));
-
-			// Register with the sequencer module that we provide auto-key handlers.
-			ISequencerModule& SequencerModule = FModuleManager::Get().LoadModuleChecked<ISequencerModule>("Sequencer");
-			SequencerModule.RegisterTrackEditor(FOnCreateTrackEditor::CreateStatic(&FMarginTrackEditor::CreateTrackEditor));
-		}
+		// Register with the sequencer module that we provide auto-key handlers.
+		ISequencerModule& SequencerModule = FModuleManager::Get().LoadModuleChecked<ISequencerModule>("Sequencer");
+		SequencerModule.RegisterTrackEditor(FOnCreateTrackEditor::CreateStatic(&FMarginTrackEditor::CreateTrackEditor));
 	}
 
 	/** Called before the module is unloaded, right before the module object is destroyed. */
@@ -59,19 +52,16 @@ public:
 		MenuExtensibilityManager.Reset();
 		ToolBarExtensibilityManager.Reset();
 
-		if ( bUMGEnabled )
+		// Unregister all the asset types that we registered
+		if ( FModuleManager::Get().IsModuleLoaded("AssetTools") )
 		{
-			// Unregister all the asset types that we registered
-			if ( FModuleManager::Get().IsModuleLoaded("AssetTools") )
+			IAssetTools& AssetTools = FModuleManager::GetModuleChecked<FAssetToolsModule>("AssetTools").Get();
+			for ( int32 Index = 0; Index < CreatedAssetTypeActions.Num(); ++Index )
 			{
-				IAssetTools& AssetTools = FModuleManager::GetModuleChecked<FAssetToolsModule>("AssetTools").Get();
-				for ( int32 Index = 0; Index < CreatedAssetTypeActions.Num(); ++Index )
-				{
-					AssetTools.UnregisterAssetTypeActions(CreatedAssetTypeActions[Index].ToSharedRef());
-				}
+				AssetTools.UnregisterAssetTypeActions(CreatedAssetTypeActions[Index].ToSharedRef());
 			}
-			CreatedAssetTypeActions.Empty();
 		}
+		CreatedAssetTypeActions.Empty();
 	}
 
 	FReply CompileWidgetBlueprint(UBlueprint* Blueprint, const FKismetCompilerOptions& CompileOptions, FCompilerResultsLog& Results, TArray<UObject*>* ObjLoaded)
@@ -103,9 +93,6 @@ private:
 	TSharedPtr<FExtensibilityManager> MenuExtensibilityManager;
 	TSharedPtr<FExtensibilityManager> ToolBarExtensibilityManager;
 
-	/** We cache the value for UMG being enabled on load so that we can shutdown properly. */
-	bool bUMGEnabled;
-
 	/** All created asset type actions.  Cached here so that we can unregister it during shutdown. */
 	TArray< TSharedPtr<IAssetTypeActions> > CreatedAssetTypeActions;
 };
diff --git a/Engine/Source/Editor/UMGEditor/Private/WidgetBlueprintFactory.cpp b/Engine/Source/Editor/UMGEditor/Private/WidgetBlueprintFactory.cpp
--- a/Engine/Source/Editor/UMGEditor/Private/WidgetBlueprintFactory.cpp
+++ b/Engine/Source/Editor/UMGEditor/Private/WidgetBlueprintFactory.cpp
@@ -28,7 +28,7 @@ bool UWidgetBlueprintFactory::ConfigureProperties()
 
 bool UWidgetBlueprintFactory::ShouldShowInNewMenu() const
 {
-	return GetDefault<UEditorExperimentalSettings>()->bUnrealMotionGraphics;
+	return true;
 }
 
 UObject* UWidgetBlueprintFactory::FactoryCreateNew(UClass* Class, UObject* InParent, FName Name, EObjectFlags Flags, UObject* Context, FFeedbackContext* Warn, FName CallingContext)
diff --git a/Engine/Source/Editor/UnrealEd/Classes/Settings/EditorExperimentalSettings.h b/Engine/Source/Editor/UnrealEd/Classes/Settings/EditorExperimentalSettings.h
--- a/Engine/Source/Editor/UnrealEd/Classes/Settings/EditorExperimentalSettings.h
+++ b/Engine/Source/Editor/UnrealEd/Classes/Settings/EditorExperimentalSettings.h
@@ -21,9 +21,6 @@ class UNREALED_API UEditorExperimentalSettings
 
 public:
 
-	UPROPERTY(EditAnywhere, config, Category = Tools, meta = ( DisplayName = "Unreal Motion Graphics (UMG)" ))
-	bool bUnrealMotionGraphics;
-
 	/** Allows usage of the Translation Editor */
 	UPROPERTY(EditAnywhere, config, Category = Tools, meta = (DisplayName = "Translation Editor"))
 	bool bEnableTranslationEditor;
