Reset UMG layout

diff --git a/Engine/Source/Editor/UMGEditor/Private/BlueprintModes/WidgetDesignerApplicationMode.cpp b/Engine/Source/Editor/UMGEditor/Private/BlueprintModes/WidgetDesignerApplicationMode.cpp
--- a/Engine/Source/Editor/UMGEditor/Private/BlueprintModes/WidgetDesignerApplicationMode.cpp
+++ b/Engine/Source/Editor/UMGEditor/Private/BlueprintModes/WidgetDesignerApplicationMode.cpp
@@ -28,7 +28,7 @@
 FWidgetDesignerApplicationMode::FWidgetDesignerApplicationMode(TSharedPtr<FWidgetBlueprintEditor> InWidgetEditor)
 	: FWidgetBlueprintApplicationMode(InWidgetEditor, FWidgetBlueprintApplicationModes::DesignerMode)
 {
-	TabLayout = FTabManager::NewLayout( "WidgetBlueprintEditor_Designer_Layout_v2" )
+	TabLayout = FTabManager::NewLayout( "WidgetBlueprintEditor_Designer_Layout_v3" )
 	->AddArea
 	(
 		FTabManager::NewPrimaryArea()
