UMG - Exposing more information dealing with keys allowing to actually compare between keys that were pressed in the blueprint.  Several changes to the hierarchy view to allow showing the Root of the hierarchy which when selected shows the defaults for the widget blueprint, letting you change defaults about the widget.

diff --git a/Engine/Source/Editor/UMGEditor/Private/Designer/SDesignerView.cpp b/Engine/Source/Editor/UMGEditor/Private/Designer/SDesignerView.cpp
--- a/Engine/Source/Editor/UMGEditor/Private/Designer/SDesignerView.cpp
+++ b/Engine/Source/Editor/UMGEditor/Private/Designer/SDesignerView.cpp
@@ -851,8 +851,8 @@ SDesignerView::DragHandle SDesignerView::HitTestDragHandles(const FGeometry& All
 		FVector2D WidgetLocalPosition = AllottedGeometry.AbsoluteToLocal(ArrangedWidget.Geometry.AbsolutePosition);
 		float X = WidgetLocalPosition.X;
 		float Y = WidgetLocalPosition.Y;
-		float Width = ArrangedWidget.Geometry.Size.X;
-		float Height = ArrangedWidget.Geometry.Size.Y;
+		float Width = ArrangedWidget.Geometry.Size.X * GetZoomAmount();
+		float Height = ArrangedWidget.Geometry.Size.Y * GetZoomAmount();
 
 		TArray<FVector2D> Handles;
 		Handles.Add(FVector2D(X, Y));					// Top - Left
@@ -923,7 +923,8 @@ void SDesignerView::UpdatePreviewWidget()
 			const bool bAbsoluteLayout = false;
 			const bool bModal = false;
 			const bool bShowCursor = false;
-			TSharedRef<SWidget> CurrentWidget = PreviewWidget->MakeViewportWidget(bAbsoluteLayout, bModal, bShowCursor);
+			TSharedPtr<SWidget> OutUserWidget;
+			TSharedRef<SWidget> CurrentWidget = PreviewWidget->MakeViewportWidget(bAbsoluteLayout, bModal, bShowCursor, OutUserWidget);
 			CurrentWidget->SlatePrepass();
 
 			if ( CurrentWidget != PreviewSlateWidget.Pin() )
diff --git a/Engine/Source/Editor/UMGEditor/Private/Details/SWidgetDetailsView.cpp b/Engine/Source/Editor/UMGEditor/Private/Details/SWidgetDetailsView.cpp
--- a/Engine/Source/Editor/UMGEditor/Private/Details/SWidgetDetailsView.cpp
+++ b/Engine/Source/Editor/UMGEditor/Private/Details/SWidgetDetailsView.cpp
@@ -123,12 +123,25 @@ void SWidgetDetailsView::OnEditorSelectionChanging()
 
 void SWidgetDetailsView::OnEditorSelectionChanged()
 {
+	SelectedObjects.Empty();
+	PropertyView->SetObjects(SelectedObjects);
+
 	TSet< FWidgetReference > SelectedWidgets = BlueprintEditor.Pin()->GetSelectedWidgets();
+	if ( SelectedWidgets.Num() > 0 )
+	{
+		for ( FWidgetReference& WidgetRef : SelectedWidgets )
+		{
+			SelectedObjects.Add(WidgetRef.GetPreview());
+		}
+	}
 
-	SelectedObjects.Empty();
-	for ( FWidgetReference& WidgetRef : SelectedWidgets )
+	TSet< TWeakObjectPtr<UObject> > Selection = BlueprintEditor.Pin()->GetSelectedObjects();
+	for ( TWeakObjectPtr<UObject> Selected : Selection )
 	{
-		SelectedObjects.Add(WidgetRef.GetPreview());
+		if ( UObject* S = Selected.Get() )
+		{
+			SelectedObjects.Add(S);
+		}
 	}
 
 	const bool bForceRefresh = false;
@@ -155,7 +168,11 @@ EVisibility SWidgetDetailsView::GetNameAreaVisibility() const
 {
 	if ( SelectedObjects.Num() == 1 )
 	{
-		return EVisibility::Visible;
+		UWidget* Widget = Cast<UWidget>(SelectedObjects[0].Get());
+		if ( Widget && !Widget->HasAnyFlags(RF_ClassDefaultObject) )
+		{
+			return EVisibility::Visible;
+		}
 	}
 	
 	return EVisibility::Collapsed;
diff --git a/Engine/Source/Editor/UMGEditor/Private/Hierarchy/SHierarchyView.cpp b/Engine/Source/Editor/UMGEditor/Private/Hierarchy/SHierarchyView.cpp
--- a/Engine/Source/Editor/UMGEditor/Private/Hierarchy/SHierarchyView.cpp
+++ b/Engine/Source/Editor/UMGEditor/Private/Hierarchy/SHierarchyView.cpp
@@ -31,10 +31,10 @@ void SHierarchyView::Construct(const FArguments& InArgs, TSharedPtr<FWidgetBluep
 	UWidgetBlueprint* Blueprint = GetBlueprint();
 	Blueprint->OnChanged().AddRaw(this, &SHierarchyView::OnBlueprintChanged);
 
-	FilterHandler = MakeShareable(new TreeFilterHandler<UWidget*>());
+	FilterHandler = MakeShareable(new TreeFilterHandler< TSharedPtr<FHierarchyModel> >());
 	FilterHandler->SetFilter(SearchBoxWidgetFilter.Get());
 	FilterHandler->SetRootItems(&RootWidgets, &TreeRootWidgets);
-	FilterHandler->SetGetChildrenDelegate(TreeFilterHandler<UWidget*>::FOnGetChildren::CreateRaw(this, &SHierarchyView::WidgetHierarchy_OnGetChildren));
+	FilterHandler->SetGetChildrenDelegate(TreeFilterHandler< TSharedPtr<FHierarchyModel> >::FOnGetChildren::CreateRaw(this, &SHierarchyView::WidgetHierarchy_OnGetChildren));
 
 	ChildSlot
 	[
@@ -113,9 +113,9 @@ FReply SHierarchyView::OnKeyDown(const FGeometry& MyGeometry, const FKeyboardEve
 	return FReply::Unhandled();
 }
 
-void SHierarchyView::TransformWidgetToString(UWidget* Widget, OUT TArray< FString >& Array)
+void SHierarchyView::TransformWidgetToString(TSharedPtr<FHierarchyModel> Item, OUT TArray< FString >& Array)
 {
-	Array.Add( Widget->GetLabel() );
+	Array.Add( Item->GetText().ToString() );
 }
 
 void SHierarchyView::OnSearchChanged(const FText& InFilterText)
@@ -144,28 +144,43 @@ void SHierarchyView::OnEditorSelectionChanged()
 
 		if ( TemplateWidget )
 		{
-			WidgetTreeView->SetItemSelection(TemplateWidget, true);
+			TArray< UWidget* > Widgets;
+			Widgets.Add(TemplateWidget);
 
-			// Attempt to scroll the first widget we find into view.
-			if ( bFirst )
+			UWidget* Parent = TemplateWidget->GetParent();
+			while ( Parent != NULL )
 			{
-				bFirst = false;
-				WidgetTreeView->RequestScrollIntoView(TemplateWidget);
+				Widgets.Add(Parent);
+				Parent = Parent->GetParent();
 			}
 
-			ExpandPathToWidget(TemplateWidget);
-		}
-	}
-}
+			TSharedPtr<FHierarchyModel> CurrentItem = RootWidgets[0];
+			WidgetTreeView->SetItemExpansion(CurrentItem, true);
 
-void SHierarchyView::ExpandPathToWidget(UWidget* TemplateWidget)
-{
-	// Expand the path leading to this widget in the tree.
-	UWidget* Parent = TemplateWidget->GetParent();
-	while ( Parent != NULL )
-	{
-		WidgetTreeView->SetItemExpansion(Parent, true);
-		Parent = Parent->GetParent();
+			for ( int32 WidgetIndex = Widgets.Num() - 1; WidgetIndex >= 0; WidgetIndex-- )
+			{
+				UWidget* WidgetItem = Widgets[WidgetIndex];
+
+				TArray< TSharedPtr<FHierarchyModel> > Children;
+				CurrentItem->GatherChildren(Children);
+
+				for ( auto C : Children )
+				{
+					if ( C->IsModel(WidgetItem) )
+					{
+						WidgetTreeView->SetItemExpansion(C, true);
+						CurrentItem = C;
+
+						if ( WidgetIndex == 0 )
+						{
+							WidgetTreeView->SetItemSelection(CurrentItem, true, ESelectInfo::Direct);
+						}
+					}
+				}
+			}
+
+			WidgetTreeView->RequestScrollIntoView(CurrentItem);
+		}
 	}
 }
 
@@ -188,18 +203,6 @@ void SHierarchyView::OnBlueprintChanged(UBlueprint* InBlueprint)
 	}
 }
 
-void SHierarchyView::ShowDetailsForObjects(TArray<UWidget*> TemplateWidgets)
-{
-	TSet<FWidgetReference> SelectedWidgets;
-	for ( UWidget* TemplateWidget : TemplateWidgets )
-	{
-		FWidgetReference Selection = FWidgetReference::FromTemplate(BlueprintEditor.Pin(), TemplateWidget);
-		SelectedWidgets.Add(Selection);
-	}
-
-	BlueprintEditor.Pin()->SelectWidgets(SelectedWidgets);
-}
-
 TSharedPtr<SWidget> SHierarchyView::WidgetHierarchy_OnContextMenuOpening()
 {
 	FMenuBuilder MenuBuilder(true, NULL);
@@ -209,39 +212,24 @@ TSharedPtr<SWidget> SHierarchyView::WidgetHierarchy_OnContextMenuOpening()
 	return MenuBuilder.MakeWidget();
 }
 
-void SHierarchyView::WidgetHierarchy_OnGetChildren(UWidget* InParent, TArray< UWidget* >& OutChildren)
+void SHierarchyView::WidgetHierarchy_OnGetChildren(TSharedPtr<FHierarchyModel> InParent, TArray< TSharedPtr<FHierarchyModel> >& OutChildren)
 {
 	VisibleItems.Add(InParent);
 
-	UPanelWidget* Widget = Cast<UPanelWidget>(InParent);
-	if ( Widget )
-	{
-		for ( int32 i = 0; i < Widget->GetChildrenCount(); i++ )
-		{
-			UWidget* Child = Widget->GetChildAt(i);
-			if ( Child )
-			{
-				OutChildren.Add(Child);
-			}
-		}
-	}
+	InParent->GatherChildren(OutChildren);
 }
 
-TSharedRef< ITableRow > SHierarchyView::WidgetHierarchy_OnGenerateRow(UWidget* InItem, const TSharedRef<STableViewBase>& OwnerTable)
+TSharedRef< ITableRow > SHierarchyView::WidgetHierarchy_OnGenerateRow(TSharedPtr<FHierarchyModel> InItem, const TSharedRef<STableViewBase>& OwnerTable)
 {
-	FWidgetReference WidgetRef = FWidgetReference::FromTemplate(BlueprintEditor.Pin(), InItem);
-
-	return SNew(SHierarchyViewItem, OwnerTable, BlueprintEditor.Pin(), WidgetRef)
+	return SNew(SHierarchyViewItem, OwnerTable, InItem)
 		.HighlightText(this, &SHierarchyView::GetSearchText);
 }
 
-void SHierarchyView::WidgetHierarchy_OnSelectionChanged(UWidget* SelectedItem, ESelectInfo::Type SelectInfo)
+void SHierarchyView::WidgetHierarchy_OnSelectionChanged(TSharedPtr<FHierarchyModel> SelectedItem, ESelectInfo::Type SelectInfo)
 {
 	if ( SelectInfo != ESelectInfo::Direct )
 	{
-		TArray<UWidget*> Items;
-		Items.Add(SelectedItem);
-		ShowDetailsForObjects(Items);
+		SelectedItem->OnSelection();
 	}
 }
 
@@ -252,7 +240,7 @@ FReply SHierarchyView::HandleDeleteSelected()
 	// Remove the selected items from the filter cache
 	for (FWidgetReference& Item : SelectedWidgets)
 	{
-		FilterHandler->RemoveCachedItem(Item.GetTemplate());
+//		FilterHandler->RemoveCachedItem(Item.GetTemplate());
 	}
 
 	FWidgetBlueprintEditorUtils::DeleteWidgets(GetBlueprint(), SelectedWidgets);
@@ -265,20 +253,17 @@ void SHierarchyView::RefreshTree()
 	VisibleItems.Empty();
 
 	RootWidgets.Empty();
-	UWidgetBlueprint* Blueprint = GetBlueprint();
-	if (Blueprint->WidgetTree->RootWidget)
-	{
-		RootWidgets.Add(Blueprint->WidgetTree->RootWidget);
-	}
+	RootWidgets.Add( MakeShareable(new FHierarchyRoot(BlueprintEditor.Pin())) );
+
 	FilterHandler->RefreshAndFilterTree();
 }
 
 void SHierarchyView::RebuildTreeView()
 {
-	SAssignNew(WidgetTreeView, STreeView< UWidget* >)
+	SAssignNew(WidgetTreeView, STreeView< TSharedPtr<FHierarchyModel> >)
 		.ItemHeight(20.0f)
 		.SelectionMode(ESelectionMode::Single)
-		.OnGetChildren(FilterHandler.ToSharedRef(), &TreeFilterHandler<UWidget*>::OnGetFilteredChildren)
+		.OnGetChildren(FilterHandler.ToSharedRef(), &TreeFilterHandler< TSharedPtr<FHierarchyModel> >::OnGetFilteredChildren)
 		.OnGenerateRow(this, &SHierarchyView::WidgetHierarchy_OnGenerateRow)
 		.OnSelectionChanged(this, &SHierarchyView::WidgetHierarchy_OnSelectionChanged)
 		.OnContextMenuOpening(this, &SHierarchyView::WidgetHierarchy_OnContextMenuOpening)
@@ -297,9 +282,9 @@ void SHierarchyView::OnObjectsReplaced(const TMap<UObject*, UObject*>& Replaceme
 {
 	if ( !bRebuildTreeRequested )
 	{
-		for ( UWidget* Widget : VisibleItems )
+		for ( TSharedPtr<FHierarchyModel> Widget : VisibleItems )
 		{
-			if ( ReplacementMap.Contains(Widget) )
+			//if ( ReplacementMap.Contains(Widget) )
 			{
 				bRefreshRequested = true;
 				bRebuildTreeRequested = true;
diff --git a/Engine/Source/Editor/UMGEditor/Private/Hierarchy/SHierarchyView.h b/Engine/Source/Editor/UMGEditor/Private/Hierarchy/SHierarchyView.h
--- a/Engine/Source/Editor/UMGEditor/Private/Hierarchy/SHierarchyView.h
+++ b/Engine/Source/Editor/UMGEditor/Private/Hierarchy/SHierarchyView.h
@@ -5,6 +5,7 @@
 #include "SCompoundWidget.h"
 #include "BlueprintEditor.h"
 #include "TreeFilterHandler.h"
+#include "SHierarchyViewItem.h"
 
 //TODO rename SUMGEditorHierarchy
 
@@ -15,7 +16,7 @@
 class SHierarchyView : public SCompoundWidget
 {
 public:
-	typedef TTextFilter< UWidget* > WidgetTextFilter;
+	typedef TTextFilter< TSharedPtr<FHierarchyModel> > WidgetTextFilter;
 
 public:
 	SLATE_BEGIN_ARGS( SHierarchyView ){}
@@ -34,17 +35,14 @@ public:
 private:
 	void BuildWrapWithMenu(FMenuBuilder& Menu);
 	TSharedPtr<SWidget> WidgetHierarchy_OnContextMenuOpening();
-	void WidgetHierarchy_OnGetChildren(UWidget* InParent, TArray< UWidget* >& OutChildren);
-	TSharedRef< ITableRow > WidgetHierarchy_OnGenerateRow(UWidget* InItem, const TSharedRef<STableViewBase>& OwnerTable);
-	void WidgetHierarchy_OnSelectionChanged(UWidget* SelectedItem, ESelectInfo::Type SelectInfo);
+	void WidgetHierarchy_OnGetChildren(TSharedPtr<FHierarchyModel> InParent, TArray< TSharedPtr<FHierarchyModel> >& OutChildren);
+	TSharedRef< ITableRow > WidgetHierarchy_OnGenerateRow(TSharedPtr<FHierarchyModel> InItem, const TSharedRef<STableViewBase>& OwnerTable);
+	void WidgetHierarchy_OnSelectionChanged(TSharedPtr<FHierarchyModel> SelectedItem, ESelectInfo::Type SelectInfo);
 
 private:
 	/** @returns the current blueprint being edited */
 	UWidgetBlueprint* GetBlueprint() const;
 
-	/** Expands every item in the tree leading to this widget */
-	void ExpandPathToWidget(UWidget* TemplateWidget);
-
 	/** Called when the blueprint is structurally changed. */
 	void OnBlueprintChanged(UBlueprint* InBlueprint);
 
@@ -76,7 +74,7 @@ private:
 	FText GetSearchText() const;
 
 	/** Transforms the widget into a searchable string */
-	void TransformWidgetToString(UWidget* Widget, OUT TArray< FString >& Array);
+	void TransformWidgetToString(TSharedPtr<FHierarchyModel> Widget, OUT TArray< FString >& Array);
 
 	/** Called when a Blueprint is recompiled and live objects are swapped out for replacements */
 	void OnObjectsReplaced(const TMap<UObject*, UObject*>& ReplacementMap);
@@ -87,22 +85,22 @@ private:
 	TWeakPtr<class FWidgetBlueprintEditor> BlueprintEditor;
 
 	/** Handles filtering the hierarchy based on an IFilter. */
-	TSharedPtr<TreeFilterHandler<UWidget*>> FilterHandler;
+	TSharedPtr< TreeFilterHandler< TSharedPtr<FHierarchyModel> > > FilterHandler;
 
 	/** The source root widgets for the tree. */
-	TArray< UWidget* > RootWidgets;
+	TArray< TSharedPtr<FHierarchyModel> > RootWidgets;
 
 	/** The root widgets which are actually displayed by the TreeView which will be managed by the TreeFilterHandler. */
-	TArray< UWidget* > TreeRootWidgets;
+	TArray< TSharedPtr<FHierarchyModel> > TreeRootWidgets;
 
 	/** Gets the visible items displayed in the tree view. */
-	TSet< UWidget* > VisibleItems;
+	TSet< TSharedPtr<FHierarchyModel> > VisibleItems;
 
 	/** The widget containing the treeview */
 	TSharedPtr<SBorder> TreeViewArea;
 
 	/** The widget hierarchy slate treeview widget */
-	TSharedPtr< STreeView< UWidget* > > WidgetTreeView;
+	TSharedPtr< STreeView< TSharedPtr<FHierarchyModel> > > WidgetTreeView;
 
 	/** The filter used by the search box */
 	TSharedPtr<WidgetTextFilter> SearchBoxWidgetFilter;
diff --git a/Engine/Source/Editor/UMGEditor/Private/Hierarchy/SHierarchyViewItem.cpp b/Engine/Source/Editor/UMGEditor/Private/Hierarchy/SHierarchyViewItem.cpp
--- a/Engine/Source/Editor/UMGEditor/Private/Hierarchy/SHierarchyViewItem.cpp
+++ b/Engine/Source/Editor/UMGEditor/Private/Hierarchy/SHierarchyViewItem.cpp
@@ -20,14 +20,14 @@
 #define LOCTEXT_NAMESPACE "UMG"
 
 /**
- * 
- */
+*
+*/
 class FHierarchyWidgetDragDropOp : public FDecoratedDragDropOp
 {
 public:
 	DRAG_DROP_OPERATOR_TYPE(FHierarchyWidgetDragDropOp, FDecoratedDragDropOp)
 
-	virtual ~FHierarchyWidgetDragDropOp();
+		virtual ~FHierarchyWidgetDragDropOp();
 
 	/** The slot properties for the old slot the widget was in, is used to attempt to reapply the same layout information */
 	TMap<FName, FString> ExportedSlotProperties;
@@ -80,89 +80,262 @@
 
 //////////////////////////////////////////////////////////////////////////
 
-void SHierarchyViewItem::Construct(const FArguments& InArgs, const TSharedRef< STableViewBase >& InOwnerTableView, TSharedPtr<FWidgetBlueprintEditor> InBlueprintEditor, FWidgetReference InItem)
+FReply ProcessHierarchyDragDrop(const FDragDropEvent& DragDropEvent, bool bIsDrop, UWidgetBlueprint* Blueprint, FWidgetReference TargetItem)
 {
-	BlueprintEditor = InBlueprintEditor;
-	Item = InItem;
+	// Is this a drag/drop op to create a new widget in the tree?
+	TSharedPtr<FWidgetTemplateDragDropOp> TemplateDragDropOp = DragDropEvent.GetOperationAs<FWidgetTemplateDragDropOp>();
+	if ( TemplateDragDropOp.IsValid() )
+	{
+		TemplateDragDropOp->SetCursorOverride(TOptional<EMouseCursor::Type>());
 
-	STableRow< UWidget* >::Construct(
-		STableRow< UWidget* >::FArguments()
-		.Padding(2.0f)
-		.OnDragDetected(this, &SHierarchyViewItem::HandleDragDetected)
-		.OnDragEnter(this, &SHierarchyViewItem::HandleDragEnter)
-		.OnDragLeave(this, &SHierarchyViewItem::HandleDragLeave)
-		.OnDrop(this, &SHierarchyViewItem::HandleDrop)
-		.Content()
-		[
-			SNew(SHorizontalBox)
+		// Are we adding to the root?
+		if ( !TargetItem.IsValid() && Blueprint->WidgetTree->RootWidget == NULL )
+		{
+			// TODO UMG Allow showing a preview of this.
+			if ( bIsDrop )
+			{
+				Blueprint->WidgetTree->RootWidget = TemplateDragDropOp->Template->Create(Blueprint->WidgetTree);
+				FBlueprintEditorUtils::MarkBlueprintAsStructurallyModified(Blueprint);
+			}
+		}
+		// Are we adding to a panel?
+		else if ( UPanelWidget* Parent = Cast<UPanelWidget>(TargetItem.GetTemplate()) )
+		{
+			// TODO UMG Allow showing a preview of this.
+			if ( bIsDrop )
+			{
+				UWidget* Widget = TemplateDragDropOp->Template->Create(Blueprint->WidgetTree);
+				
+				Parent->AddChild(Widget);
 
-			// Widget icon
-			+ SHorizontalBox::Slot()
-			.AutoWidth()
-			[
-				SNew(SImage)
-				.ColorAndOpacity(FLinearColor(1,1,1,0.5))
-				.Image(Item.GetTemplate()->GetEditorIcon())
-			]
+				FBlueprintEditorUtils::MarkBlueprintAsStructurallyModified(Blueprint);
+			}
+		}
+		//TODO Are we adding to a named slot?
+		else
+		{
+			TemplateDragDropOp->SetCursorOverride(EMouseCursor::SlashedCircle);
+		}
 
-			// Name of the widget
-			+ SHorizontalBox::Slot()
-			.FillWidth(1.0f)
-			.Padding(2, 0, 0, 0)
-			.VAlign(VAlign_Center)
-			[
-				SNew(SInlineEditableTextBlock)
-				.Font(this, &SHierarchyViewItem::GetItemFont)
-				.Text(this, &SHierarchyViewItem::GetItemText)
-				.HighlightText(InArgs._HighlightText)
-				.OnVerifyTextChanged(this, &SHierarchyViewItem::OnVerifyNameTextChanged)
-				.OnTextCommitted(this, &SHierarchyViewItem::OnNameTextCommited)
-				.IsSelected(this, &SHierarchyViewItem::IsSelectedExclusively)
-			]
-		],
-		InOwnerTableView);
-}
+		return FReply::Handled();
+	}
 
-bool SHierarchyViewItem::OnVerifyNameTextChanged(const FText& InText, FText& OutErrorMessage)
-{
-	FString NewName = InText.ToString();
+	TSharedPtr<FHierarchyWidgetDragDropOp> HierarchyDragDropOp = DragDropEvent.GetOperationAs<FHierarchyWidgetDragDropOp>();
+	if ( HierarchyDragDropOp.IsValid() )
+	{
+		HierarchyDragDropOp->SetCursorOverride(TOptional<EMouseCursor::Type>());
 
-	UWidgetBlueprint* Blueprint = BlueprintEditor.Pin()->GetWidgetBlueprintObj();
-	UWidget* ExistingTemplate = Blueprint->WidgetTree->FindWidget(NewName);
+		// If the target item is valid we're dealing with a normal widget in the hierarchy, otherwise we should assume it's
+		// the null case and we should be adding it as the root widget.
+		if ( TargetItem.IsValid() )
+		{
+			const bool bIsPanel = TargetItem.GetTemplate()->IsA(UPanelWidget::StaticClass());
+			const bool bIsDraggedObject = TargetItem.GetTemplate() == HierarchyDragDropOp->Widget.GetTemplate();
+			const bool bIsChildOfDraggedObject = TargetItem.GetTemplate()->IsChildOf(HierarchyDragDropOp->Widget.GetTemplate());
 
-	bool bIsSameWidget = false;
-	if ( ExistingTemplate != NULL )
-	{
-		if ( Item.GetTemplate() == ExistingTemplate )
-		{
-			OutErrorMessage = LOCTEXT("ExistingWidgetName", "Existing Widget Name");
-			return false;
+			if ( bIsPanel && !bIsDraggedObject && !bIsChildOfDraggedObject )
+			{
+				UWidget* Widget = bIsDrop ? HierarchyDragDropOp->Widget.GetTemplate() : HierarchyDragDropOp->Widget.GetPreview();
+
+				if ( Widget->GetParent() )
+				{
+					Widget->GetParent()->RemoveChild(Widget);
+				}
+
+				UPanelWidget* NewParent = Cast<UPanelWidget>(bIsDrop ? TargetItem.GetTemplate() : TargetItem.GetPreview());
+
+				if ( bIsDrop && NewParent != HierarchyDragDropOp->WidgetParent )
+				{
+					NewParent->SetFlags(RF_Transactional);
+					NewParent->Modify();
+				}
+
+				if ( UPanelSlot* Slot = NewParent->AddChild(Widget) )
+				{
+					FWidgetBlueprintEditorUtils::ImportPropertiesFromText(Slot, HierarchyDragDropOp->ExportedSlotProperties);
+
+					Slot->SetDesiredPosition(FVector2D::ZeroVector);
+
+					if ( bIsDrop )
+					{
+						FBlueprintEditorUtils::MarkBlueprintAsStructurallyModified(Blueprint);
+					}
+					else
+					{
+						TSharedPtr<SWidget> Widget = TargetItem.GetPreview()->GetCachedWidget();
+						if ( Widget.IsValid() )
+						{
+							Widget->SlatePrepass();
+						}
+					}
+				}
+				else
+				{
+					HierarchyDragDropOp->SetCursorOverride(EMouseCursor::SlashedCircle);
+
+					// TODO UMG ERROR Slot can not be created because maybe the max children has been reached.
+					//          Maybe we can traverse the hierarchy and add it to the first parent that will accept it?
+				}
+			}
+			else
+			{
+				HierarchyDragDropOp->SetCursorOverride(EMouseCursor::SlashedCircle);
+			}
 		}
 		else
 		{
-			bIsSameWidget = true;
+			HierarchyDragDropOp->SetCursorOverride(EMouseCursor::SlashedCircle);
 		}
+
+		return FReply::Handled();
 	}
 
-	FKismetNameValidator Validator(Blueprint);
+	return FReply::Unhandled();
+}
 
-	const bool bUniqueNameForVariable = ( EValidatorResult::Ok == Validator.IsValid(NewName) );
+//////////////////////////////////////////////////////////////////////////
 
-	if ( !bUniqueNameForVariable && !bIsSameWidget )
+FHierarchyModel::FHierarchyModel()
+	: bInitialized(false)
+{
+
+}
+
+FReply FHierarchyModel::OnDragOver(const FGeometry& MyGeometry, const FDragDropEvent& DragDropEvent)
+{
+	return FReply::Unhandled();
+}
+
+FReply FHierarchyModel::HandleDragDetected(const FGeometry& MyGeometry, const FPointerEvent& MouseEvent)
+{
+	return FReply::Unhandled();
+}
+
+void FHierarchyModel::HandleDragEnter(const FDragDropEvent& DragDropEvent)
+{
+
+}
+
+void FHierarchyModel::HandleDragLeave(const FDragDropEvent& DragDropEvent)
+{
+	TSharedPtr<FDecoratedDragDropOp> DecoratedDragDropOp = DragDropEvent.GetOperationAs<FDecoratedDragDropOp>();
+	if ( DecoratedDragDropOp.IsValid() )
 	{
-		OutErrorMessage = LOCTEXT("ExistingVariableName", "Existing Variable Name");
-		return false;
+		DecoratedDragDropOp->SetCursorOverride(TOptional<EMouseCursor::Type>());
+		DecoratedDragDropOp->ResetToDefaultToolTip();
+	}
+}
+
+FReply FHierarchyModel::HandleDrop(FDragDropEvent const& DragDropEvent)
+{
+	return FReply::Unhandled();
+}
+
+bool FHierarchyModel::OnVerifyNameTextChanged(const FText& InText, FText& OutErrorMessage)
+{
+	return false;
+}
+
+void FHierarchyModel::OnNameTextCommited(const FText& InText, ETextCommit::Type CommitInfo)
+{
+
+}
+
+void FHierarchyModel::GatherChildren(TArray< TSharedPtr<FHierarchyModel> >& Children)
+{
+	if ( !bInitialized )
+	{
+		bInitialized = true;
+		GetChildren(Models);
 	}
 
-	return true;
+	Children.Append(Models);
 }
 
-void SHierarchyViewItem::OnNameTextCommited(const FText& InText, ETextCommit::Type CommitInfo)
+//////////////////////////////////////////////////////////////////////////
+
+FHierarchyRoot::FHierarchyRoot(TSharedPtr<FWidgetBlueprintEditor> InBlueprintEditor)
+	: BlueprintEditor(InBlueprintEditor)
 {
-	FWidgetBlueprintEditorUtils::RenameWidget(BlueprintEditor.Pin().ToSharedRef(), Item.GetTemplate()->GetFName(), FName(*InText.ToString()));
 }
 
-FSlateFontInfo SHierarchyViewItem::GetItemFont() const
+FText FHierarchyRoot::GetText() const
+{
+	return LOCTEXT("Root", "[Root]");
+}
+
+const FSlateBrush* FHierarchyRoot::GetImage() const
+{
+	return NULL;
+}
+
+FSlateFontInfo FHierarchyRoot::GetFont() const
+{
+	return FSlateFontInfo(FPaths::EngineContentDir() / TEXT("Slate/Fonts/Roboto-Bold.ttf"), 10);
+}
+
+void FHierarchyRoot::GetChildren(TArray< TSharedPtr<FHierarchyModel> >& Children)
+{
+	TSharedPtr<FWidgetBlueprintEditor> BPEd = BlueprintEditor.Pin();
+	UWidgetBlueprint* Blueprint = BPEd->GetWidgetBlueprintObj();
+
+	if ( Blueprint->WidgetTree->RootWidget )
+	{
+		TSharedPtr<FHierarchyWidget> RootChild = MakeShareable(new FHierarchyWidget(FWidgetReference::FromTemplate(BPEd, Blueprint->WidgetTree->RootWidget), BPEd));
+		Children.Add(RootChild);
+	}
+}
+
+void FHierarchyRoot::OnSelection()
+{
+	TSharedPtr<FWidgetBlueprintEditor> BPEd = BlueprintEditor.Pin();
+	if ( UWidget* Default =BPEd->GetWidgetBlueprintObj()->GeneratedClass->GetDefaultObject<UWidget>() )
+	{
+		TSet<UObject*> SelectedObjects;
+		SelectedObjects.Add(Default);
+		BPEd->SelectObjects(SelectedObjects);
+	}
+}
+
+FReply FHierarchyRoot::OnDragOver(const FGeometry& MyGeometry, const FDragDropEvent& DragDropEvent)
+{
+	UWidgetBlueprint* Blueprint = BlueprintEditor.Pin()->GetWidgetBlueprintObj();
+	bool bIsDrop = false;
+	return ProcessHierarchyDragDrop(DragDropEvent, bIsDrop, Blueprint, FWidgetReference());
+}
+
+FReply FHierarchyRoot::HandleDrop(FDragDropEvent const& DragDropEvent)
+{
+	UWidgetBlueprint* Blueprint = BlueprintEditor.Pin()->GetWidgetBlueprintObj();
+	bool bIsDrop = true;
+	return ProcessHierarchyDragDrop(DragDropEvent, bIsDrop, Blueprint, FWidgetReference());
+}
+
+//////////////////////////////////////////////////////////////////////////
+
+FHierarchyWidget::FHierarchyWidget(FWidgetReference InItem, TSharedPtr<FWidgetBlueprintEditor> InBlueprintEditor)
+	: Item(InItem)
+	, BlueprintEditor(InBlueprintEditor)
+{
+}
+
+FText FHierarchyWidget::GetText() const
 {
 	UWidget* WidgetTemplate = Item.GetTemplate();
 	if ( WidgetTemplate )
+	{
+		return FText::FromString(WidgetTemplate->GetLabel());
+	}
+
+	return FText::GetEmpty();
+}
+
+const FSlateBrush* FHierarchyWidget::GetImage() const
+{
+	return Item.GetTemplate()->GetEditorIcon();
+}
+
+FSlateFontInfo FHierarchyWidget::GetFont() const
+{
+	UWidget* WidgetTemplate = Item.GetTemplate();
+	if ( WidgetTemplate )
@@ -178,18 +351,27 @@
 	return FCoreStyle::Get().GetFontStyle(NormalFont);
 }
 
-FText SHierarchyViewItem::GetItemText() const
+FReply FHierarchyWidget::OnDragOver(const FGeometry& MyGeometry, const FDragDropEvent& DragDropEvent)
 {
-	UWidget* WidgetTemplate = Item.GetTemplate();
-	if ( WidgetTemplate )
+	TSharedPtr<FWidgetTemplateDragDropOp> TemplateDragDropOp = DragDropEvent.GetOperationAs<FWidgetTemplateDragDropOp>();
+	if ( TemplateDragDropOp.IsValid() )
 	{
-		return FText::FromString(WidgetTemplate->GetLabel());
+		// TODO UMG - Show feedback
+		return FReply::Handled();
 	}
 
-	return FText::GetEmpty();
+	TSharedPtr<FHierarchyWidgetDragDropOp> HierarchyDragDropOp = DragDropEvent.GetOperationAs<FHierarchyWidgetDragDropOp>();
+	if ( HierarchyDragDropOp.IsValid() )
+	{
+		bool bIsDrop = false;
+		UWidgetBlueprint* Blueprint = BlueprintEditor.Pin()->GetWidgetBlueprintObj();
+		return ProcessHierarchyDragDrop(DragDropEvent, bIsDrop, Blueprint, Item);
+	}
+
+	return FReply::Unhandled();
 }
 
-FReply SHierarchyViewItem::HandleDragDetected(const FGeometry& MyGeometry, const FPointerEvent& MouseEvent)
+FReply FHierarchyWidget::HandleDragDetected(const FGeometry& MyGeometry, const FPointerEvent& MouseEvent)
 {
 	const bool bIsRoot = Item.GetTemplate()->GetParent() == NULL;
 
@@ -201,11 +383,7 @@
 	return FReply::Unhandled();
 }
 
-void SHierarchyViewItem::HandleDragEnter(FDragDropEvent const& DragDropEvent)
+void FHierarchyWidget::HandleDragLeave(const FDragDropEvent& DragDropEvent)
 {
-}
-
-void SHierarchyViewItem::HandleDragLeave(const FDragDropEvent& DragDropEvent)
-{
 	TSharedPtr<FDecoratedDragDropOp> DecoratedDragDropOp = DragDropEvent.GetOperationAs<FDecoratedDragDropOp>();
 	if ( DecoratedDragDropOp.IsValid() )
@@ -215,122 +393,166 @@
 	}
 }
 
-FReply SHierarchyViewItem::OnDragOver(const FGeometry& MyGeometry, const FDragDropEvent& DragDropEvent)
+FReply FHierarchyWidget::HandleDrop(const FDragDropEvent& DragDropEvent)
 {
-	TSharedPtr<FWidgetTemplateDragDropOp> TemplateDragDropOp = DragDropEvent.GetOperationAs<FWidgetTemplateDragDropOp>();
-	if ( TemplateDragDropOp.IsValid() )
-	{
-		// TODO UMG - Show feedback
-		return FReply::Handled();
-	}
-
-	TSharedPtr<FHierarchyWidgetDragDropOp> HierarchyDragDropOp = DragDropEvent.GetOperationAs<FHierarchyWidgetDragDropOp>();
-	if ( HierarchyDragDropOp.IsValid() )
-	{
-		bool bIsDrop = false;
-		return ProcessDragDrop(DragDropEvent, bIsDrop);
-	}
-
-	return FReply::Unhandled();
-}
-
-FReply SHierarchyViewItem::HandleDrop(const FDragDropEvent& DragDropEvent)
-{
-	TSharedPtr<FWidgetTemplateDragDropOp> TemplateDragDropOp = DragDropEvent.GetOperationAs<FWidgetTemplateDragDropOp>();
-	if ( TemplateDragDropOp.IsValid() )
-	{
-		if ( Item.GetTemplate()->IsA(UPanelWidget::StaticClass()) )
-		{
-			UWidgetBlueprint* BP = CastChecked<UWidgetBlueprint>(BlueprintEditor.Pin()->GetBlueprintObj());			
-			UPanelWidget* Parent = Cast<UPanelWidget>(Item.GetTemplate());
-
-			UWidget* Widget = TemplateDragDropOp->Template->Create(BP->WidgetTree);
-
-			Parent->AddChild(Widget);
-
-			FBlueprintEditorUtils::MarkBlueprintAsStructurallyModified(BP);
-		}
-
-		return FReply::Handled();
-	}
-
+	UWidgetBlueprint* Blueprint = BlueprintEditor.Pin()->GetWidgetBlueprintObj();
 	bool bIsDrop = true;
-	return ProcessDragDrop(DragDropEvent, bIsDrop);
+	return ProcessHierarchyDragDrop(DragDropEvent, bIsDrop, Blueprint, Item);
 }
 
-FReply SHierarchyViewItem::ProcessDragDrop(const FDragDropEvent& DragDropEvent, bool bIsDrop)
+bool FHierarchyWidget::OnVerifyNameTextChanged(const FText& InText, FText& OutErrorMessage)
 {
+	FString NewName = InText.ToString();
+
 	UWidgetBlueprint* Blueprint = BlueprintEditor.Pin()->GetWidgetBlueprintObj();
+	UWidget* ExistingTemplate = Blueprint->WidgetTree->FindWidget(NewName);
 
-	TSharedPtr<FHierarchyWidgetDragDropOp> HierarchyDragDropOp = DragDropEvent.GetOperationAs<FHierarchyWidgetDragDropOp>();
-	if ( HierarchyDragDropOp.IsValid() )
+	bool bIsSameWidget = false;
+	if ( ExistingTemplate != NULL )
 	{
-		HierarchyDragDropOp->SetCursorOverride(TOptional<EMouseCursor::Type>());
-
-		if ( Item.IsValid() )
+		if ( Item.GetTemplate() == ExistingTemplate )
 		{
-			const bool bIsPanel = Item.GetTemplate()->IsA(UPanelWidget::StaticClass());
-			const bool bIsDraggedObject = Item.GetTemplate() == HierarchyDragDropOp->Widget.GetTemplate();
-			const bool bIsChildOfDraggedObject = Item.GetTemplate()->IsChildOf(HierarchyDragDropOp->Widget.GetTemplate());
-
-			if ( bIsPanel && !bIsDraggedObject && !bIsChildOfDraggedObject )
-			{
-				UWidget* Widget = bIsDrop ? HierarchyDragDropOp->Widget.GetTemplate() : HierarchyDragDropOp->Widget.GetPreview();
-
-				if ( Widget->GetParent() )
-				{
-					Widget->GetParent()->RemoveChild(Widget);
-				}
-
-				UPanelWidget* NewParent = Cast<UPanelWidget>( bIsDrop ? Item.GetTemplate() : Item.GetPreview() );
-
-				if ( bIsDrop && NewParent != HierarchyDragDropOp->WidgetParent )
-				{
-					NewParent->SetFlags(RF_Transactional);
-					NewParent->Modify();
-				}
-
-				if ( UPanelSlot* Slot = NewParent->AddChild(Widget) )
-				{
-					FWidgetBlueprintEditorUtils::ImportPropertiesFromText(Slot, HierarchyDragDropOp->ExportedSlotProperties);
-
-					Slot->SetDesiredPosition(FVector2D::ZeroVector);
-
-					if ( bIsDrop )
-					{
-						FBlueprintEditorUtils::MarkBlueprintAsStructurallyModified(Blueprint);
-					}
-					else
-					{
-						TSharedPtr<SWidget> Widget = Item.GetPreview()->GetCachedWidget();
-						if ( Widget.IsValid() )
-						{
-							Widget->SlatePrepass();
-						}
-					}
-				}
-				else
-				{
-					HierarchyDragDropOp->SetCursorOverride(EMouseCursor::SlashedCircle);
-
-					// TODO UMG ERROR Slot can not be created because maybe the max children has been reached.
-					//          Maybe we can traverse the hierarchy and add it to the first parent that will accept it?
-				}
-			}
-			else
-			{
-				HierarchyDragDropOp->SetCursorOverride(EMouseCursor::SlashedCircle);
-			}
+			OutErrorMessage = LOCTEXT("ExistingWidgetName", "Existing Widget Name");
+			return false;
 		}
 		else
 		{
-			HierarchyDragDropOp->SetCursorOverride(EMouseCursor::SlashedCircle);
+			bIsSameWidget = true;
 		}
-
-		return FReply::Handled();
 	}
 
-	return FReply::Unhandled();
+	FKismetNameValidator Validator(Blueprint);
+
+	const bool bUniqueNameForVariable = ( EValidatorResult::Ok == Validator.IsValid(NewName) );
+
+	if ( !bUniqueNameForVariable && !bIsSameWidget )
+	{
+		OutErrorMessage = LOCTEXT("ExistingVariableName", "Existing Variable Name");
+		return false;
+	}
+
+	return true;
 }
 
+void FHierarchyWidget::OnNameTextCommited(const FText& InText, ETextCommit::Type CommitInfo)
+{
+	FWidgetBlueprintEditorUtils::RenameWidget(BlueprintEditor.Pin().ToSharedRef(), Item.GetTemplate()->GetFName(), FName(*InText.ToString()));
+}
+
+void FHierarchyWidget::GetChildren(TArray< TSharedPtr<FHierarchyModel> >& Children)
+{
+	TSharedPtr<FWidgetBlueprintEditor> BPEd = BlueprintEditor.Pin();
+	
+	UPanelWidget* Widget = Cast<UPanelWidget>(Item.GetTemplate());
+	if ( Widget )
+	{
+		for ( int32 i = 0; i < Widget->GetChildrenCount(); i++ )
+		{
+			UWidget* Child = Widget->GetChildAt(i);
+			if ( Child )
+			{
+				TSharedPtr<FHierarchyWidget> ChildItem = MakeShareable(new FHierarchyWidget(FWidgetReference::FromTemplate(BPEd, Child), BPEd));
+				Children.Add(ChildItem);
+			}
+		}
+	}
+}
+
+void FHierarchyWidget::OnSelection()
+{
+	TSet<FWidgetReference> SelectedWidgets;
+	SelectedWidgets.Add(Item);
+
+	BlueprintEditor.Pin()->SelectWidgets(SelectedWidgets);
+}
+
+//////////////////////////////////////////////////////////////////////////
+
+void SHierarchyViewItem::Construct(const FArguments& InArgs, const TSharedRef< STableViewBase >& InOwnerTableView, TSharedPtr<FHierarchyModel> InModel)
+{
+	Model = InModel;
+
+	STableRow< TSharedPtr<FHierarchyModel> >::Construct(
+		STableRow< TSharedPtr<FHierarchyModel> >::FArguments()
+		.Padding(2.0f)
+		.OnDragDetected(this, &SHierarchyViewItem::HandleDragDetected)
+		.OnDragEnter(this, &SHierarchyViewItem::HandleDragEnter)
+		.OnDragLeave(this, &SHierarchyViewItem::HandleDragLeave)
+		.OnDrop(this, &SHierarchyViewItem::HandleDrop)
+		.Content()
+		[
+			SNew(SHorizontalBox)
+
+			// Widget icon
+			+ SHorizontalBox::Slot()
+			.AutoWidth()
+			[
+				SNew(SImage)
+				.ColorAndOpacity(FLinearColor(1,1,1,0.5))
+				.Image(Model->GetImage())
+			]
+
+			// Name of the widget
+			+ SHorizontalBox::Slot()
+			.FillWidth(1.0f)
+			.Padding(2, 0, 0, 0)
+			.VAlign(VAlign_Center)
+			[
+				SNew(SInlineEditableTextBlock)
+				.Font(this, &SHierarchyViewItem::GetItemFont)
+				.Text(this, &SHierarchyViewItem::GetItemText)
+				.HighlightText(InArgs._HighlightText)
+				.OnVerifyTextChanged(this, &SHierarchyViewItem::OnVerifyNameTextChanged)
+				.OnTextCommitted(this, &SHierarchyViewItem::OnNameTextCommited)
+				.IsSelected(this, &SHierarchyViewItem::IsSelectedExclusively)
+			]
+		],
+		InOwnerTableView);
+}
+
+bool SHierarchyViewItem::OnVerifyNameTextChanged(const FText& InText, FText& OutErrorMessage)
+{
+	return Model->OnVerifyNameTextChanged(InText, OutErrorMessage);
+}
+
+void SHierarchyViewItem::OnNameTextCommited(const FText& InText, ETextCommit::Type CommitInfo)
+{
+	return Model->OnNameTextCommited(InText, CommitInfo);
+}
+
+FSlateFontInfo SHierarchyViewItem::GetItemFont() const
+{
+	return Model->GetFont();
+}
+
+FText SHierarchyViewItem::GetItemText() const
+{
+	return Model->GetText();
+}
+
+FReply SHierarchyViewItem::OnDragOver(const FGeometry& MyGeometry, const FDragDropEvent& DragDropEvent)
+{
+	return Model->OnDragOver(MyGeometry, DragDropEvent);
+}
+
+FReply SHierarchyViewItem::HandleDragDetected(const FGeometry& MyGeometry, const FPointerEvent& MouseEvent)
+{
+	return Model->HandleDragDetected(MyGeometry, MouseEvent);
+}
+
+void SHierarchyViewItem::HandleDragEnter(FDragDropEvent const& DragDropEvent)
+{
+	Model->HandleDragEnter(DragDropEvent);
+}
+
+void SHierarchyViewItem::HandleDragLeave(const FDragDropEvent& DragDropEvent)
+{
+	Model->HandleDragLeave(DragDropEvent);
+}
+
+FReply SHierarchyViewItem::HandleDrop(const FDragDropEvent& DragDropEvent)
+{
+	return Model->HandleDrop(DragDropEvent);
+}
+
 #undef LOCTEXT_NAMESPACE
diff --git a/Engine/Source/Editor/UMGEditor/Private/Hierarchy/SHierarchyViewItem.h b/Engine/Source/Editor/UMGEditor/Private/Hierarchy/SHierarchyViewItem.h
--- a/Engine/Source/Editor/UMGEditor/Private/Hierarchy/SHierarchyViewItem.h
+++ b/Engine/Source/Editor/UMGEditor/Private/Hierarchy/SHierarchyViewItem.h
@@ -5,10 +5,112 @@
 #include "SCompoundWidget.h"
 #include "BlueprintEditor.h"
 
+class FHierarchyModel : public TSharedFromThis < FHierarchyModel >
+{
+public:
+	FHierarchyModel();
+
+	/* @returns the widget name to use for the tree item */
+	virtual FText GetText() const = 0;
+
+	virtual const FSlateBrush* GetImage() const = 0;
+
+	virtual FSlateFontInfo GetFont() const = 0;
+
+	virtual FReply OnDragOver(const FGeometry& MyGeometry, const FDragDropEvent& DragDropEvent);
+
+	virtual FReply HandleDragDetected(const FGeometry& MyGeometry, const FPointerEvent& MouseEvent);
+	virtual void HandleDragEnter(const FDragDropEvent& DragDropEvent);
+	virtual void HandleDragLeave(const FDragDropEvent& DragDropEvent);
+	virtual FReply HandleDrop(FDragDropEvent const& DragDropEvent);
+
+	virtual bool OnVerifyNameTextChanged(const FText& InText, FText& OutErrorMessage);
+
+	virtual void OnNameTextCommited(const FText& InText, ETextCommit::Type CommitInfo);
+
+	virtual void GatherChildren(TArray< TSharedPtr<FHierarchyModel> >& Children);
+
+	virtual void OnSelection() = 0;
+
+	virtual bool IsModel(UObject* PossibleModel) const { return false; }
+
+protected:
+	virtual void GetChildren(TArray< TSharedPtr<FHierarchyModel> >& Children) = 0;
+
+protected:
+
+	bool bInitialized;
+	TArray< TSharedPtr<FHierarchyModel> > Models;
+};
+
+class FHierarchyRoot : public FHierarchyModel
+{
+public:
+	FHierarchyRoot(TSharedPtr<FWidgetBlueprintEditor> InBlueprintEditor);
+
+	/* @returns the widget name to use for the tree item */
+	virtual FText GetText() const override;
+
+	virtual const FSlateBrush* GetImage() const override;
+
+	virtual FSlateFontInfo GetFont() const override;
+
+	virtual void OnSelection();
+
+	virtual FReply OnDragOver(const FGeometry& MyGeometry, const FDragDropEvent& DragDropEvent);
+
+	virtual FReply HandleDrop(FDragDropEvent const& DragDropEvent);
+
+protected:
+	virtual void GetChildren(TArray< TSharedPtr<FHierarchyModel> >& Children) override;
+
+private:
+
+	TWeakPtr<FWidgetBlueprintEditor> BlueprintEditor;
+};
+
+class FHierarchyWidget : public FHierarchyModel
+{
+public:
+	FHierarchyWidget(FWidgetReference InItem, TSharedPtr<FWidgetBlueprintEditor> InBlueprintEditor);
+
+	/* @returns the widget name to use for the tree item */
+	virtual FText GetText() const override;
+
+	virtual const FSlateBrush* GetImage() const override;
+
+	virtual FSlateFontInfo GetFont() const override;
+
+	virtual FReply OnDragOver(const FGeometry& MyGeometry, const FDragDropEvent& DragDropEvent) override;
+
+	virtual FReply HandleDragDetected(const FGeometry& MyGeometry, const FPointerEvent& MouseEvent) override;
+	virtual void HandleDragLeave(const FDragDropEvent& DragDropEvent) override;
+	virtual FReply HandleDrop(FDragDropEvent const& DragDropEvent) override;
+
+	virtual bool OnVerifyNameTextChanged(const FText& InText, FText& OutErrorMessage) override;
+
+	virtual void OnNameTextCommited(const FText& InText, ETextCommit::Type CommitInfo) override;
+
+	virtual void OnSelection();
+
+	virtual bool IsModel(UObject* PossibleModel) const override
+	{
+		return Item.GetTemplate() == PossibleModel;
+	}
+
+protected:
+	virtual void GetChildren(TArray< TSharedPtr<FHierarchyModel> >& Children) override;
+
+private:
+	FWidgetReference Item;
+
+	TWeakPtr<FWidgetBlueprintEditor> BlueprintEditor;
+};
+
 /**
  * An widget item in the hierarchy tree view.
  */
-class SHierarchyViewItem : public STableRow< UWidget* >
+class SHierarchyViewItem : public STableRow< TSharedPtr<FHierarchyModel> >
 {
 public:
 	SLATE_BEGIN_ARGS( SHierarchyViewItem ){}
@@ -18,7 +120,7 @@ public:
 
 	SLATE_END_ARGS()
 
-	void Construct(const FArguments& InArgs, const TSharedRef< STableViewBase >& InOwnerTableView, TSharedPtr<FWidgetBlueprintEditor> InBlueprintEditor, FWidgetReference InItem);
+	void Construct(const FArguments& InArgs, const TSharedRef< STableViewBase >& InOwnerTableView, TSharedPtr<FHierarchyModel> InModel);
 
 	// Begin SWidget
 	virtual FReply OnDragOver(const FGeometry& MyGeometry, const FDragDropEvent& DragDropEvent) override;
@@ -30,9 +132,6 @@ private:
 	void HandleDragLeave(const FDragDropEvent& DragDropEvent);
 	FReply HandleDrop(FDragDropEvent const& DragDropEvent);
 
-	/** Process Drag / Drop events providing either feedback, or generating the final drop action */
-	FReply ProcessDragDrop(const FDragDropEvent& DragDropEvent, bool bIsDrop);
-
 	/** Called when text is being committed to check for validity */
 	bool OnVerifyNameTextChanged(const FText& InText, FText& OutErrorMessage);
 
@@ -45,9 +144,6 @@ private:
 	/* @returns the widget name to use for the tree item */
 	FText GetItemText() const;
 
-	/* The blueprint editor this tree view item's tree is associated with. */
-	TWeakPtr<FWidgetBlueprintEditor> BlueprintEditor;
-
-	/* The template widget that this tree item represents */
-	FWidgetReference Item;
+	/* The mode that this tree item represents */
+	TSharedPtr<FHierarchyModel> Model;
 };
diff --git a/Engine/Source/Editor/UMGEditor/Private/WidgetBlueprintEditor.cpp b/Engine/Source/Editor/UMGEditor/Private/WidgetBlueprintEditor.cpp
--- a/Engine/Source/Editor/UMGEditor/Private/WidgetBlueprintEditor.cpp
+++ b/Engine/Source/Editor/UMGEditor/Private/WidgetBlueprintEditor.cpp
@@ -134,11 +134,28 @@ void FWidgetBlueprintEditor::SelectWidgets(const TSet<FWidgetReference>& Widgets
 	// to ensure values that are pending are committed on focus loss, and migrated properly
 	// to the old selected widgets.
 	SelectedWidgets.Empty();
+	SelectedObjects.Empty();
+
 	SelectedWidgets.Append(TempSelection);
 
 	OnSelectedWidgetsChanged.Broadcast();
 }
 
+void FWidgetBlueprintEditor::SelectObjects(const TSet<UObject*>& Objects)
+{
+	OnSelectedWidgetsChanging.Broadcast();
+
+	SelectedWidgets.Empty();
+	SelectedObjects.Empty();
+
+	for ( UObject* Obj : Objects )
+	{
+		SelectedObjects.Add(Obj);
+	}
+
+	OnSelectedWidgetsChanged.Broadcast();
+}
+
 void FWidgetBlueprintEditor::CleanSelection()
 {
 	TSet<FWidgetReference> TempSelection;
@@ -169,6 +186,11 @@ const TSet<FWidgetReference>& FWidgetBlueprintEditor::GetSelectedWidgets() const
 	return SelectedWidgets;
 }
 
+const TSet< TWeakObjectPtr<UObject> >& FWidgetBlueprintEditor::GetSelectedObjects() const
+{
+	return SelectedObjects;
+}
+
 void FWidgetBlueprintEditor::OnBlueprintChanged(UBlueprint* InBlueprint)
 {
 	FBlueprintEditor::OnBlueprintChanged(InBlueprint);
diff --git a/Engine/Source/Editor/UMGEditor/Private/WidgetBlueprintEditor.h b/Engine/Source/Editor/UMGEditor/Private/WidgetBlueprintEditor.h
--- a/Engine/Source/Editor/UMGEditor/Private/WidgetBlueprintEditor.h
+++ b/Engine/Source/Editor/UMGEditor/Private/WidgetBlueprintEditor.h
@@ -52,12 +52,18 @@ public:
 	/** Sets the currently selected set of widgets */
 	void SelectWidgets(const TSet<FWidgetReference>& Widgets);
 
+	/** Sets the currently selected set of objects */
+	void SelectObjects(const TSet<UObject*>& Objects);
+
 	/** Removes removed widgets from the selection set. */
 	void CleanSelection();
 
 	/** @return The selected set of widgets */
 	const TSet<FWidgetReference>& GetSelectedWidgets() const;
 
+	/** @return The selected set of Objects */
+	const TSet< TWeakObjectPtr<UObject> >& GetSelectedObjects() const;
+
 	/** @return Notification for when the preview widget has been updated */
 	FOnWidgetPreviewUpdated& GetOnWidgetPreviewUpdated() { return OnWidgetPreviewUpdated; }
 
@@ -136,6 +142,9 @@ private:
 	/** The currently selected preview widgets in the preview GUI */
 	TSet<FWidgetReference> SelectedWidgets;
 
+	/** The currently selected objects in the designer */
+	TSet< TWeakObjectPtr<UObject> > SelectedObjects;
+
 	/** The preview GUI object */
 	mutable TWeakObjectPtr<UUserWidget> PreviewWidgetPtr;
 
diff --git a/Engine/Source/Runtime/UMG/Private/Components/NamedSlotInterface.cpp b/Engine/Source/Runtime/UMG/Private/Components/NamedSlotInterface.cpp
new file mode 100644
--- /dev/null
+++ b/Engine/Source/Runtime/UMG/Private/Components/NamedSlotInterface.cpp
@@ -0,0 +1,9 @@
+// Copyright 1998-2014 Epic Games, Inc. All Rights Reserved.
+
+#include "UMGPrivatePCH.h"
+
+UNamedSlotInterface::UNamedSlotInterface(const FPostConstructInitializeProperties& PCIP)
+	: Super(PCIP)
+{
+}
+
diff --git a/Engine/Source/Runtime/UMG/Private/Slate/SObjectWidget.cpp b/Engine/Source/Runtime/UMG/Private/Slate/SObjectWidget.cpp
--- a/Engine/Source/Runtime/UMG/Private/Slate/SObjectWidget.cpp
+++ b/Engine/Source/Runtime/UMG/Private/Slate/SObjectWidget.cpp
@@ -58,6 +58,16 @@ int32 SObjectWidget::OnPaint(const FPaintArgs& Args, const FGeometry& AllottedGe
 	return MaxLayer;
 }
 
+bool SObjectWidget::SupportsKeyboardFocus() const
+{
+	if ( WidgetObject && !WidgetObject->IsDesignTime() )
+	{
+		return WidgetObject->bSupportsKeyboardFocus;
+	}
+
+	return false;
+}
+
 FReply SObjectWidget::OnKeyboardFocusReceived(const FGeometry& MyGeometry, const FKeyboardFocusEvent& InKeyboardFocusEvent)
 {
 	if ( WidgetObject && !WidgetObject->IsDesignTime() )
diff --git a/Engine/Source/Runtime/UMG/Private/UserWidget.cpp b/Engine/Source/Runtime/UMG/Private/UserWidget.cpp
--- a/Engine/Source/Runtime/UMG/Private/UserWidget.cpp
+++ b/Engine/Source/Runtime/UMG/Private/UserWidget.cpp
@@ -20,9 +20,10 @@ class SViewportWidgetHost : public SCompoundWidget
 
 	SLATE_END_ARGS()
 
-	void Construct(const FArguments& InArgs, bool bInModal)
+	void Construct(const FArguments& InArgs, TSharedPtr<SWidget> InUserWidget, bool bInModal)
 	{
 		bModal = bInModal;
+		UserWidget = InUserWidget;
 
 		SetVisibility(bModal ? EVisibility::Visible : EVisibility::SelfHitTestInvisible);
 
@@ -40,9 +41,9 @@ class SViewportWidgetHost : public SCompoundWidget
 	FReply OnKeyboardFocusReceived(const FGeometry& MyGeometry, const FKeyboardFocusEvent& InKeyboardFocusEvent) override
 	{
 		// if we support focus, release the focus captures and lock the focus to this widget 
-		if ( SupportsKeyboardFocus() )
+		if ( bModal )
 		{
-			return FReply::Handled().ReleaseMouseCapture().ReleaseJoystickCapture();// .LockMouseToWidget(SharedThis(this));
+			return FReply::Handled();// .ReleaseMouseCapture().ReleaseJoystickCapture();// .LockMouseToWidget(SharedThis(this));
 		}
 		else
 		{
@@ -65,48 +66,92 @@ class SViewportWidgetHost : public SCompoundWidget
 
 	virtual FReply OnKeyChar(const FGeometry& MyGeometry, const FCharacterEvent& InCharacterEvent) override
 	{
-		return bModal ? FReply::Handled() : SCompoundWidget::OnKeyChar(MyGeometry, InCharacterEvent);
+		return bModal ? FReply::Handled() : FReply::Unhandled();
 	}
-			   
+
 	virtual FReply OnKeyDown(const FGeometry& MyGeometry, const FKeyboardEvent& InKeyboardEvent) override
 	{
-		return bModal ? FReply::Handled() : SCompoundWidget::OnKeyDown(MyGeometry, InKeyboardEvent);
+		FReply Result = UserWidget->OnKeyDown(MyGeometry, InKeyboardEvent);
+		if ( bModal && !Result.IsEventHandled() )
+		{
+			return FReply::Handled();
+		}
+
+		return Result;
 	}
-		   
+	
 	virtual FReply OnKeyUp(const FGeometry& MyGeometry, const FKeyboardEvent& InKeyboardEvent) override
 	{
-		return bModal ? FReply::Handled() : SCompoundWidget::OnKeyUp(MyGeometry, InKeyboardEvent);
+		FReply Result = UserWidget->OnKeyUp(MyGeometry, InKeyboardEvent);
+		if ( bModal && !Result.IsEventHandled() )
+		{
+			return FReply::Handled();
+		}
+
+		return Result;
 	}
-		   
+
 	virtual FReply OnMouseButtonDown(const FGeometry& MyGeometry, const FPointerEvent& MouseEvent) override
 	{
-		return bModal ? FReply::Handled().CaptureMouse(AsShared()) : SCompoundWidget::OnMouseButtonDown(MyGeometry, MouseEvent);
+		FReply Result = UserWidget->OnMouseButtonDown(MyGeometry, MouseEvent);
+		if ( bModal && !Result.IsEventHandled() )
+		{
+			return FReply::Handled().CaptureMouse(AsShared());
+		}
+
+		return Result;
 	}
-		   
+
 	virtual FReply OnMouseButtonUp(const FGeometry& MyGeometry, const FPointerEvent& MouseEvent) override
 	{
-		return bModal ? FReply::Handled().ReleaseMouseCapture() : SCompoundWidget::OnMouseButtonUp(MyGeometry, MouseEvent);
+		FReply Result = UserWidget->OnMouseButtonUp(MyGeometry, MouseEvent);
+		if ( bModal && !Result.IsEventHandled() )
+		{
+			return FReply::Handled().ReleaseMouseCapture();
+		}
+
+		return Result;
 	}
 
 	virtual FReply OnMouseButtonDoubleClick(const FGeometry& MyGeometry, const FPointerEvent& MouseEvent) override
 	{
-		return bModal ? FReply::Handled() : SCompoundWidget::OnMouseButtonDoubleClick(MyGeometry, MouseEvent);
+		FReply Result = UserWidget->OnMouseButtonDoubleClick(MyGeometry, MouseEvent);
+		if ( bModal && !Result.IsEventHandled() )
+		{
+			return FReply::Handled();
+		}
+
+		return Result;
 	}
 
 	virtual FReply OnMouseMove(const FGeometry& MyGeometry, const FPointerEvent& MouseEvent) override
 	{
-		return bModal ? FReply::Handled() : SCompoundWidget::OnMouseMove(MyGeometry, MouseEvent);
+		FReply Result = UserWidget->OnMouseMove(MyGeometry, MouseEvent);
+		if ( bModal && !Result.IsEventHandled() )
+		{
+			return FReply::Handled();
+		}
+
+		return Result;
 	}
 
 	virtual FReply OnMouseWheel(const FGeometry& MyGeometry, const FPointerEvent& MouseEvent) override
 	{
-		return bModal ? FReply::Handled() : SCompoundWidget::OnMouseWheel(MyGeometry, MouseEvent);
+		FReply Result = UserWidget->OnMouseWheel(MyGeometry, MouseEvent);
+		if ( bModal && !Result.IsEventHandled() )
+		{
+			return FReply::Handled();
+		}
+
+		return Result;
 	}
 
 	bool IsModal() const { return bModal; }
 
 protected:
 	bool bModal;
+
+	TSharedPtr<SWidget> UserWidget;
 };
 
 
@@ -123,11 +168,14 @@ UUserWidget::UUserWidget(const FPostConstructInitializeProperties& PCIP)
 
 	bInitialized = false;
 	CachedWorld = NULL;
+
+	bSupportsKeyboardFocus = true;
 }
 
 void UUserWidget::Initialize()
 {
-	if ( !bInitialized )
+	// If it's not initialized initialize it, as long as it's not the CDO, we never initialize the CDO.
+	if ( !bInitialized && !HasAnyFlags(RF_ClassDefaultObject) )
 	{
 		bInitialized = true;
 
@@ -335,8 +383,10 @@ void UUserWidget::NativeTick(const FGeometry& MyGeometry, float InDeltaTime )
 	Tick( MyGeometry, InDeltaTime );
 }
 
-TSharedRef<SWidget> UUserWidget::MakeViewportWidget(bool bAbsoluteLayout, bool bModal, bool bShowCursor)
+TSharedRef<SWidget> UUserWidget::MakeViewportWidget(bool bAbsoluteLayout, bool bModal, bool bShowCursor, TSharedPtr<SWidget>& UserSlateWidget)
 {
+	UserSlateWidget = TakeWidget();
+
 	if ( bAbsoluteLayout )
 	{
 		//BIND_UOBJECT_ATTRIBUTE
@@ -349,7 +399,7 @@ TSharedRef<SWidget> UUserWidget::MakeViewportWidget(bool bAbsoluteLayout, bool b
 			.Alignment(BIND_UOBJECT_ATTRIBUTE(FVector2D, GetFullScreenAlignment))
 			.ZOrder(BIND_UOBJECT_ATTRIBUTE(int32, GetFullScreenZOrder))
 			[
-				TakeWidget()
+				UserSlateWidget.ToSharedRef()
 			];
 	}
 	else
@@ -361,7 +411,7 @@ TSharedRef<SWidget> UUserWidget::MakeViewportWidget(bool bAbsoluteLayout, bool b
 			.HAlign(HorizontalAlignment)
 			.VAlign(VerticalAlignment)
 			[
-				TakeWidget()
+				UserSlateWidget.ToSharedRef()
 			];
 
 		NewSlot.SizeParam = UWidget::ConvertSerializedSizeParamToRuntime(Size);
@@ -384,9 +434,10 @@ void UUserWidget::AddToViewport(bool bAbsoluteLayout, bool bModal, bool bShowCur
 {
 	if ( !FullScreenWidget.IsValid() )
 	{
-		TSharedRef<SWidget> RootWidget = MakeViewportWidget(bAbsoluteLayout, bModal, bShowCursor);
+		TSharedPtr<SWidget> OutUserSlateWidget;
+		TSharedRef<SWidget> RootWidget = MakeViewportWidget(bAbsoluteLayout, bModal, bShowCursor, OutUserSlateWidget);
 
-		TSharedRef<SViewportWidgetHost> WidgetHost = SNew(SViewportWidgetHost, (bool)bModal)
+		TSharedRef<SViewportWidgetHost> WidgetHost = SNew(SViewportWidgetHost, OutUserSlateWidget, bModal)
 			[
 				RootWidget
 			];
@@ -407,8 +458,8 @@ void UUserWidget::AddToViewport(bool bAbsoluteLayout, bool bModal, bool bShowCur
 				TWeakPtr<SViewport> GameViewportWidget = Viewport->GetGameViewport()->GetViewportWidget();
 				if ( GameViewportWidget.IsValid() )
 				{
-					GameViewportWidget.Pin()->SetWidgetToFocusOnActivate(RootWidget);
-					FSlateApplication::Get().SetKeyboardFocus(RootWidget);
+					GameViewportWidget.Pin()->SetWidgetToFocusOnActivate(OutUserSlateWidget);
+					FSlateApplication::Get().SetKeyboardFocus(OutUserSlateWidget);
 				}
 			}
 		}
@@ -419,16 +470,16 @@ void UUserWidget::RemoveFromViewport()
 {
 	if ( FullScreenWidget.IsValid() )
 	{
-		TSharedPtr<SViewportWidgetHost> RootWidget = FullScreenWidget.Pin();
+		TSharedPtr<SViewportWidgetHost> WidgetHost = FullScreenWidget.Pin();
 
 		// If this is a game world add the widget to the current worlds viewport.
 		UWorld* World = GetWorld();
 		if ( World && World->IsGameWorld() )
 		{
 			UGameViewportClient* Viewport = World->GetGameViewport();
-			Viewport->RemoveViewportWidgetContent(RootWidget.ToSharedRef());
+			Viewport->RemoveViewportWidgetContent(WidgetHost.ToSharedRef());
 
-			if ( RootWidget->IsModal() )
+			if ( WidgetHost->IsModal() )
 			{
 				TWeakPtr<SViewport> GameViewportWidget = Viewport->GetGameViewport()->GetViewportWidget();
 				if ( GameViewportWidget.IsValid() )
diff --git a/Engine/Source/Runtime/UMG/Private/WidgetBlueprintGeneratedClass.cpp b/Engine/Source/Runtime/UMG/Private/WidgetBlueprintGeneratedClass.cpp
--- a/Engine/Source/Runtime/UMG/Private/WidgetBlueprintGeneratedClass.cpp
+++ b/Engine/Source/Runtime/UMG/Private/WidgetBlueprintGeneratedClass.cpp
@@ -3,6 +3,7 @@
 #include "UMGPrivatePCH.h"
 #include "MovieScene.h"
 #include "WidgetAnimation.h"
+#include "Engine/InputDelegateBinding.h"
 
 #define LOCTEXT_NAMESPACE "UMG"
 
@@ -157,6 +158,10 @@ void UWidgetBlueprintGeneratedClass::InitializeWidget(UUserWidget* UserWidget) c
 
 		//TODO Add OnWidgetInitialized
 	}
+
+	// Create the input component so that we can simulate input to the blueprint in the forum users are accustomed to.
+	UserWidget->InputComponent = ConstructObject<UInputComponent>(UInputComponent::StaticClass(), UserWidget);
+	UInputDelegateBinding::BindInputDelegates(this, UserWidget->InputComponent);
 }
 
 #undef LOCTEXT_NAMESPACE
diff --git a/Engine/Source/Runtime/UMG/Public/Blueprint/UserWidget.h b/Engine/Source/Runtime/UMG/Public/Blueprint/UserWidget.h
--- a/Engine/Source/Runtime/UMG/Public/Blueprint/UserWidget.h
+++ b/Engine/Source/Runtime/UMG/Public/Blueprint/UserWidget.h
@@ -154,72 +154,72 @@ public:
 
 	//TODO UMG HitTest
 
-	UFUNCTION(BlueprintImplementableEvent, Category="User Interface")
+	UFUNCTION(BlueprintImplementableEvent, Category="User Interface | Painting")
 	void OnPaint(UPARAM(ref) FPaintContext& Context) const;
 
-	UFUNCTION(BlueprintImplementableEvent, Category="User Interface")
+	UFUNCTION(BlueprintImplementableEvent, Category="Keyboard")
 	FSReply OnKeyboardFocusReceived(FGeometry MyGeometry, FKeyboardFocusEvent InKeyboardFocusEvent);
-	UFUNCTION(BlueprintImplementableEvent, Category="User Interface")
+	UFUNCTION(BlueprintImplementableEvent, Category="Keyboard")
 	void OnKeyboardFocusLost(FKeyboardFocusEvent InKeyboardFocusEvent);
 	//UFUNCTION(BlueprintImplementableEvent, Category="User Interface")
 	//void OnKeyboardFocusChanging(FWeakWidgetPath PreviousFocusPath, FWidgetPath NewWidgetPath);
 
-	UFUNCTION(BlueprintImplementableEvent, Category="User Interface")
+	UFUNCTION(BlueprintImplementableEvent, Category="Keyboard")
 	FSReply OnKeyChar(FGeometry MyGeometry, FCharacterEvent InCharacterEvent);
-	UFUNCTION(BlueprintImplementableEvent, Category="User Interface")
+	UFUNCTION(BlueprintImplementableEvent, Category="Keyboard")
 	FSReply OnPreviewKeyDown(FGeometry MyGeometry, FKeyboardEvent InKeyboardEvent);
-	UFUNCTION(BlueprintImplementableEvent, Category="User Interface")
+	UFUNCTION(BlueprintImplementableEvent, Category="Keyboard")
 	FSReply OnKeyDown(FGeometry MyGeometry, FKeyboardEvent InKeyboardEvent);
-	UFUNCTION(BlueprintImplementableEvent, Category="User Interface")
+	UFUNCTION(BlueprintImplementableEvent, Category="Keyboard")
 	FSReply OnKeyUp(FGeometry MyGeometry, FKeyboardEvent InKeyboardEvent);
 
-	UFUNCTION(BlueprintImplementableEvent, Category="User Interface")
+	UFUNCTION(BlueprintImplementableEvent, Category="Mouse")
 	FSReply OnMouseButtonDown(FGeometry MyGeometry, const FPointerEvent& MouseEvent);
-	UFUNCTION(BlueprintImplementableEvent, Category="User Interface")
+	UFUNCTION(BlueprintImplementableEvent, Category="Mouse")
 	FSReply OnPreviewMouseButtonDown(FGeometry MyGeometry, const FPointerEvent& MouseEvent);
-	UFUNCTION(BlueprintImplementableEvent, Category="User Interface")
+	UFUNCTION(BlueprintImplementableEvent, Category="Mouse")
 	FSReply OnMouseButtonUp(FGeometry MyGeometry, const FPointerEvent& MouseEvent);
-	UFUNCTION(BlueprintImplementableEvent, Category="User Interface")
+	UFUNCTION(BlueprintImplementableEvent, Category="Mouse")
 	FSReply OnMouseMove(FGeometry MyGeometry, const FPointerEvent& MouseEvent);
-	UFUNCTION(BlueprintImplementableEvent, Category="User Interface")
+	UFUNCTION(BlueprintImplementableEvent, Category="Mouse")
 	void OnMouseEnter(FGeometry MyGeometry, const FPointerEvent& MouseEvent);
-	UFUNCTION(BlueprintImplementableEvent, Category="User Interface")
+	UFUNCTION(BlueprintImplementableEvent, Category="Mouse")
 	void OnMouseLeave(const FPointerEvent& MouseEvent);
-	UFUNCTION(BlueprintImplementableEvent, Category="User Interface")
+	UFUNCTION(BlueprintImplementableEvent, Category="Mouse")
 	FSReply OnMouseWheel(FGeometry MyGeometry, const FPointerEvent& MouseEvent);
-	//UFUNCTION(BlueprintImplementableEvent, Category="User Interface")
+	//UFUNCTION(BlueprintImplementableEvent, Category="Mouse")
 	//FCursorReply OnCursorQuery(FGeometry MyGeometry, const FPointerEvent& CursorEvent) const;
-	UFUNCTION(BlueprintImplementableEvent, Category="User Interface")
+	UFUNCTION(BlueprintImplementableEvent, Category="Mouse")
 	FSReply OnMouseButtonDoubleClick(FGeometry InMyGeometry, const FPointerEvent& InMouseEvent);
 
-	UFUNCTION(BlueprintImplementableEvent, Category="User Interface")
+	UFUNCTION(BlueprintImplementableEvent, Category="Drag and Drop")
 	void OnDragDetected(FGeometry MyGeometry, const FPointerEvent& MouseEvent, UDragDropOperation*& Operation);
 	
-	UFUNCTION(BlueprintImplementableEvent, Category="User Interface")
+	UFUNCTION(BlueprintImplementableEvent, Category="Drag and Drop")
 	void OnDragEnter(FGeometry MyGeometry, FPointerEvent PointerEvent, UDragDropOperation* Operation);
-	UFUNCTION(BlueprintImplementableEvent, Category="User Interface")
+	UFUNCTION(BlueprintImplementableEvent, Category="Drag and Drop")
 	void OnDragLeave(FPointerEvent PointerEvent, UDragDropOperation* Operation);
-	UFUNCTION(BlueprintImplementableEvent, Category="User Interface")
+	UFUNCTION(BlueprintImplementableEvent, Category="Drag and Drop")
 	bool OnDragOver(FGeometry MyGeometry, FPointerEvent PointerEvent, UDragDropOperation* Operation);
-	UFUNCTION(BlueprintImplementableEvent, Category="User Interface")
+	UFUNCTION(BlueprintImplementableEvent, Category="Drag and Drop")
 	bool OnDrop(FGeometry MyGeometry, FPointerEvent PointerEvent, UDragDropOperation* Operation);
 
-	UFUNCTION(BlueprintImplementableEvent, Category="User Interface")
+	UFUNCTION(BlueprintImplementableEvent, Category="Gamepad Input")
 	FSReply OnControllerButtonPressed(FGeometry MyGeometry, FControllerEvent ControllerEvent);
-	UFUNCTION(BlueprintImplementableEvent, Category="User Interface")
+	UFUNCTION(BlueprintImplementableEvent, Category="Gamepad Input")
 	FSReply OnControllerButtonReleased(FGeometry MyGeometry, FControllerEvent ControllerEvent);
-	UFUNCTION(BlueprintImplementableEvent, Category="User Interface")
+	UFUNCTION(BlueprintImplementableEvent, Category="Gamepad Input")
 	FSReply OnControllerAnalogValueChanged(FGeometry MyGeometry, FControllerEvent ControllerEvent);
 
-	UFUNCTION(BlueprintImplementableEvent, Category="User Interface")
+	UFUNCTION(BlueprintImplementableEvent, Category="Touch Input")
 	FSReply OnTouchGesture(FGeometry MyGeometry, const FPointerEvent& GestureEvent);
-	UFUNCTION(BlueprintImplementableEvent, Category="User Interface")
+	UFUNCTION(BlueprintImplementableEvent, Category="Touch Input")
 	FSReply OnTouchStarted(FGeometry MyGeometry, const FPointerEvent& InTouchEvent);
-	UFUNCTION(BlueprintImplementableEvent, Category="User Interface")
+	UFUNCTION(BlueprintImplementableEvent, Category="Touch Input")
 	FSReply OnTouchMoved(FGeometry MyGeometry, const FPointerEvent& InTouchEvent);
-	UFUNCTION(BlueprintImplementableEvent, Category="User Interface")
+	UFUNCTION(BlueprintImplementableEvent, Category="Touch Input")
 	FSReply OnTouchEnded(FGeometry MyGeometry, const FPointerEvent& InTouchEvent);
-	UFUNCTION(BlueprintImplementableEvent, Category="User Interface")
+	UFUNCTION(BlueprintImplementableEvent, Category="Touch Input")
 	FSReply OnMotionDetected(FGeometry MyGeometry, FMotionEvent InMotionEvent);
 
 	UFUNCTION(BlueprintCallable, Category="User Interface|Drag and Drop")
@@ -250,7 +250,7 @@ public:
 	UWidget* GetWidgetHandle(TSharedRef<SWidget> InWidget);
 
 	/** Creates a fullscreen host widget, that wraps this widget. */
-	TSharedRef<SWidget> MakeViewportWidget(bool bAbsoluteLayout, bool bModal, bool bShowCursor);
+	TSharedRef<SWidget> MakeViewportWidget(bool bAbsoluteLayout, bool bModal, bool bShowCursor, TSharedPtr<SWidget>& UserSlateWidget);
 
 	/** @returns The root UObject widget wrapper */
 	UWidget* GetRootWidgetComponent();
@@ -271,15 +271,20 @@ public:
 #endif
 
 public:
+
+	/** Called when the visibility changes. */
+	UPROPERTY(EditAnywhere, BlueprintReadWrite, AdvancedDisplay, Category="Behavior")
+	bool bSupportsKeyboardFocus;
+
 	/** Called when the visibility changes. */
 	UPROPERTY(BlueprintAssignable)
 	FOnVisibilityChangedEvent OnVisibilityChanged;
 
-	UPROPERTY(EditDefaultsOnly, Category=Appearance)
+	UPROPERTY()
 	FMargin Padding;
 
 	/** How much space this slot should occupy in the direction of the panel. */
-	UPROPERTY(EditDefaultsOnly, Category=Appearance)
+	UPROPERTY()
 	FSlateChildSize Size;
 
 	/**
@@ -294,7 +299,7 @@ public:
 	*
 	*  Note: FILL is NOT supported in absolute layout
 	*/
-	UPROPERTY(EditAnywhere, Category=Appearance)
+	UPROPERTY()
 	TEnumAsByte<EHorizontalAlignment> HorizontalAlignment;
 
 	/**
@@ -310,7 +315,7 @@ public:
 	*
 	*  Note: FILL is NOT supported in absolute layout
 	*/
-	UPROPERTY(EditAnywhere, Category=Appearance)
+	UPROPERTY()
 	TEnumAsByte<EVerticalAlignment> VerticalAlignment;
 
 	/** The components contained in this user widget. */
@@ -328,6 +333,10 @@ public:
 	/** List of sequence players to cache and clean up when safe */
 	UPROPERTY(Transient)
 	TArray<UUMGSequencePlayer*> StoppedSequencePlayers;
+
+	UPROPERTY(Transient)
+	UInputComponent* InputComponent;
+
 protected:
 	virtual TSharedRef<SWidget> RebuildWidget() override;
 
diff --git a/Engine/Source/Runtime/UMG/Public/Blueprint/WidgetTree.h b/Engine/Source/Runtime/UMG/Public/Blueprint/WidgetTree.h
--- a/Engine/Source/Runtime/UMG/Public/Blueprint/WidgetTree.h
+++ b/Engine/Source/Runtime/UMG/Public/Blueprint/WidgetTree.h
@@ -36,11 +36,6 @@ public:
 	template< class T >
 	T* ConstructWidget(TSubclassOf<UWidget> WidgetType)
 	{
-		// TODO UMG Editor only?
-		Modify();
-
-		// TODO UMG Don't have the widget tree responsible for construction and adding to the tree.
-
 		if ( WidgetType->IsChildOf(UUserWidget::StaticClass()) )
 		{
 			UUserWidget* Widget = ConstructObject<UUserWidget>(WidgetType, this);
@@ -65,6 +60,13 @@ public:
 		Super::PreSave();
 	}
 
+	virtual void PostLoad() override
+	{
+		//AllWidgets.Empty();
+
+		Super::PostLoad();
+	}
+
 public:
 	/** The root widget of the tree */
 	UPROPERTY()
diff --git a/Engine/Source/Runtime/UMG/Public/Components/CheckBox.h b/Engine/Source/Runtime/UMG/Public/Components/CheckBox.h
--- a/Engine/Source/Runtime/UMG/Public/Components/CheckBox.h
+++ b/Engine/Source/Runtime/UMG/Public/Components/CheckBox.h
@@ -6,7 +6,9 @@
 
 DECLARE_DYNAMIC_MULTICAST_DELEGATE_OneParam( FOnCheckBoxComponentStateChanged, bool, bIsChecked );
 
-/** Check box widget */
+/**
+ * Check box widget
+ */
 UCLASS(meta=( Category="Common" ), ClassGroup=UserInterface)
 class UMG_API UCheckBox : public UContentWidget
 {
diff --git a/Engine/Source/Runtime/UMG/Public/Components/NamedSlotInterface.h b/Engine/Source/Runtime/UMG/Public/Components/NamedSlotInterface.h
new file mode 100644
--- /dev/null
+++ b/Engine/Source/Runtime/UMG/Public/Components/NamedSlotInterface.h
@@ -0,0 +1,20 @@
+// Copyright 1998-2014 Epic Games, Inc. All Rights Reserved.
+
+#pragma once
+
+#include "NamedSlotInterface.generated.h"
+
+/**  */
+UINTERFACE(MinimalAPI, meta=( CannotImplementInterfaceInBlueprint ))
+class UNamedSlotInterface : public UInterface
+{
+	GENERATED_UINTERFACE_BODY()
+};
+
+class UMG_API INamedSlotInterface
+{
+	GENERATED_IINTERFACE_BODY()
+
+	/**  */
+	virtual void GetSlotNames(TArray<FName>& SlotNames) const = 0;
+};
diff --git a/Engine/Source/Runtime/UMG/Public/Slate/SObjectWidget.h b/Engine/Source/Runtime/UMG/Public/Slate/SObjectWidget.h
--- a/Engine/Source/Runtime/UMG/Public/Slate/SObjectWidget.h
+++ b/Engine/Source/Runtime/UMG/Public/Slate/SObjectWidget.h
@@ -30,6 +30,8 @@ class SObjectWidget : public SCompoundWidget, public FGCObject
 	//virtual bool OnHitTest(const FGeometry& MyGeometry, FVector2D InAbsoluteCursorPosition) override;
 	virtual int32 OnPaint(const FPaintArgs& Args, const FGeometry& AllottedGeometry, const FSlateRect& MyClippingRect, FSlateWindowElementList& OutDrawElements, int32 LayerId, const FWidgetStyle& InWidgetStyle, bool bParentEnabled) const override;
 
+	virtual bool SupportsKeyboardFocus() const override;
+
 	virtual FReply OnKeyboardFocusReceived(const FGeometry& MyGeometry, const FKeyboardFocusEvent& InKeyboardFocusEvent) override;
 	virtual void OnKeyboardFocusLost(const FKeyboardFocusEvent& InKeyboardFocusEvent) override;
 	virtual void OnKeyboardFocusChanging(const FWeakWidgetPath& PreviousFocusPath, const FWidgetPath& NewWidgetPath) override;
diff --git a/Engine/Source/Runtime/UMG/Public/UMG.h b/Engine/Source/Runtime/UMG/Public/UMG.h
--- a/Engine/Source/Runtime/UMG/Public/UMG.h
+++ b/Engine/Source/Runtime/UMG/Public/UMG.h
@@ -32,6 +32,8 @@
 
 #include "Widget.h"
 
+#include "NamedSlotInterface.h"
+
 #include "PanelWidget.h"
 #include "ContentWidget.h"
 
