Clang warns that a virtual destructor should be declared for classes with virtual functions and we compile with warnings-as-errors.

diff --git a/Engine/Source/Editor/UMGEditor/Private/Hierarchy/SHierarchyViewItem.h b/Engine/Source/Editor/UMGEditor/Private/Hierarchy/SHierarchyViewItem.h
--- a/Engine/Source/Editor/UMGEditor/Private/Hierarchy/SHierarchyViewItem.h
+++ b/Engine/Source/Editor/UMGEditor/Private/Hierarchy/SHierarchyViewItem.h
@@ -47,6 +47,8 @@ class FHierarchyRoot : public FHierarchyModel
 {
 public:
 	FHierarchyRoot(TSharedPtr<FWidgetBlueprintEditor> InBlueprintEditor);
+	
+	virtual ~FHierarchyRoot() {}
 
 	/* @returns the widget name to use for the tree item */
 	virtual FText GetText() const override;
@@ -74,6 +76,8 @@ class FHierarchyWidget : public FHierarchyModel
 public:
 	FHierarchyWidget(FWidgetReference InItem, TSharedPtr<FWidgetBlueprintEditor> InBlueprintEditor);
 
+	virtual ~FHierarchyWidget() {}
+	
 	/* @returns the widget name to use for the tree item */
 	virtual FText GetText() const override;
 
