UMG - Design surface now has zoom text.  The move handle and resize handles now work correctly at different zoom levels.

diff --git a/Engine/Source/Editor/UMGEditor/Private/Designer/SDesignSurface.cpp b/Engine/Source/Editor/UMGEditor/Private/Designer/SDesignSurface.cpp
--- a/Engine/Source/Editor/UMGEditor/Private/Designer/SDesignSurface.cpp
+++ b/Engine/Source/Editor/UMGEditor/Private/Designer/SDesignSurface.cpp
@@ -403,6 +403,16 @@ void SDesignSurface::ZoomToFit(bool bInstantZoom)
 	bTeleportInsteadOfScrollingWhenZoomingToFit = bInstantZoom;
 }
 
+FText SDesignSurface::GetZoomText() const
+{
+	return ZoomLevels->GetZoomText(ZoomLevel);
+}
+
+FSlateColor SDesignSurface::GetZoomTextColorAndOpacity() const
+{
+	return FLinearColor(1, 1, 1, 1.25f - ZoomLevelFade.GetLerp());
+}
+
 FSlateRect SDesignSurface::ComputeAreaBounds() const
 {
 	return FSlateRect(0, 0, 0, 0);
diff --git a/Engine/Source/Editor/UMGEditor/Private/Designer/SDesignSurface.h b/Engine/Source/Editor/UMGEditor/Private/Designer/SDesignSurface.h
--- a/Engine/Source/Editor/UMGEditor/Private/Designer/SDesignSurface.h
+++ b/Engine/Source/Editor/UMGEditor/Private/Designer/SDesignSurface.h
@@ -49,6 +49,9 @@ protected:
 
 	void ZoomToFit(bool bInstantZoom);
 
+	FText GetZoomText() const;
+	FSlateColor GetZoomTextColorAndOpacity() const;
+
 	FVector2D GetViewOffset() const;
 
 	virtual FSlateRect ComputeAreaBounds() const;
diff --git a/Engine/Source/Editor/UMGEditor/Private/Designer/SDesignerView.cpp b/Engine/Source/Editor/UMGEditor/Private/Designer/SDesignerView.cpp
--- a/Engine/Source/Editor/UMGEditor/Private/Designer/SDesignerView.cpp
+++ b/Engine/Source/Editor/UMGEditor/Private/Designer/SDesignerView.cpp
@@ -175,11 +175,28 @@ void SDesignerView::Construct(const FArguments& InArgs, TSharedPtr<FWidgetBluepr
 
 			// Top bar with buttons for changing the designer
 			+ SOverlay::Slot()
-			.HAlign(HAlign_Right)
+			.HAlign(HAlign_Fill)
 			.VAlign(VAlign_Top)
 			[
 				SNew(SHorizontalBox)
 
+				+ SHorizontalBox::Slot()
+				.AutoWidth()
+				.Padding(6, 2, 0, 0)
+				[
+					SNew(STextBlock)
+					.TextStyle(FEditorStyle::Get(), "Graph.ZoomText")
+					.Text(this, &SDesignerView::GetZoomText)
+					.ColorAndOpacity(this, &SDesignerView::GetZoomTextColorAndOpacity)
+				]
+
+				+ SHorizontalBox::Slot()
+				.FillWidth(1.0f)
+				[
+					SNew(SSpacer)
+					.Size(FVector2D(1, 1))
+				]
+
 				+ SHorizontalBox::Slot()
 				.AutoWidth()
 				[
@@ -242,6 +259,11 @@ SDesignerView::~SDesignerView()
 	}
 }
 
+float SDesignerView::GetPreviewScale() const
+{
+	return GetZoomAmount() * GetPreviewDPIScale();
+}
+
 FOptionalSize SDesignerView::GetPreviewWidth() const
 {
 	return (float)PreviewWidth;
@@ -400,7 +422,7 @@ UWidgetBlueprint* SDesignerView::GetBlueprint() const
 
 void SDesignerView::Register(TSharedRef<FDesignerExtension> Extension)
 {
-	Extension->Initialize(GetBlueprint());
+	Extension->Initialize(this, GetBlueprint());
 	DesignerExtensions.Add(Extension);
 }
 
@@ -554,12 +576,12 @@ FReply SDesignerView::OnMouseMove(const FGeometry& MyGeometry, const FPointerEve
 
 			if ( Preview->Slot )
 			{
-				Preview->Slot->Resize(DragDirections[CurrentHandle], MouseEvent.GetCursorDelta());
+				Preview->Slot->Resize(DragDirections[CurrentHandle], MouseEvent.GetCursorDelta() * ( 1.0f / GetPreviewScale() ));
 			}
 
 			if ( Template->Slot )
 			{
-				Template->Slot->Resize(DragDirections[CurrentHandle], MouseEvent.GetCursorDelta());
+				Template->Slot->Resize(DragDirections[CurrentHandle], MouseEvent.GetCursorDelta() * ( 1.0f / GetPreviewScale() ));
 			}
 
 			FBlueprintEditorUtils::MarkBlueprintAsModified(GetBlueprint());
diff --git a/Engine/Source/Editor/UMGEditor/Private/Designer/SDesignerView.h b/Engine/Source/Editor/UMGEditor/Private/Designer/SDesignerView.h
--- a/Engine/Source/Editor/UMGEditor/Private/Designer/SDesignerView.h
+++ b/Engine/Source/Editor/UMGEditor/Private/Designer/SDesignerView.h
@@ -48,6 +48,10 @@ public:
 
 	void Register(TSharedRef<FDesignerExtension> Extension);
 
+	// IUMGDesigner interface
+	float GetPreviewScale() const override;
+	// End of IUMGDesigner interface
+
 private:
 	/** The width of the preview screen for the UI */
 	FOptionalSize GetPreviewWidth() const;
diff --git a/Engine/Source/Editor/UMGEditor/Private/Extensions/CanvasSlotExtension.cpp b/Engine/Source/Editor/UMGEditor/Private/Extensions/CanvasSlotExtension.cpp
--- a/Engine/Source/Editor/UMGEditor/Private/Extensions/CanvasSlotExtension.cpp
+++ b/Engine/Source/Editor/UMGEditor/Private/Extensions/CanvasSlotExtension.cpp
@@ -133,7 +133,7 @@ FReply FCanvasSlotExtension::HandleDragging(const FGeometry& Geometry, const FPo
 {
 	if ( bDragging )
 	{
-		float InverseScale = ( 1.0f / Geometry.Scale );
+		float InverseScale = ( 1.0f / Designer->GetPreviewScale() );
 
 		for ( FWidgetReference& Selection : SelectionCache )
 		{
diff --git a/Engine/Source/Editor/UMGEditor/Private/Extensions/DesignerExtension.cpp b/Engine/Source/Editor/UMGEditor/Private/Extensions/DesignerExtension.cpp
--- a/Engine/Source/Editor/UMGEditor/Private/Extensions/DesignerExtension.cpp
+++ b/Engine/Source/Editor/UMGEditor/Private/Extensions/DesignerExtension.cpp
@@ -14,8 +14,9 @@ FDesignerExtension::FDesignerExtension()
 
 }
 
-void FDesignerExtension::Initialize(UWidgetBlueprint* InBlueprint)
+void FDesignerExtension::Initialize(IUMGDesigner* InDesigner, UWidgetBlueprint* InBlueprint)
 {
+	Designer = InDesigner;
 	Blueprint = InBlueprint;
 }
 
diff --git a/Engine/Source/Editor/UMGEditor/Public/DesignerExtension.h b/Engine/Source/Editor/UMGEditor/Public/DesignerExtension.h
--- a/Engine/Source/Editor/UMGEditor/Public/DesignerExtension.h
+++ b/Engine/Source/Editor/UMGEditor/Public/DesignerExtension.h
@@ -3,6 +3,7 @@
 #pragma once
 
 #include "WidgetReference.h"
+#include "IUMGDesigner.h"
 
 class UWidget;
 class SWidget;
@@ -87,7 +88,7 @@ public:
 	FDesignerExtension();
 
 	/** Initializes the designer extension, this is called the first time a designer extension is registered */
-	virtual void Initialize(UWidgetBlueprint* InBlueprint);
+	virtual void Initialize(IUMGDesigner* InDesigner, UWidgetBlueprint* InBlueprint);
 
 	virtual bool CanExtendSelection(const TArray< FWidgetReference >& Selection) const
 	{
@@ -118,6 +119,7 @@ protected:
 protected:
 	FName ExtensionId;
 	UWidgetBlueprint* Blueprint;
+	IUMGDesigner* Designer;
 
 	TArray< FWidgetReference > SelectionCache;
 
diff --git a/Engine/Source/Editor/UMGEditor/Public/IUMGDesigner.h b/Engine/Source/Editor/UMGEditor/Public/IUMGDesigner.h
--- a/Engine/Source/Editor/UMGEditor/Public/IUMGDesigner.h
+++ b/Engine/Source/Editor/UMGEditor/Public/IUMGDesigner.h
@@ -9,5 +9,6 @@
 class UMGEDITOR_API IUMGDesigner
 {
 public:
-
+	/** @return the effective preview scale after both the DPI and Zoom scale has been applied. */
+	virtual float GetPreviewScale() const = 0;
 };
