UMG - Removing the experimental text from the designer.

diff --git a/Engine/Source/Editor/UMGEditor/Private/Designer/SDesignerView.cpp b/Engine/Source/Editor/UMGEditor/Private/Designer/SDesignerView.cpp
--- a/Engine/Source/Editor/UMGEditor/Private/Designer/SDesignerView.cpp
+++ b/Engine/Source/Editor/UMGEditor/Private/Designer/SDesignerView.cpp
@@ -225,18 +225,6 @@ void SDesignerView::Construct(const FArguments& InArgs, TSharedPtr<FWidgetBluepr
 					]
 				]
 			]
-
-			// Bottom-right corner text indicating the experimental nature of umg
-			+ SOverlay::Slot()
-			.Padding(10)
-			.VAlign(VAlign_Bottom)
-			.HAlign(HAlign_Right)
-			[
-				SNew(STextBlock)
-				.Visibility( EVisibility::HitTestInvisible )
-				.TextStyle( FEditorStyle::Get(), "Graph.CornerText" )
-				.Text(LOCTEXT("ExperimentalWarning", "Experimental"))
-			]
 		]
 	);
 
