UMG - Adding a custom node to do something similar for drag/drop operations as was done for create widget.

diff --git a/Engine/Source/Editor/UMGEditor/Private/Nodes/K2Node_CreateDragDropOperation.cpp b/Engine/Source/Editor/UMGEditor/Private/Nodes/K2Node_CreateDragDropOperation.cpp
new file mode 100644
--- /dev/null
+++ b/Engine/Source/Editor/UMGEditor/Private/Nodes/K2Node_CreateDragDropOperation.cpp
@@ -0,0 +1,130 @@
+// Copyright 1998-2014 Epic Games, Inc. All Rights Reserved.
+
+#include "UMGEditorPrivatePCH.h"
+
+#include "KismetCompiler.h"
+#include "BlueprintNodeSpawner.h"
+#include "EditorCategoryUtils.h"
+#include "K2ActionMenuBuilder.h"
+
+#include "K2Node_CreateDragDropOperation.h"
+
+#define LOCTEXT_NAMESPACE "UMG"
+
+UK2Node_CreateDragDropOperation::UK2Node_CreateDragDropOperation(const class FPostConstructInitializeProperties& PCIP)
+	: Super(PCIP)
+{
+	NodeTooltip = LOCTEXT("NodeTooltip", "Creates a new drag drop operation").ToString();
+}
+
+FText UK2Node_CreateDragDropOperation::GetBaseNodeTitle() const
+{
+	return LOCTEXT("CreateWidget_BaseTitle", "Create Drag & Drop Operation");
+}
+
+FText UK2Node_CreateDragDropOperation::GetNodeTitleFormat() const
+{
+	return LOCTEXT("CreateWidget", "Create {ClassName}");
+}
+
+UClass* UK2Node_CreateDragDropOperation::GetClassPinBaseClass() const
+{
+	return UDragDropOperation::StaticClass();
+}
+
+void UK2Node_CreateDragDropOperation::GetMenuEntries(FGraphContextMenuBuilder& ContextMenuBuilder) const
+{
+	UK2Node_CreateDragDropOperation* TemplateNode = NewObject<UK2Node_CreateDragDropOperation>(GetTransientPackage(), GetClass());
+
+	const FString Category = TEXT("User Interface");
+	const FText   MenuDesc = LOCTEXT("CreateDragDropOperationMenuOption", "Create Drag & Drop Operation...");
+	const FString Tooltip  = TEXT("Create a new UI drag and drop operation.  Use inside of OnDragDetected.");
+
+	TSharedPtr<FEdGraphSchemaAction_K2NewNode> NodeAction = FK2ActionMenuBuilder::AddNewNodeAction(ContextMenuBuilder, Category, MenuDesc, Tooltip);
+	NodeAction->NodeTemplate = TemplateNode;
+}
+
+void UK2Node_CreateDragDropOperation::AllocateDefaultPins()
+{
+	Super::AllocateDefaultPins();
+
+	UEdGraphPin* ClassPin = GetClassPin();
+	if ( ClassPin->DefaultObject == NULL )
+	{
+		ClassPin->DefaultObject = UDragDropOperation::StaticClass();
+
+		UClass* UseSpawnClass = GetClassToSpawn();
+		if ( UseSpawnClass != NULL )
+		{
+			CreatePinsForClass(UseSpawnClass);
+		}
+	}
+}
+
+void UK2Node_CreateDragDropOperation::ExpandNode(class FKismetCompilerContext& CompilerContext, UEdGraph* SourceGraph)
+{
+	Super::ExpandNode(CompilerContext, SourceGraph);
+
+	if ( CompilerContext.bIsFullCompile )
+	{
+		static FName Create_FunctionName = GET_FUNCTION_NAME_CHECKED(UWidgetBlueprintLibrary, CreateDragDropOperation);
+		static FString OperationClass_ParamName = FString(TEXT("OperationClass"));
+
+		UK2Node_CreateDragDropOperation* CreateOpNode = this;
+		UEdGraphPin* SpawnNodeExec = CreateOpNode->GetExecPin();
+		UEdGraphPin* SpawnClassPin = CreateOpNode->GetClassPin();
+		UEdGraphPin* SpawnNodeThen = CreateOpNode->GetThenPin();
+		UEdGraphPin* SpawnNodeResult = CreateOpNode->GetResultPin();
+
+		UClass* SpawnClass = ( SpawnClassPin != NULL ) ? Cast<UClass>(SpawnClassPin->DefaultObject) : NULL;
+		//if ( ( 0 == SpawnClassPin->LinkedTo.Num() ) && ( NULL == SpawnClass ) )
+		//{
+		//	CompilerContext.MessageLog.Error(*LOCTEXT("CreateWidgetNodeMissingClass_Error", "Create Drag/Drop node @@ must have a class specified.").ToString(), CreateOpNode);
+		//	// we break exec links so this is the only error we get, don't want the CreateWidget node being considered and giving 'unexpected node' type warnings
+		//	CreateOpNode->BreakAllNodeLinks();
+		//	return;
+		//}
+
+		//////////////////////////////////////////////////////////////////////////
+		// create 'UWidgetBlueprintLibrary::CreateDragDropOperation' call node
+		UK2Node_CallFunction* CallCreateNode = CompilerContext.SpawnIntermediateNode<UK2Node_CallFunction>(CreateOpNode, SourceGraph);
+		CallCreateNode->FunctionReference.SetExternalMember(Create_FunctionName, UWidgetBlueprintLibrary::StaticClass());
+		CallCreateNode->AllocateDefaultPins();
+
+		UEdGraphPin* CallCreateExec = CallCreateNode->GetExecPin();
+		UEdGraphPin* CallCreateOperationClassPin = CallCreateNode->FindPinChecked(OperationClass_ParamName);
+		UEdGraphPin* CallCreateResult = CallCreateNode->GetReturnValuePin();
+
+		// Move 'exec' connection from create widget node to 'UWidgetBlueprintLibrary::CreateDragDropOperation'
+		CompilerContext.MovePinLinksToIntermediate(*SpawnNodeExec, *CallCreateExec);
+
+		if ( SpawnClassPin->LinkedTo.Num() > 0 )
+		{
+			// Copy the 'blueprint' connection from the spawn node to 'UWidgetBlueprintLibrary::CreateDragDropOperation'
+			CompilerContext.MovePinLinksToIntermediate(*SpawnClassPin, *CallCreateOperationClassPin);
+		}
+		else
+		{
+			// Copy blueprint literal onto 'UWidgetBlueprintLibrary::CreateDragDropOperation' call 
+			CallCreateOperationClassPin->DefaultObject = SpawnClass;
+		}
+
+		// Move result connection from spawn node to 'UWidgetBlueprintLibrary::CreateDragDropOperation'
+		CallCreateResult->PinType = SpawnNodeResult->PinType; // Copy type so it uses the right actor subclass
+		CompilerContext.MovePinLinksToIntermediate(*SpawnNodeResult, *CallCreateResult);
+
+		//////////////////////////////////////////////////////////////////////////
+		// create 'set var' nodes
+
+		// Get 'result' pin from 'begin spawn', this is the actual actor we want to set properties on
+		UEdGraphPin* LastThen = FKismetCompilerUtilities::GenerateAssignmentNodes(CompilerContext, SourceGraph, CallCreateNode, CreateOpNode, CallCreateResult, GetClassToSpawn());
+
+		// Move 'then' connection from create widget node to the last 'then'
+		CompilerContext.MovePinLinksToIntermediate(*SpawnNodeThen, *LastThen);
+
+		// Break any links to the expanded node
+		CreateOpNode->BreakAllNodeLinks();
+	}
+}
+
+#undef LOCTEXT_NAMESPACE
diff --git a/Engine/Source/Editor/UMGEditor/Private/Nodes/K2Node_CreateDragDropOperation.h b/Engine/Source/Editor/UMGEditor/Private/Nodes/K2Node_CreateDragDropOperation.h
new file mode 100644
--- /dev/null
+++ b/Engine/Source/Editor/UMGEditor/Private/Nodes/K2Node_CreateDragDropOperation.h
@@ -0,0 +1,29 @@
+// Copyright 1998-2014 Epic Games, Inc. All Rights Reserved.
+
+#pragma once
+
+#include "K2Node.h"
+#include "K2Node_CreateDragDropOperation.generated.h"
+
+UCLASS()
+class UMGEDITOR_API UK2Node_CreateDragDropOperation : public UK2Node_ConstructObjectFromClass
+{
+	GENERATED_UCLASS_BODY()
+
+	// Begin UEdGraphNode interface.
+	virtual void AllocateDefaultPins() override;
+	virtual void ExpandNode(class FKismetCompilerContext& CompilerContext, UEdGraph* SourceGraph) override;
+	// End UEdGraphNode interface.
+
+	// Begin UK2Node interface
+	void GetMenuEntries(FGraphContextMenuBuilder& ContextMenuBuilder) const;
+	// End UK2Node interface.
+
+protected:
+	/** Gets the default node title when no class is selected */
+	virtual FText GetBaseNodeTitle() const;
+	/** Gets the node title when a class has been selected. */
+	virtual FText GetNodeTitleFormat() const;
+	/** Gets base class to use for the 'class' pin.  UObject by default. */
+	virtual UClass* GetClassPinBaseClass() const;
+};
diff --git a/Engine/Source/Editor/UMGEditor/Private/Nodes/K2Node_CreateWidget.cpp b/Engine/Source/Editor/UMGEditor/Private/Nodes/K2Node_CreateWidget.cpp
--- a/Engine/Source/Editor/UMGEditor/Private/Nodes/K2Node_CreateWidget.cpp
+++ b/Engine/Source/Editor/UMGEditor/Private/Nodes/K2Node_CreateWidget.cpp
@@ -59,7 +59,7 @@ void UK2Node_CreateWidget::GetMenuEntries(FGraphContextMenuBuilder& ContextMenuB
 {
 	UK2Node_CreateWidget* TemplateNode = NewObject<UK2Node_CreateWidget>(GetTransientPackage(), GetClass());
 
-	const FString Category = TEXT("");
+	const FString Category = TEXT("User Interface");
 	const FText   MenuDesc = LOCTEXT("CreateWidgetMenuOption", "Create Widget...");
 	const FString Tooltip  = TEXT("Create a new UI widget");
 
@@ -143,7 +143,7 @@ void UK2Node_CreateWidget::ExpandNode(class FKismetCompilerContext& CompilerCont
 		// Copy the 'Owning Player' connection from the spawn node to 'UWidgetBlueprintLibrary::Create'
 		CompilerContext.MovePinLinksToIntermediate(*SpawnOwningPlayerPin, *CallCreateOwningPlayerPin);
 
-		// Move result connection from spawn node to 'finish spawn'
+		// Move result connection from spawn node to 'UWidgetBlueprintLibrary::Create'
 		CallCreateResult->PinType = SpawnNodeResult->PinType; // Copy type so it uses the right actor subclass
 		CompilerContext.MovePinLinksToIntermediate(*SpawnNodeResult, *CallCreateResult);
 
diff --git a/Engine/Source/Runtime/UMG/Private/Slate/UMGDragDropOp.cpp b/Engine/Source/Runtime/UMG/Private/Slate/UMGDragDropOp.cpp
--- a/Engine/Source/Runtime/UMG/Private/Slate/UMGDragDropOp.cpp
+++ b/Engine/Source/Runtime/UMG/Private/Slate/UMGDragDropOp.cpp
@@ -60,6 +60,8 @@ TSharedRef<FUMGDragDropOp> FUMGDragDropOp::New(UDragDropOperation* InOperation,
 		Operation->DecoratorWidget = InOperation->DefaultDragVisual->TakeWidget();
 	}
 
+	Operation->DecoratorWidget->SlatePrepass();
+
 	Operation->Construct();
 
 	return Operation;
diff --git a/Engine/Source/Runtime/UMG/Private/UserWidget.cpp b/Engine/Source/Runtime/UMG/Private/UserWidget.cpp
--- a/Engine/Source/Runtime/UMG/Private/UserWidget.cpp
+++ b/Engine/Source/Runtime/UMG/Private/UserWidget.cpp
@@ -556,11 +556,6 @@ int32 UUserWidget::GetFullScreenZOrder() const
 	return ViewportZOrder;
 }
 
-UDragDropOperation* UUserWidget::CreateDragDropOperation(TSubclassOf<UDragDropOperation> Operation)
-{
-	return ConstructObject<UDragDropOperation>(Operation);
-}
-
 #if WITH_EDITOR
 
 const FSlateBrush* UUserWidget::GetEditorIcon()
diff --git a/Engine/Source/Runtime/UMG/Private/WidgetBlueprintLibrary.cpp b/Engine/Source/Runtime/UMG/Private/WidgetBlueprintLibrary.cpp
--- a/Engine/Source/Runtime/UMG/Private/WidgetBlueprintLibrary.cpp
+++ b/Engine/Source/Runtime/UMG/Private/WidgetBlueprintLibrary.cpp
@@ -27,6 +27,18 @@ UUserWidget* UWidgetBlueprintLibrary::Create(UObject* WorldContextObject, TSubcl
 	}
 }
 
+UDragDropOperation* UWidgetBlueprintLibrary::CreateDragDropOperation(TSubclassOf<UDragDropOperation> Operation)
+{
+	if ( Operation )
+	{
+		return ConstructObject<UDragDropOperation>(Operation);
+	}
+	else
+	{
+		return ConstructObject<UDragDropOperation>(UDragDropOperation::StaticClass());
+	}
+}
+
 void UWidgetBlueprintLibrary::SetFocusToGameViewport()
 {
 	FSlateApplication::Get().SetFocusToGameViewport();
diff --git a/Engine/Source/Runtime/UMG/Public/Blueprint/UserWidget.h b/Engine/Source/Runtime/UMG/Public/Blueprint/UserWidget.h
--- a/Engine/Source/Runtime/UMG/Public/Blueprint/UserWidget.h
+++ b/Engine/Source/Runtime/UMG/Public/Blueprint/UserWidget.h
@@ -222,9 +222,6 @@ public:
 	UFUNCTION(BlueprintImplementableEvent, Category="Touch Input")
 	FEventReply OnMotionDetected(FGeometry MyGeometry, FMotionEvent InMotionEvent);
 
-	UFUNCTION(BlueprintCallable, Category="User Interface|Drag and Drop")
-	UDragDropOperation* CreateDragDropOperation(TSubclassOf<UDragDropOperation> Operation);
-
 	//virtual bool OnVisualizeTooltip(const TSharedPtr<SWidget>& TooltipContent);
 
 	/**
diff --git a/Engine/Source/Runtime/UMG/Public/Blueprint/WidgetBlueprintLibrary.h b/Engine/Source/Runtime/UMG/Public/Blueprint/WidgetBlueprintLibrary.h
--- a/Engine/Source/Runtime/UMG/Public/Blueprint/WidgetBlueprintLibrary.h
+++ b/Engine/Source/Runtime/UMG/Public/Blueprint/WidgetBlueprintLibrary.h
@@ -13,6 +13,10 @@ class UWidgetBlueprintLibrary : public UBlueprintFunctionLibrary
 	UFUNCTION(BlueprintCallable, meta=( HidePin="WorldContextObject", DefaultToSelf="WorldContextObject", FriendlyName = "Create Widget", BlueprintInternalUseOnly = "true" ), Category="User Interface|Widget")
 	static class UUserWidget* Create(UObject* WorldContextObject, TSubclassOf<class UUserWidget> WidgetType, APlayerController* OwningPlayer);
 
+	UFUNCTION(BlueprintCallable, Category="User Interface|Drag and Drop", meta=( BlueprintInternalUseOnly = "true" ))
+	static UDragDropOperation* CreateDragDropOperation(TSubclassOf<UDragDropOperation> OperationClass);
+	
+
 	UFUNCTION(BlueprintCallable, Category="Focus")
 	static void SetFocusToGameViewport();
 
