- Updated resolution & zoom to fit buttons in UMG designer to match level viewport
- Updated View button in texture editor to match as well
- Resolution menu items show check box if the identical resolution is being used
- Resolution & AR shown in bottom left of UMG designer
- Resolution text fades like the zoom text when resolution changes
- Added config settings for default resolution and AR

diff --git a/Engine/Source/Editor/EditorStyle/Private/SlateEditorStyle.cpp b/Engine/Source/Editor/EditorStyle/Private/SlateEditorStyle.cpp
--- a/Engine/Source/Editor/EditorStyle/Private/SlateEditorStyle.cpp
+++ b/Engine/Source/Editor/EditorStyle/Private/SlateEditorStyle.cpp
@@ -5628,7 +5628,7 @@ void FSlateEditorStyle::FStyle::SetupAutomationStyles()
 
 void FSlateEditorStyle::FStyle::SetupUMGEditorStyles()
 {
-	Set("UMGEditor.ZoomToFit", new IMAGE_BRUSH("Icons/UMG/Fit_16x", Icon16x16));
+	Set("UMGEditor.ZoomToFit", new IMAGE_BRUSH("Icons/UMG/Fit_16x", Icon16x16, FLinearColor(.05f, .05f, .05f, 1.f)));
 
 	Set("UMGEditor.AnchorCenter", new IMAGE_BRUSH("Icons/umg_anchor_center", Icon16x16));
 	Set("UMGEditor.AnchorTopBottom", new IMAGE_BRUSH("Icons/umg_anchor_top_bottom", FVector2D(16, 32)));
diff --git a/Engine/Source/Editor/TextureEditor/Private/Widgets/STextureEditorViewportToolbar.cpp b/Engine/Source/Editor/TextureEditor/Private/Widgets/STextureEditorViewportToolbar.cpp
--- a/Engine/Source/Editor/TextureEditor/Private/Widgets/STextureEditorViewportToolbar.cpp
+++ b/Engine/Source/Editor/TextureEditor/Private/Widgets/STextureEditorViewportToolbar.cpp
@@ -28,7 +28,7 @@ void STextureEditorViewportToolbar::Construct( const FArguments& InArgs, const T
 							.ClickMethod(EButtonClickMethod::MouseDown)
 							.ContentPadding(FMargin(5.0f, 2.0f))
 							.VAlign(VAlign_Center)
-							.ButtonStyle(FEditorStyle::Get(), "EditorViewportToolBar.MenuButton")
+							.ButtonStyle(FEditorStyle::Get(), "ViewPortMenu.Button")
 							.OnClicked(this, &STextureEditorViewportToolbar::HandleViewOptionsMenuButtonClicked)
 							[
 								SNew(SHorizontalBox)
@@ -37,15 +37,19 @@ void STextureEditorViewportToolbar::Construct( const FArguments& InArgs, const T
 									.AutoWidth()
 									[
 										SNew(STextBlock)
-											.Text(LOCTEXT("ViewButtonText", "View"))
+										.Text(LOCTEXT("ViewButtonText", "View"))
+										.TextStyle(FEditorStyle::Get(), "ViewportMenu.Label")
 									]
 
 								+ SHorizontalBox::Slot()
 									.AutoWidth()
-									.Padding(4.0f, 0.0f, 0.0f, 0.0f)
+									.HAlign(HAlign_Center)
+									.VAlign(VAlign_Center)
+									.Padding(2.0f)
 									[
 										SNew(SImage)
-											.Image(FEditorStyle::GetBrush("EditorViewportToolBar.MenuDropdown"))
+										.Image(FEditorStyle::GetBrush("ComboButton.Arrow"))
+										.ColorAndOpacity(FLinearColor::Black)
 									]
 							]
 					]
diff --git a/Engine/Source/Editor/UMGEditor/Private/Designer/SDesignerView.cpp b/Engine/Source/Editor/UMGEditor/Private/Designer/SDesignerView.cpp
--- a/Engine/Source/Editor/UMGEditor/Private/Designer/SDesignerView.cpp
+++ b/Engine/Source/Editor/UMGEditor/Private/Designer/SDesignerView.cpp
@@ -91,6 +91,11 @@ static bool LocateWidgetsUnderCursor_Helper(FArrangedWidget& Candidate, FVector2
 /////////////////////////////////////////////////////
 // SDesignerView
 
+const FString SDesignerView::ConfigSectionName = "UMGEditor.Designer";
+const uint32 SDesignerView::DefaultResolutionWidth = 1280;
+const uint32 SDesignerView::DefaultResolutionHeight = 720;
+const FString SDesignerView::DefaultAspectRatio = "16:9";
+
 void SDesignerView::Construct(const FArguments& InArgs, TSharedPtr<FWidgetBlueprintEditor> InBlueprintEditor)
 {
 	ScopedTransaction = NULL;
@@ -101,11 +106,10 @@ void SDesignerView::Construct(const FArguments& InArgs, TSharedPtr<FWidgetBluepr
 	BlueprintEditor = InBlueprintEditor;
 	CurrentHandle = DH_NONE;
 
-	//TODO UMG Store as a setting
-	//PreviewWidth = 1920;
-	//PreviewHeight = 1080;
-	PreviewWidth = 1280;
-	PreviewHeight = 720;
+	SetStartupResolution();
+
+	ResolutionTextFade = FCurveSequence(0.0f, 1.0f);
+	ResolutionTextFade.Play();
 
 	HoverTime = 0;
 
@@ -199,9 +203,10 @@ void SDesignerView::Construct(const FArguments& InArgs, TSharedPtr<FWidgetBluepr
 
 				+ SHorizontalBox::Slot()
 				.AutoWidth()
+				.Padding(5.0f)
 				[
 					SNew(SButton)
-					.ButtonStyle(FEditorStyle::Get(), "ToggleButton")
+					.ButtonStyle(FEditorStyle::Get(), "ViewportMenu.Button")
 					.ToolTipText(LOCTEXT("ZoomToFit_ToolTip", "Zoom To Fit"))
 					.OnClicked(this, &SDesignerView::HandleZoomToFitClicked)
 					[
@@ -212,19 +217,37 @@ void SDesignerView::Construct(const FArguments& InArgs, TSharedPtr<FWidgetBluepr
 				
 				+ SHorizontalBox::Slot()
 				.AutoWidth()
+				.Padding(5.0f)
 				[
 					SNew(SComboButton)
-					.ButtonStyle(FEditorStyle::Get(), "PropertyEditor.AssetComboStyle")
-					.ForegroundColor(FEditorStyle::GetColor("PropertyEditor.AssetName.ColorAndOpacity"))
+					.ButtonStyle(FEditorStyle::Get(), "ViewportMenu.Button")
+					.ForegroundColor(FLinearColor::Black)
 					.OnGetMenuContent(this, &SDesignerView::GetAspectMenu)
 					.ContentPadding(2.0f)
 					.ButtonContent()
 					[
 						SNew(STextBlock)
-						.Text(LOCTEXT("AspectRatio", "Aspect Ratio"))
+						.Text(LOCTEXT("Resolution", "Resolution"))
+						.TextStyle(FEditorStyle::Get(), "ViewportMenu.Label")
 					]
 				]
 			]
+			// Bottom bar to show current resolution & AR
+			+ SOverlay::Slot()
+			.HAlign(HAlign_Fill)
+			.VAlign(VAlign_Bottom)
+			[
+				SNew(SHorizontalBox)
+				+ SHorizontalBox::Slot()
+				.AutoWidth()
+				.Padding(6, 0, 0, 2)
+				[
+					SNew(STextBlock)
+					.TextStyle(FEditorStyle::Get(), "Graph.ZoomText")
+					.Text(this, &SDesignerView::GetCurrentResolutionText)
+					.ColorAndOpacity(this, &SDesignerView::GetResolutionTextColorAndOpacity)
+				]
+			]
 		]
 	);
 
@@ -247,6 +270,29 @@ SDesignerView::~SDesignerView()
 	}
 }
 
+void SDesignerView::SetStartupResolution()
+{
+	// Use previously set resolution (or create new entries using default values)
+	// Width
+	if (!GConfig->GetInt(*ConfigSectionName, TEXT("PreviewWidth"), PreviewWidth, GEditorUserSettingsIni))
+	{
+		GConfig->SetInt(*ConfigSectionName, TEXT("PreviewWidth"), DefaultResolutionWidth, GEditorUserSettingsIni);
+		PreviewWidth = DefaultResolutionWidth;
+	}
+	// Height
+	if (!GConfig->GetInt(*ConfigSectionName, TEXT("PreviewHeight"), PreviewHeight, GEditorUserSettingsIni))
+	{
+		GConfig->SetInt(*ConfigSectionName, TEXT("PreviewHeight"), DefaultResolutionHeight, GEditorUserSettingsIni);
+		PreviewHeight = DefaultResolutionHeight;
+	}
+	// Aspect Ratio
+	if (!GConfig->GetString(*ConfigSectionName, TEXT("PreviewAspectRatio"), PreviewAspectRatio, GEditorUserSettingsIni))
+	{
+		GConfig->SetString(*ConfigSectionName, TEXT("PreviewAspectRatio"), *DefaultAspectRatio, GEditorUserSettingsIni);
+		PreviewAspectRatio = DefaultAspectRatio;
+	}
+}
+
 float SDesignerView::GetPreviewScale() const
 {
 	return GetZoomAmount() * GetPreviewDPIScale();
@@ -1316,10 +1362,43 @@ FReply SDesignerView::OnDrop(const FGeometry& MyGeometry, const FDragDropEvent&
 	return FReply::Unhandled();
 }
 
-void SDesignerView::HandleCommonResolutionSelected(int32 Width, int32 Height)
+FText SDesignerView::GetResolutionText(int32 Width, int32 Height, const FString& AspectRatio) const
+{
+	FInternationalization& I18N = FInternationalization::Get();
+	FFormatNamedArguments Args;
+	Args.Add(TEXT("Width"), FText::AsNumber(Width, nullptr, I18N.GetInvariantCulture()));
+	Args.Add(TEXT("Height"), FText::AsNumber(Height, nullptr, I18N.GetInvariantCulture()));
+	Args.Add(TEXT("AspectRatio"), FText::FromString(AspectRatio));
+
+	return FText::Format(LOCTEXT("CommonResolutionFormat", "{Width} x {Height} ({AspectRatio})"), Args);
+}
+
+FText SDesignerView::GetCurrentResolutionText() const
+{
+	return GetResolutionText(PreviewWidth, PreviewHeight, PreviewAspectRatio);
+}
+
+FSlateColor SDesignerView::GetResolutionTextColorAndOpacity() const
+{
+	return FLinearColor(1, 1, 1, 1.25f - ResolutionTextFade.GetLerp());
+}
+
+void SDesignerView::HandleOnCommonResolutionSelected(int32 Width, int32 Height, FString AspectRatio)
 {
 	PreviewWidth = Width;
 	PreviewHeight = Height;
+	PreviewAspectRatio = AspectRatio;
+
+	GConfig->SetInt(*ConfigSectionName, TEXT("PreviewWidth"), Width, GEditorUserSettingsIni);
+	GConfig->SetInt(*ConfigSectionName, TEXT("PreviewHeight"), Height, GEditorUserSettingsIni);
+	GConfig->SetString(*ConfigSectionName, TEXT("PreviewAspectRatio"), *AspectRatio, GEditorUserSettingsIni);
+
+	ResolutionTextFade.Play();
+}
+
+bool SDesignerView::HandleIsCommonResolutionSelected(int32 Width, int32 Height) const
+{
+	return (Width == PreviewWidth) && (Height == PreviewHeight);
 }
 
 void SDesignerView::AddScreenResolutionSection(FMenuBuilder& MenuBuilder, const TArray<FPlayScreenResolution>& Resolutions, const FText& SectionName)
@@ -1328,16 +1407,13 @@ void SDesignerView::AddScreenResolutionSection(FMenuBuilder& MenuBuilder, const
 	{
 		for ( auto Iter = Resolutions.CreateConstIterator(); Iter; ++Iter )
 		{
-			FUIAction Action(FExecuteAction::CreateRaw(this, &SDesignerView::HandleCommonResolutionSelected, Iter->Width, Iter->Height));
-
-			FInternationalization& I18N = FInternationalization::Get();
-
-			FFormatNamedArguments Args;
-			Args.Add(TEXT("Width"), FText::AsNumber(Iter->Width, NULL, I18N.GetInvariantCulture()));
-			Args.Add(TEXT("Height"), FText::AsNumber(Iter->Height, NULL, I18N.GetInvariantCulture()));
-			Args.Add(TEXT("AspectRatio"), FText::FromString(Iter->AspectRatio));
+			// Actions for the resolution menu entry
+			FExecuteAction OnResolutionSelected = FExecuteAction::CreateRaw(this, &SDesignerView::HandleOnCommonResolutionSelected, Iter->Width, Iter->Height, Iter->AspectRatio);
+			FCanExecuteAction OnCanResolutionBeSelected = FCanExecuteAction::CreateRaw(&FSlateApplication::Get(), &FSlateApplication::IsNormalExecution);
+			FIsActionChecked OnIsResolutionSelected = FIsActionChecked::CreateRaw(this, &SDesignerView::HandleIsCommonResolutionSelected, Iter->Width, Iter->Height);
+			FUIAction Action(OnResolutionSelected, OnCanResolutionBeSelected, OnIsResolutionSelected);
 
-			MenuBuilder.AddMenuEntry(FText::FromString(Iter->Description), FText::Format(LOCTEXT("CommonResolutionFormat", "{Width} x {Height} (AspectRatio)"), Args), FSlateIcon(), Action);
+			MenuBuilder.AddMenuEntry(FText::FromString(Iter->Description), GetResolutionText(Iter->Width, Iter->Height, Iter->AspectRatio), FSlateIcon(), Action, NAME_None, EUserInterfaceActionType::Check);
 		}
 	}
 	MenuBuilder.EndSection();
diff --git a/Engine/Source/Editor/UMGEditor/Private/Designer/SDesignerView.h b/Engine/Source/Editor/UMGEditor/Private/Designer/SDesignerView.h
--- a/Engine/Source/Editor/UMGEditor/Private/Designer/SDesignerView.h
+++ b/Engine/Source/Editor/UMGEditor/Private/Designer/SDesignerView.h
@@ -53,6 +53,9 @@ public:
 	// End of IUMGDesigner interface
 
 private:
+	/** Establishes the resolution and aspect ratio to use on construction from config settings */
+	void SetStartupResolution();
+
 	/** The width of the preview screen for the UI */
 	FOptionalSize GetPreviewWidth() const;
 
@@ -86,8 +89,15 @@ private:
 
 	void CacheSelectedWidgetGeometry();
 
+	/** @return Formatted text for the given resolution params */
+	FText GetResolutionText(int32 Width, int32 Height, const FString& AspectRatio) const;
+
+	FText GetCurrentResolutionText() const;
+	FSlateColor GetResolutionTextColorAndOpacity() const;
+	
 	// Handles selecting a common screen resolution.
-	void HandleCommonResolutionSelected(int32 Width, int32 Height);
+	void HandleOnCommonResolutionSelected(int32 Width, int32 Height, FString AspectRatio);
+	bool HandleIsCommonResolutionSelected(int32 Width, int32 Height) const;
 	void AddScreenResolutionSection(FMenuBuilder& MenuBuilder, const TArray<FPlayScreenResolution>& Resolutions, const FText& SectionName);
 	TSharedRef<SWidget> GetAspectMenu();
 
@@ -98,6 +108,11 @@ private:
 	FReply HandleZoomToFitClicked();
 
 private:
+	static const FString ConfigSectionName;
+	static const uint32 DefaultResolutionWidth;
+	static const uint32 DefaultResolutionHeight;
+	static const FString DefaultAspectRatio;
+
 	enum DragHandle
 	{
 		DH_NONE = -1,
@@ -186,6 +201,11 @@ private:
 	/** The current slate widget being hovered, this is refreshed every frame in case it changes */
 	TWeakPtr<SWidget> HoveredSlateWidget;
 
+	// Resolution info
 	int32 PreviewWidth;
 	int32 PreviewHeight;
+	FString PreviewAspectRatio;
+
+	/** Curve to handle fading of the resolution */
+	FCurveSequence ResolutionTextFade;
 };
