UMG - Anchor point gizmo.  There is now an gizmo you can use to manipulate anchor points of widgets inside a canvas.  Speed improvements, widget references are a lot less heavy weight now.

diff --git a/Engine/Source/Editor/EditorStyle/Private/SlateEditorStyle.cpp b/Engine/Source/Editor/EditorStyle/Private/SlateEditorStyle.cpp
--- a/Engine/Source/Editor/EditorStyle/Private/SlateEditorStyle.cpp
+++ b/Engine/Source/Editor/EditorStyle/Private/SlateEditorStyle.cpp
@@ -5640,9 +5640,31 @@ void FSlateEditorStyle::FStyle::SetupUMGEditorStyles()
 {
 	Set("UMGEditor.ZoomToFit", new IMAGE_BRUSH("Icons/UMG/Fit_16x", Icon16x16, FLinearColor(.05f, .05f, .05f, 1.f)));
 
-	Set("UMGEditor.AnchorCenter", new IMAGE_BRUSH("Icons/umg_anchor_center", Icon16x16));
-	Set("UMGEditor.AnchorTopBottom", new IMAGE_BRUSH("Icons/umg_anchor_top_bottom", FVector2D(16, 32)));
-	Set("UMGEditor.AnchorLeftRight", new IMAGE_BRUSH("Icons/umg_anchor_left_right", FVector2D(32, 16)));
+	Set("UMGEditor.AnchorGizmo.Center", new IMAGE_BRUSH("Icons/UMG/AnchorGizmo/center", Icon16x16));
+	Set("UMGEditor.AnchorGizmo.Center.Hovered", new IMAGE_BRUSH("Icons/UMG/AnchorGizmo/center", Icon16x16, FLinearColor(0, 1, 0)));
+	
+	Set("UMGEditor.AnchorGizmo.Left", new IMAGE_BRUSH("Icons/UMG/AnchorGizmo/left", FVector2D(32, 16)));
+	Set("UMGEditor.AnchorGizmo.Left.Hovered", new IMAGE_BRUSH("Icons/UMG/AnchorGizmo/left", FVector2D(32, 16), FLinearColor(0, 1, 0)));
+	Set("UMGEditor.AnchorGizmo.Right", new IMAGE_BRUSH("Icons/UMG/AnchorGizmo/right", FVector2D(32, 16)));
+	Set("UMGEditor.AnchorGizmo.Right.Hovered", new IMAGE_BRUSH("Icons/UMG/AnchorGizmo/right", FVector2D(32, 16), FLinearColor(0, 1, 0)));
+	
+	Set("UMGEditor.AnchorGizmo.Top", new IMAGE_BRUSH("Icons/UMG/AnchorGizmo/top", FVector2D(16, 32)));
+	Set("UMGEditor.AnchorGizmo.Top.Hovered", new IMAGE_BRUSH("Icons/UMG/AnchorGizmo/top", FVector2D(16, 32), FLinearColor(0, 1, 0)));
+	Set("UMGEditor.AnchorGizmo.Bottom", new IMAGE_BRUSH("Icons/UMG/AnchorGizmo/bottom", FVector2D(16, 32)));
+	Set("UMGEditor.AnchorGizmo.Bottom.Hovered", new IMAGE_BRUSH("Icons/UMG/AnchorGizmo/bottom", FVector2D(16, 32), FLinearColor(0, 1, 0)));
+
+	Set("UMGEditor.AnchorGizmo.TopLeft", new IMAGE_BRUSH("Icons/UMG/AnchorGizmo/topleft", FVector2D(24, 24)));
+	Set("UMGEditor.AnchorGizmo.TopLeft.Hovered", new IMAGE_BRUSH("Icons/UMG/AnchorGizmo/topleft", FVector2D(24, 24), FLinearColor(0, 1, 0)));
+
+	Set("UMGEditor.AnchorGizmo.TopRight", new IMAGE_BRUSH("Icons/UMG/AnchorGizmo/topright", FVector2D(24, 24)));
+	Set("UMGEditor.AnchorGizmo.TopRight.Hovered", new IMAGE_BRUSH("Icons/UMG/AnchorGizmo/topright", FVector2D(24, 24), FLinearColor(0, 1, 0)));
+
+	Set("UMGEditor.AnchorGizmo.BottomLeft", new IMAGE_BRUSH("Icons/UMG/AnchorGizmo/bottomleft", FVector2D(24, 24)));
+	Set("UMGEditor.AnchorGizmo.BottomLeft.Hovered", new IMAGE_BRUSH("Icons/UMG/AnchorGizmo/bottomleft", FVector2D(24, 24), FLinearColor(0, 1, 0)));
+
+	Set("UMGEditor.AnchorGizmo.BottomRight", new IMAGE_BRUSH("Icons/UMG/AnchorGizmo/bottomright", FVector2D(24, 24)));
+	Set("UMGEditor.AnchorGizmo.BottomRight.Hovered", new IMAGE_BRUSH("Icons/UMG/AnchorGizmo/bottomright", FVector2D(24, 24), FLinearColor(0, 1, 0)));
+
 
 	Set("UMGEditor.TopLeft", new IMAGE_BRUSH("Icons/UMG/umg_anchors_top_left", Icon64x64));
 	Set("UMGEditor.TopCenter", new IMAGE_BRUSH("Icons/UMG/umg_anchors_top_center", Icon64x64));
diff --git a/Engine/Source/Editor/UMGEditor/Private/Designer/DesignTimeUtils.cpp b/Engine/Source/Editor/UMGEditor/Private/Designer/DesignTimeUtils.cpp
--- a/Engine/Source/Editor/UMGEditor/Private/Designer/DesignTimeUtils.cpp
+++ b/Engine/Source/Editor/UMGEditor/Private/Designer/DesignTimeUtils.cpp
@@ -60,4 +60,12 @@ bool FDesignTimeUtils::GetArrangedWidgetRelativeToParent(TSharedRef<const SWidge
 	return false;
 }
 
+void FDesignTimeUtils::GetArrangedWidgetRelativeToParent(FWidgetPath& WidgetPath, TSharedRef<const SWidget> Widget, TSharedRef<const SWidget> Parent, FArrangedWidget& ArrangedWidget)
+{
+	FArrangedWidget ArrangedDesigner = WidgetPath.FindArrangedWidget(Parent);
+
+	ArrangedWidget = WidgetPath.FindArrangedWidget(Widget);
+	ArrangedWidget.Geometry.AbsolutePosition -= ArrangedDesigner.Geometry.AbsolutePosition;
+}
+
 #undef LOCTEXT_NAMESPACE
diff --git a/Engine/Source/Editor/UMGEditor/Private/Designer/DesignTimeUtils.h b/Engine/Source/Editor/UMGEditor/Private/Designer/DesignTimeUtils.h
--- a/Engine/Source/Editor/UMGEditor/Private/Designer/DesignTimeUtils.h
+++ b/Engine/Source/Editor/UMGEditor/Private/Designer/DesignTimeUtils.h
@@ -13,4 +13,6 @@ public:
 	static bool GetArrangedWidgetRelativeToWindow(TSharedRef<SWidget> Widget, FArrangedWidget& ArrangedWidget);
 
 	static bool GetArrangedWidgetRelativeToParent(TSharedRef<const SWidget> Widget, TSharedRef<const SWidget> Parent, FArrangedWidget& ArrangedWidget);
+
+	static void GetArrangedWidgetRelativeToParent(FWidgetPath& WidgetPath, TSharedRef<const SWidget> Widget, TSharedRef<const SWidget> Parent, FArrangedWidget& ArrangedWidget);
 };
diff --git a/Engine/Source/Editor/UMGEditor/Private/Designer/SDesignerView.cpp b/Engine/Source/Editor/UMGEditor/Private/Designer/SDesignerView.cpp
--- a/Engine/Source/Editor/UMGEditor/Private/Designer/SDesignerView.cpp
+++ b/Engine/Source/Editor/UMGEditor/Private/Designer/SDesignerView.cpp
@@ -192,6 +192,7 @@ void SDesignerView::Construct(const FArguments& InArgs, TSharedPtr<FWidgetBluepr
 					.TextStyle(FEditorStyle::Get(), "Graph.ZoomText")
 					.Text(this, &SDesignerView::GetZoomText)
 					.ColorAndOpacity(this, &SDesignerView::GetZoomTextColorAndOpacity)
+					.Visibility(EVisibility::SelfHitTestInvisible)
 				]
 
 				+ SHorizontalBox::Slot()
@@ -397,13 +398,35 @@ FVector2D SDesignerView::GetExtensionPosition(TSharedRef<FDesignerSurfaceElement
 	const FVector2D TopLeft = CachedDesignerWidgetLocation;
 	const FVector2D Size = CachedDesignerWidgetSize * GetZoomAmount() * GetPreviewDPIScale();
 
+	// Calculate the parent position and size.  We use this information for calculating offsets.
+	FVector2D ParentPosition, ParentSize;
+	{
+		FWidgetReference ParentRef = BlueprintEditor.Pin()->GetReferenceFromTemplate(SelectedWidget.GetTemplate()->GetParent());
+
+		TSharedPtr<SWidget> PreviewSlateWidget = ParentRef.GetPreview()->GetCachedWidget();
+		if ( PreviewSlateWidget.IsValid() )
+		{
+			FWidgetPath WidgetPath;
+			SelectedWidgetPath.ToWidgetPath(WidgetPath);
+		
+			FArrangedWidget ArrangedWidget(SNullWidget::NullWidget, FGeometry());
+			FDesignTimeUtils::GetArrangedWidgetRelativeToParent(WidgetPath, PreviewSlateWidget.ToSharedRef(), AsShared(), ArrangedWidget);
+
+			ParentPosition = ArrangedWidget.Geometry.AbsolutePosition;
+			ParentSize = ArrangedWidget.Geometry.Size * GetZoomAmount() * GetPreviewDPIScale();
+		}
+	}
+
 	FVector2D FinalPosition(0, 0);
 
+	// Get the intial offset based on the location around the selected object.
 	switch ( ExtensionElement->GetLocation() )
 	{
 	case EExtensionLayoutLocation::Absolute:
+	{
+		FinalPosition = ParentPosition;
 		break;
-
+	}
 	case EExtensionLayoutLocation::TopLeft:
 		FinalPosition = TopLeft;
 		break;
@@ -435,6 +458,9 @@ FVector2D SDesignerView::GetExtensionPosition(TSharedRef<FDesignerSurfaceElement
 		break;
 	}
 
+	// Add the alignment offset
+	FinalPosition += ParentSize * ExtensionElement->GetAlignment();
+
 	return FinalPosition + ExtensionElement->GetOffset();
 }
 
@@ -516,7 +542,7 @@ FWidgetReference SDesignerView::GetWidgetAtCursor(const FGeometry& MyGeometry, c
 
 		if ( Preview )
 		{
-			return FWidgetReference::FromPreview(BlueprintEditor.Pin(), Preview);
+			return BlueprintEditor.Pin()->GetReferenceFromPreview(Preview);
 		}
 	}
 
@@ -698,8 +724,18 @@ void SDesignerView::CacheSelectedWidgetGeometry()
 	{
 		TSharedRef<SWidget> Widget = SelectedSlateWidget.Pin().ToSharedRef();
 
+		FWidgetPath WidgetPath;
+		if ( FSlateApplication::Get().GeneratePathToWidgetUnchecked(Widget, WidgetPath) )
+		{
+			SelectedWidgetPath = FWeakWidgetPath(WidgetPath);
+		}
+		else
+		{
+			SelectedWidgetPath = FWeakWidgetPath();
+		}
+
 		FArrangedWidget ArrangedWidget(SNullWidget::NullWidget, FGeometry());
-		FDesignTimeUtils::GetArrangedWidgetRelativeToParent(Widget, AsShared(), ArrangedWidget);
+		FDesignTimeUtils::GetArrangedWidgetRelativeToParent(WidgetPath, Widget, AsShared(), ArrangedWidget);
 
 		CachedDesignerWidgetLocation = ArrangedWidget.Geometry.AbsolutePosition;
 		CachedDesignerWidgetSize = ArrangedWidget.Geometry.Size;
@@ -775,79 +811,6 @@ int32 SDesignerView::OnPaint(const FPaintArgs& Args, const FGeometry& AllottedGe
 		);
 
 		DrawDragHandles(SelectionGeometry, MyClippingRect, OutDrawElements, LayerId, InWidgetStyle);
-
-		//TODO UMG Move this to a function.
-		// Draw anchors.
-		if ( SelectedWidget.IsValid() && SelectedWidget.GetTemplate()->Slot )
-		{
-			const float X = AllottedGeometry.AbsolutePosition.X;
-			const float Y = AllottedGeometry.AbsolutePosition.Y;
-			const float Width = AllottedGeometry.Size.X;
-			const float Height = AllottedGeometry.Size.Y;
-
-			float Selection_X = SelectionGeometry.DrawPosition.X;
-			float Selection_Y = SelectionGeometry.DrawPosition.Y;
-			float Selection_Width = SelectionGeometry.DrawSize.X;
-			float Selection_Height = SelectionGeometry.DrawSize.Y;
-
-			const FVector2D LeftRightSize = FVector2D(32, 16);
-			const FVector2D TopBottomSize = FVector2D(16, 32);
-
-			// @TODO UMG - Don't use the curve editors brushes
-			const FSlateBrush* KeyBrush = FEditorStyle::GetBrush("CurveEd.CurveKey");
-			FLinearColor KeyColor = FLinearColor(1.0f, 0.0f, 0.0f, 1.f);
-
-			// Normalized offset of the handle, 0, no offset, 1 move the handle completely out by its own size.
-			//const float HandleOffset = 1.0f;
-
-			if ( UCanvasPanelSlot* CanvasSlot = Cast<UCanvasPanelSlot>(SelectedWidget.GetTemplate()->Slot) )
-			{
-				// Left
-				FSlateDrawElement::MakeBox(
-					OutDrawElements,
-					++LayerId,
-					FPaintGeometry(FVector2D(X + Width * CanvasSlot->LayoutData.Anchors.Minimum.X - LeftRightSize.X * 0.5f, SelectionGeometry.DrawPosition.Y + Selection_Height * 0.5f - LeftRightSize.Y * 0.5f), LeftRightSize, 1.0f),
-					FEditorStyle::Get().GetBrush("UMGEditor.AnchorLeftRight"),
-					MyClippingRect,
-					ESlateDrawEffect::None,
-					KeyBrush->GetTint(InWidgetStyle) * InWidgetStyle.GetColorAndOpacityTint() * KeyColor
-					);
-
-				// Right
-				FSlateDrawElement::MakeBox(
-					OutDrawElements,
-					++LayerId,
-					FPaintGeometry(FVector2D(X + Width * CanvasSlot->LayoutData.Anchors.Maximum.X - LeftRightSize.X * 0.5f, SelectionGeometry.DrawPosition.Y + Selection_Height * 0.5f - LeftRightSize.Y * 0.5f), LeftRightSize, 1.0f),
-					FEditorStyle::Get().GetBrush("UMGEditor.AnchorLeftRight"),
-					MyClippingRect,
-					ESlateDrawEffect::None,
-					KeyBrush->GetTint(InWidgetStyle) * InWidgetStyle.GetColorAndOpacityTint() * KeyColor
-					);
-
-
-				// Top
-				FSlateDrawElement::MakeBox(
-					OutDrawElements,
-					++LayerId,
-					FPaintGeometry(FVector2D(SelectionGeometry.DrawPosition.X + Selection_Width * 0.5f - TopBottomSize.X * 0.5f, Y + Height * CanvasSlot->LayoutData.Anchors.Minimum.Y - TopBottomSize.Y * 0.5f), TopBottomSize, 1.0f),
-					FEditorStyle::Get().GetBrush("UMGEditor.AnchorTopBottom"),
-					MyClippingRect,
-					ESlateDrawEffect::None,
-					KeyBrush->GetTint(InWidgetStyle) * InWidgetStyle.GetColorAndOpacityTint() * KeyColor
-					);
-
-				// Bottom
-				FSlateDrawElement::MakeBox(
-					OutDrawElements,
-					++LayerId,
-					FPaintGeometry(FVector2D(SelectionGeometry.DrawPosition.X + Selection_Width * 0.5f - TopBottomSize.X * 0.5f, Y + Height * CanvasSlot->LayoutData.Anchors.Maximum.Y - TopBottomSize.Y * 0.5f), TopBottomSize, 1.0f),
-					FEditorStyle::Get().GetBrush("UMGEditor.AnchorTopBottom"),
-					MyClippingRect,
-					ESlateDrawEffect::None,
-					KeyBrush->GetTint(InWidgetStyle) * InWidgetStyle.GetColorAndOpacityTint() * KeyColor
-					);
-			}
-		}
 	}
 
 	return LayerId;
@@ -1186,7 +1149,7 @@ UWidget* SDesignerView::ProcessDropAndAddWidget(const FGeometry& MyGeometry, con
 
 			BP->WidgetTree->RootWidget = Widget;
 
-			SelectedWidget = FWidgetReference::FromTemplate(BlueprintEditor.Pin(), Widget);
+			SelectedWidget = BlueprintEditor.Pin()->GetReferenceFromTemplate(Widget);
 
 			DropPreviewParent = NULL;
 
diff --git a/Engine/Source/Editor/UMGEditor/Private/Designer/SDesignerView.h b/Engine/Source/Editor/UMGEditor/Private/Designer/SDesignerView.h
--- a/Engine/Source/Editor/UMGEditor/Private/Designer/SDesignerView.h
+++ b/Engine/Source/Editor/UMGEditor/Private/Designer/SDesignerView.h
@@ -201,11 +201,18 @@ private:
 	/** The current slate widget being hovered, this is refreshed every frame in case it changes */
 	TWeakPtr<SWidget> HoveredSlateWidget;
 
-	// Resolution info
+	/** The configured Width of the preview area, simulates screen size. */
 	int32 PreviewWidth;
+
+	/** The configured Height of the preview area, simulates screen size. */
 	int32 PreviewHeight;
+
+	// Resolution Info
 	FString PreviewAspectRatio;
 
 	/** Curve to handle fading of the resolution */
 	FCurveSequence ResolutionTextFade;
+
+	/**  */
+	FWeakWidgetPath SelectedWidgetPath;
 };
diff --git a/Engine/Source/Editor/UMGEditor/Private/Details/SWidgetDetailsView.cpp b/Engine/Source/Editor/UMGEditor/Private/Details/SWidgetDetailsView.cpp
--- a/Engine/Source/Editor/UMGEditor/Private/Details/SWidgetDetailsView.cpp
+++ b/Engine/Source/Editor/UMGEditor/Private/Details/SWidgetDetailsView.cpp
@@ -234,7 +234,7 @@ bool SWidgetDetailsView::HandleVerifyNameTextChanged(const FText& InText, FText&
 		bool bIsSameWidget = false;
 		if ( TemplateWidget != NULL )
 		{
-			if ( FWidgetReference::FromTemplate(BlueprintEditor.Pin(), TemplateWidget).GetPreview() != PreviewWidget )
+			if ( BlueprintEditor.Pin()->GetReferenceFromTemplate( TemplateWidget ).GetPreview() != PreviewWidget )
 			{
 				OutErrorMessage = LOCTEXT("ExistingWidgetName", "Existing Widget Name");
 				return false;
@@ -306,7 +306,7 @@ void SWidgetDetailsView::HandleIsVariableChanged(ESlateCheckBoxState::Type Check
 		UWidget* Widget = Cast<UWidget>(SelectedObjects[0].Get());
 		UWidgetBlueprint* Blueprint = BlueprintEditor.Pin()->GetWidgetBlueprintObj();
 		
-		FWidgetReference WidgetRef = FWidgetReference::FromTemplate(BPEditor, Blueprint->WidgetTree->FindWidget(Widget->GetName()));
+		FWidgetReference WidgetRef = BPEditor->GetReferenceFromTemplate(Blueprint->WidgetTree->FindWidget(Widget->GetName()));
 		if ( WidgetRef.IsValid() )
 		{
 			UWidget* Template = WidgetRef.GetTemplate();
diff --git a/Engine/Source/Editor/UMGEditor/Private/Extensions/CanvasSlotExtension.cpp b/Engine/Source/Editor/UMGEditor/Private/Extensions/CanvasSlotExtension.cpp
--- a/Engine/Source/Editor/UMGEditor/Private/Extensions/CanvasSlotExtension.cpp
+++ b/Engine/Source/Editor/UMGEditor/Private/Extensions/CanvasSlotExtension.cpp
@@ -16,7 +16,7 @@ static float DistancePointToLine2D(const FVector2D& LinePointA, const FVector2D&
 }
 
 FCanvasSlotExtension::FCanvasSlotExtension()
-	: bDragging(false)
+	: bMovingAnchor(false)
 {
 	ExtensionId = FName(TEXT("CanvasSlot"));
 }
@@ -40,12 +40,12 @@ void FCanvasSlotExtension::ExtendSelection(const TArray< FWidgetReference >& Sel
 
 	MoveHandle =
 		SNew(SBorder)
-		//.BorderImage(FEditorStyle::Get().GetBrush("ToolBar.Background"))
-		//.BorderBackgroundColor(FLinearColor(0,0,0))
+		.BorderImage(FEditorStyle::Get().GetBrush("NoBrush"))
 		.OnMouseButtonDown(this, &FCanvasSlotExtension::HandleBeginDrag)
 		.OnMouseButtonUp(this, &FCanvasSlotExtension::HandleEndDrag)
 		.OnMouseMove(this, &FCanvasSlotExtension::HandleDragging)
 		.Padding(FMargin(0))
+		.Cursor(EMouseCursor::GrabHand)
 		[
 			SNew(SImage)
 			.Image(FCoreStyle::Get().GetBrush("SoftwareCursor_CardinalCross"))
@@ -54,6 +54,151 @@ void FCanvasSlotExtension::ExtendSelection(const TArray< FWidgetReference >& Sel
 	MoveHandle->SlatePrepass();
 
 	SurfaceElements.Add(MakeShareable(new FDesignerSurfaceElement(MoveHandle.ToSharedRef(), EExtensionLayoutLocation::TopLeft, FVector2D(0, -(MoveHandle->GetDesiredSize().Y + 10)))));
+
+	AnchorWidgets.SetNumZeroed(EAnchorWidget::MAX_COUNT);
+	AnchorWidgets[EAnchorWidget::Center] = MakeAnchorWidget(EAnchorWidget::Center, 16, 16);
+
+	AnchorWidgets[EAnchorWidget::Left] = MakeAnchorWidget(EAnchorWidget::Left, 32, 16);
+	AnchorWidgets[EAnchorWidget::Right] = MakeAnchorWidget(EAnchorWidget::Right, 32, 16);
+	AnchorWidgets[EAnchorWidget::Top] = MakeAnchorWidget(EAnchorWidget::Top, 16, 32);
+	AnchorWidgets[EAnchorWidget::Bottom] = MakeAnchorWidget(EAnchorWidget::Bottom, 16, 32);
+
+	AnchorWidgets[EAnchorWidget::TopLeft] = MakeAnchorWidget(EAnchorWidget::TopLeft, 24, 24);
+	AnchorWidgets[EAnchorWidget::TopRight] = MakeAnchorWidget(EAnchorWidget::TopRight, 24, 24);
+	AnchorWidgets[EAnchorWidget::BottomLeft] = MakeAnchorWidget(EAnchorWidget::BottomLeft, 24, 24);
+	AnchorWidgets[EAnchorWidget::BottomRight] = MakeAnchorWidget(EAnchorWidget::BottomRight, 24, 24);
+
+
+	TArray<FVector2D> AnchorPos;
+	AnchorPos.SetNumZeroed(EAnchorWidget::MAX_COUNT);
+	AnchorPos[EAnchorWidget::Center] = FVector2D(-8, -8);
+	
+	AnchorPos[EAnchorWidget::Left] = FVector2D(-32, -8);
+	AnchorPos[EAnchorWidget::Right] = FVector2D(0, -8);
+	AnchorPos[EAnchorWidget::Top] = FVector2D(-8, -32);
+	AnchorPos[EAnchorWidget::Bottom] = FVector2D(-8, 0);
+
+	AnchorPos[EAnchorWidget::TopLeft] = FVector2D(-24, -24);
+	AnchorPos[EAnchorWidget::TopRight] = FVector2D(0, -24);
+	AnchorPos[EAnchorWidget::BottomLeft] = FVector2D(-24, 0);
+	AnchorPos[EAnchorWidget::BottomRight] = FVector2D(0, 0);
+
+	for ( int32 AnchorIndex = (int32)EAnchorWidget::MAX_COUNT - 1; AnchorIndex >= 0; AnchorIndex-- )
+	{
+		if ( !AnchorWidgets[AnchorIndex].IsValid() )
+		{
+			continue;
+		}
+
+		AnchorWidgets[AnchorIndex]->SlatePrepass();
+		TAttribute<FVector2D> AnchorAlignment = TAttribute<FVector2D>::Create(TAttribute<FVector2D>::FGetter::CreateSP(SharedThis(this), &FCanvasSlotExtension::GetAnchorAlignment, (EAnchorWidget::Type)AnchorIndex));
+		SurfaceElements.Add(MakeShareable(new FDesignerSurfaceElement(AnchorWidgets[AnchorIndex].ToSharedRef(), EExtensionLayoutLocation::Absolute, AnchorPos[AnchorIndex], AnchorAlignment)));
+	}
+}
+
+TSharedRef<SWidget> FCanvasSlotExtension::MakeAnchorWidget(EAnchorWidget::Type AnchorType, float Width, float Height)
+{
+	return SNew(SBorder)
+		.BorderImage(FEditorStyle::Get().GetBrush("NoBrush"))
+		.OnMouseButtonDown(this, &FCanvasSlotExtension::HandleAnchorBeginDrag, AnchorType)
+		.OnMouseButtonUp(this, &FCanvasSlotExtension::HandleAnchorEndDrag, AnchorType)
+		.OnMouseMove(this, &FCanvasSlotExtension::HandleAnchorDragging, AnchorType)
+		.Visibility(this, &FCanvasSlotExtension::GetAnchorVisibility, AnchorType)
+		.Padding(FMargin(0))
+		[
+			SNew(SBox)
+			.WidthOverride(Width)
+			.HeightOverride(Height)
+			.HAlign(HAlign_Fill)
+			.VAlign(VAlign_Fill)
+			[
+				SNew(SImage)
+				.Image(this, &FCanvasSlotExtension::GetAnchorBrush, AnchorType)
+			]
+		];
+}
+
+const FSlateBrush* FCanvasSlotExtension::GetAnchorBrush(EAnchorWidget::Type AnchorType) const
+{
+	switch ( AnchorType )
+	{
+	case EAnchorWidget::Center:
+		return AnchorWidgets[EAnchorWidget::Center]->IsHovered() ? FEditorStyle::Get().GetBrush("UMGEditor.AnchorGizmo.Center.Hovered") : FEditorStyle::Get().GetBrush("UMGEditor.AnchorGizmo.Center");
+	case EAnchorWidget::Left:
+		return AnchorWidgets[EAnchorWidget::Left]->IsHovered() ? FEditorStyle::Get().GetBrush("UMGEditor.AnchorGizmo.Left.Hovered") : FEditorStyle::Get().GetBrush("UMGEditor.AnchorGizmo.Left");
+	case EAnchorWidget::Right:
+		return AnchorWidgets[EAnchorWidget::Right]->IsHovered() ? FEditorStyle::Get().GetBrush("UMGEditor.AnchorGizmo.Right.Hovered") : FEditorStyle::Get().GetBrush("UMGEditor.AnchorGizmo.Right");
+	case EAnchorWidget::Top:
+		return AnchorWidgets[EAnchorWidget::Top]->IsHovered() ? FEditorStyle::Get().GetBrush("UMGEditor.AnchorGizmo.Top.Hovered") : FEditorStyle::Get().GetBrush("UMGEditor.AnchorGizmo.Top");
+	case EAnchorWidget::Bottom:
+		return AnchorWidgets[EAnchorWidget::Bottom]->IsHovered() ? FEditorStyle::Get().GetBrush("UMGEditor.AnchorGizmo.Bottom.Hovered") : FEditorStyle::Get().GetBrush("UMGEditor.AnchorGizmo.Bottom");
+	case EAnchorWidget::TopLeft:
+		return AnchorWidgets[EAnchorWidget::TopLeft]->IsHovered() ? FEditorStyle::Get().GetBrush("UMGEditor.AnchorGizmo.TopLeft.Hovered") : FEditorStyle::Get().GetBrush("UMGEditor.AnchorGizmo.TopLeft");
+	case EAnchorWidget::BottomRight:
+		return AnchorWidgets[EAnchorWidget::BottomRight]->IsHovered() ? FEditorStyle::Get().GetBrush("UMGEditor.AnchorGizmo.BottomRight.Hovered") : FEditorStyle::Get().GetBrush("UMGEditor.AnchorGizmo.BottomRight");
+	case EAnchorWidget::TopRight:
+		return AnchorWidgets[EAnchorWidget::TopRight]->IsHovered() ? FEditorStyle::Get().GetBrush("UMGEditor.AnchorGizmo.TopRight.Hovered") : FEditorStyle::Get().GetBrush("UMGEditor.AnchorGizmo.TopRight");
+	case EAnchorWidget::BottomLeft:
+		return AnchorWidgets[EAnchorWidget::BottomLeft]->IsHovered() ? FEditorStyle::Get().GetBrush("UMGEditor.AnchorGizmo.BottomLeft.Hovered") : FEditorStyle::Get().GetBrush("UMGEditor.AnchorGizmo.BottomLeft");
+	}
+
+	return FCoreStyle::Get().GetBrush("Selection");
+}
+
+EVisibility FCanvasSlotExtension::GetAnchorVisibility(EAnchorWidget::Type AnchorType) const
+{
+	for ( const FWidgetReference& Selection : SelectionCache )
+	{
+		UWidget* PreviewWidget = Selection.GetPreview();
+
+		UCanvasPanelSlot* PreviewCanvasSlot = CastChecked<UCanvasPanelSlot>(PreviewWidget->Slot);
+		switch ( AnchorType )
+		{
+		case EAnchorWidget::Center:
+			return PreviewCanvasSlot->LayoutData.Anchors.Minimum == PreviewCanvasSlot->LayoutData.Anchors.Maximum ? EVisibility::Visible : EVisibility::Collapsed;
+		case EAnchorWidget::Left:
+		case EAnchorWidget::Right:
+			return PreviewCanvasSlot->LayoutData.Anchors.Minimum.Y == PreviewCanvasSlot->LayoutData.Anchors.Maximum.Y ? EVisibility::Visible : EVisibility::Collapsed;
+		case EAnchorWidget::Top:
+		case EAnchorWidget::Bottom:
+			return PreviewCanvasSlot->LayoutData.Anchors.Minimum.X == PreviewCanvasSlot->LayoutData.Anchors.Maximum.X ? EVisibility::Visible : EVisibility::Collapsed;
+		}
+
+		return EVisibility::Visible;
+	}
+
+	return EVisibility::Collapsed;
+}
+
+FVector2D FCanvasSlotExtension::GetAnchorAlignment(EAnchorWidget::Type AnchorType) const
+{
+	for ( const FWidgetReference& Selection : SelectionCache )
+	{
+		UWidget* PreviewWidget = Selection.GetPreview();
+		UCanvasPanelSlot* PreviewCanvasSlot = CastChecked<UCanvasPanelSlot>(PreviewWidget->Slot);
+		
+		FVector2D Minimum = PreviewCanvasSlot->LayoutData.Anchors.Minimum;
+		FVector2D Maximum = PreviewCanvasSlot->LayoutData.Anchors.Maximum;
+		
+		switch ( AnchorType )
+		{
+		case EAnchorWidget::Center:
+		case EAnchorWidget::Left:
+		case EAnchorWidget::Top:
+		case EAnchorWidget::TopLeft:
+			return Minimum;
+		case EAnchorWidget::Right:
+		case EAnchorWidget::Bottom:
+		case EAnchorWidget::BottomRight:
+			return Maximum;
+		case EAnchorWidget::TopRight:
+			return  FVector2D(Maximum.X, Minimum.Y);
+		case EAnchorWidget::BottomLeft:
+			return  FVector2D(Minimum.X, Maximum.Y);
+		}
+	}
+
+	return FVector2D(0, 0);
 }
 
 bool FCanvasSlotExtension::GetCollisionSegmentsForSlot(UCanvasPanel* Canvas, int32 SlotIndex, TArray<FVector2D>& Segments)
@@ -109,10 +254,114 @@ void FCanvasSlotExtension::Paint(const TSet< FWidgetReference >& Selection, cons
 	PaintCollisionLines(Selection, AllottedGeometry, MyClippingRect, OutDrawElements, LayerId, InWidgetStyle, bParentEnabled);
 }
 
-FReply FCanvasSlotExtension::HandleBeginDrag(const FGeometry& Geometry, const FPointerEvent& Event)
+FReply FCanvasSlotExtension::HandleAnchorBeginDrag(const FGeometry& Geometry, const FPointerEvent& Event, EAnchorWidget::Type AnchorType)
+{
+	bMovingAnchor = true;
+	return FReply::Handled().CaptureMouse(AnchorWidgets[(int32)AnchorType].ToSharedRef());
+}
+
+FReply FCanvasSlotExtension::HandleAnchorEndDrag(const FGeometry& Geometry, const FPointerEvent& Event, EAnchorWidget::Type AnchorType)
+{
+	bMovingAnchor = false;
+	return FReply::Handled().ReleaseMouseCapture();
+}
+
+FReply FCanvasSlotExtension::HandleAnchorDragging(const FGeometry& Geometry, const FPointerEvent& Event, EAnchorWidget::Type AnchorType)
 {
-	bDragging = true;
+	if ( bMovingAnchor && !Event.GetCursorDelta().IsZero() )
+	{
+		float InverseScale = ( 1.0f / Designer->GetPreviewScale() );
+
+		FVector2D DeltaScaled = Event.GetCursorDelta() * InverseScale;
 
+		for ( FWidgetReference& Selection : SelectionCache )
+		{
+			UWidget* PreviewWidget = Selection.GetPreview();
+			if ( UCanvasPanel* Canvas = Cast<UCanvasPanel>(PreviewWidget->GetParent()) )
+			{
+				UCanvasPanelSlot* PreviewCanvasSlot = Cast<UCanvasPanelSlot>(PreviewWidget->Slot);
+
+				FGeometry Geometry;
+				if ( Canvas->GetGeometryForSlot(PreviewCanvasSlot, Geometry) )
+				{
+					// Calculate the default anchor offset, ie where would this control be laid out if no offset were provided.
+					FVector2D CanvasSize = Canvas->GetCanvasWidget()->GetCachedGeometry().Size;
+
+					FVector2D AnchorDelta = DeltaScaled / CanvasSize;
+
+					PreviewCanvasSlot->SaveBaseLayout();
+					switch ( AnchorType )
+					{
+					case EAnchorWidget::Center:
+						PreviewCanvasSlot->LayoutData.Anchors.Maximum += AnchorDelta;
+						PreviewCanvasSlot->LayoutData.Anchors.Minimum += AnchorDelta;
+						PreviewCanvasSlot->LayoutData.Anchors.Minimum.X = FMath::Clamp(PreviewCanvasSlot->LayoutData.Anchors.Minimum.X, 0.0f, 1.0f);
+						PreviewCanvasSlot->LayoutData.Anchors.Maximum.X = FMath::Clamp(PreviewCanvasSlot->LayoutData.Anchors.Maximum.X, 0.0f, 1.0f);
+						PreviewCanvasSlot->LayoutData.Anchors.Minimum.Y = FMath::Clamp(PreviewCanvasSlot->LayoutData.Anchors.Minimum.Y, 0.0f, 1.0f);
+						PreviewCanvasSlot->LayoutData.Anchors.Maximum.Y = FMath::Clamp(PreviewCanvasSlot->LayoutData.Anchors.Maximum.Y, 0.0f, 1.0f);
+						break;
+					}
+
+					switch ( AnchorType )
+					{
+					case EAnchorWidget::Left:
+					case EAnchorWidget::TopLeft:
+					case EAnchorWidget::BottomLeft:
+						PreviewCanvasSlot->LayoutData.Anchors.Minimum.X += AnchorDelta.X;
+						PreviewCanvasSlot->LayoutData.Anchors.Minimum.X = FMath::Clamp(PreviewCanvasSlot->LayoutData.Anchors.Minimum.X, 0.0f, PreviewCanvasSlot->LayoutData.Anchors.Maximum.X);
+						break;
+					}
+
+					switch ( AnchorType )
+					{
+					case EAnchorWidget::Right:
+					case EAnchorWidget::TopRight:
+					case EAnchorWidget::BottomRight:
+						PreviewCanvasSlot->LayoutData.Anchors.Maximum.X += AnchorDelta.X;
+						PreviewCanvasSlot->LayoutData.Anchors.Maximum.X = FMath::Clamp(PreviewCanvasSlot->LayoutData.Anchors.Maximum.X, PreviewCanvasSlot->LayoutData.Anchors.Minimum.X, 1.0f);
+						break;
+					}
+
+					switch ( AnchorType )
+					{
+					case EAnchorWidget::Top:
+					case EAnchorWidget::TopLeft:
+					case EAnchorWidget::TopRight:
+						PreviewCanvasSlot->LayoutData.Anchors.Minimum.Y += AnchorDelta.Y;
+						PreviewCanvasSlot->LayoutData.Anchors.Minimum.Y = FMath::Clamp(PreviewCanvasSlot->LayoutData.Anchors.Minimum.Y, 0.0f, PreviewCanvasSlot->LayoutData.Anchors.Maximum.Y);
+						break;
+					}
+
+					switch ( AnchorType )
+					{
+					case EAnchorWidget::Bottom:
+					case EAnchorWidget::BottomLeft:
+					case EAnchorWidget::BottomRight:
+						PreviewCanvasSlot->LayoutData.Anchors.Maximum.Y += AnchorDelta.Y;
+						PreviewCanvasSlot->LayoutData.Anchors.Maximum.Y = FMath::Clamp(PreviewCanvasSlot->LayoutData.Anchors.Maximum.Y, PreviewCanvasSlot->LayoutData.Anchors.Minimum.Y, 1.0f);
+						break;
+					}
+
+					PreviewCanvasSlot->RebaseLayout();
+
+
+					// Update the Template widget to match
+					UWidget* TemplateWidget = Selection.GetTemplate();
+					UCanvasPanelSlot* TemplateSlot = CastChecked<UCanvasPanelSlot>(TemplateWidget->Slot);
+
+					TemplateSlot->LayoutData = PreviewCanvasSlot->LayoutData;
+				}
+			};
+
+			return FReply::Handled();
+		}
+	}
+
+	return FReply::Unhandled();
+}
+
+FReply FCanvasSlotExtension::HandleBeginDrag(const FGeometry& Geometry, const FPointerEvent& Event)
+{
 	BeginTransaction(LOCTEXT("MoveWidget", "Move Widget"));
 
 	return FReply::Handled().CaptureMouse(MoveHandle.ToSharedRef());
@@ -120,8 +369,6 @@ FReply FCanvasSlotExtension::HandleBeginDrag(const FGeometry& Geometry, const FP
 
 FReply FCanvasSlotExtension::HandleEndDrag(const FGeometry& Geometry, const FPointerEvent& Event)
 {
-	bDragging = false;
-
 	EndTransaction();
 
 	FBlueprintEditorUtils::MarkBlueprintAsModified(Blueprint);
@@ -131,7 +378,7 @@ FReply FCanvasSlotExtension::HandleEndDrag(const FGeometry& Geometry, const FPoi
 
 FReply FCanvasSlotExtension::HandleDragging(const FGeometry& Geometry, const FPointerEvent& Event)
 {
-	if ( bDragging )
+	if ( MoveHandle->HasMouseCapture() )
 	{
 		float InverseScale = ( 1.0f / Designer->GetPreviewScale() );
 
@@ -201,7 +448,7 @@ void FCanvasSlotExtension::PaintCollisionLines(const TSet< FWidgetReference >& S
 			if ( UCanvasPanel* Canvas = CastChecked<UCanvasPanel>(CanvasSlot->Parent) )
 			{
 				//TODO UMG Only show guide lines when near them and dragging
-				if ( bDragging )
+				if ( MoveHandle->HasMouseCapture() )
 				{
 					FGeometry MyArrangedGeometry;
 					if ( !Canvas->GetGeometryForSlot(CanvasSlot, MyArrangedGeometry) )
diff --git a/Engine/Source/Editor/UMGEditor/Private/Extensions/CanvasSlotExtension.h b/Engine/Source/Editor/UMGEditor/Private/Extensions/CanvasSlotExtension.h
--- a/Engine/Source/Editor/UMGEditor/Private/Extensions/CanvasSlotExtension.h
+++ b/Engine/Source/Editor/UMGEditor/Private/Extensions/CanvasSlotExtension.h
@@ -4,6 +4,25 @@
 
 #include "DesignerExtension.h"
 
+/** Set of anchor widget types */
+namespace EAnchorWidget
+{
+	enum Type
+	{
+		Center,
+		Left,
+		Right,
+		Top,
+		Bottom,
+		TopLeft,
+		TopRight,
+		BottomLeft,
+		BottomRight,
+
+		MAX_COUNT
+	};
+}
+
 /**
  * The canvas slot extension provides design time widgets for widgets that are selected in the canvas.
  */
@@ -22,10 +41,20 @@ public:
 
 private:
 
+	FReply HandleAnchorBeginDrag(const FGeometry& Geometry, const FPointerEvent& Event, EAnchorWidget::Type AnchorType);
+	FReply HandleAnchorEndDrag(const FGeometry& Geometry, const FPointerEvent& Event, EAnchorWidget::Type AnchorType);
+	FReply HandleAnchorDragging(const FGeometry& Geometry, const FPointerEvent& Event, EAnchorWidget::Type AnchorType);
+
 	FReply HandleBeginDrag(const FGeometry& Geometry, const FPointerEvent& Event);
 	FReply HandleEndDrag(const FGeometry& Geometry, const FPointerEvent& Event);
 	FReply HandleDragging(const FGeometry& Geometry, const FPointerEvent& Event);
 
+	TSharedRef<SWidget> MakeAnchorWidget(EAnchorWidget::Type AnchorType, float Width, float Height);
+
+	const FSlateBrush* GetAnchorBrush(EAnchorWidget::Type AnchorType) const;
+	EVisibility GetAnchorVisibility(EAnchorWidget::Type AnchorType) const;
+	FVector2D GetAnchorAlignment(EAnchorWidget::Type AnchorType) const;
+
 	void MoveByAmount(FWidgetReference& WidgetRef, FVector2D Delta);
 
 	static bool GetCollisionSegmentsForSlot(UCanvasPanel* Canvas, int32 SlotIndex, TArray<FVector2D>& Segments);
@@ -36,7 +65,9 @@ private:
 
 private:
 
-	bool bDragging;
-
 	TSharedPtr<SBorder> MoveHandle;
+
+	TArray< TSharedPtr<SWidget> > AnchorWidgets;
+
+	bool bMovingAnchor;
 };
\ No newline at end of file
diff --git a/Engine/Source/Editor/UMGEditor/Private/Hierarchy/SHierarchyViewItem.cpp b/Engine/Source/Editor/UMGEditor/Private/Hierarchy/SHierarchyViewItem.cpp
--- a/Engine/Source/Editor/UMGEditor/Private/Hierarchy/SHierarchyViewItem.cpp
+++ b/Engine/Source/Editor/UMGEditor/Private/Hierarchy/SHierarchyViewItem.cpp
@@ -281,7 +281,7 @@ void FHierarchyRoot::GetChildren(TArray< TSharedPtr<FHierarchyModel> >& Children
 
 	if ( Blueprint->WidgetTree->RootWidget )
 	{
-		TSharedPtr<FHierarchyWidget> RootChild = MakeShareable(new FHierarchyWidget(FWidgetReference::FromTemplate(BPEd, Blueprint->WidgetTree->RootWidget), BPEd));
+		TSharedPtr<FHierarchyWidget> RootChild = MakeShareable(new FHierarchyWidget(BPEd->GetReferenceFromTemplate(Blueprint->WidgetTree->RootWidget), BPEd));
 		Children.Add(RootChild);
 	}
 }
@@ -451,7 +451,7 @@ void FHierarchyWidget::GetChildren(TArray< TSharedPtr<FHierarchyModel> >& Childr
 			UWidget* Child = Widget->GetChildAt(i);
 			if ( Child )
 			{
-				TSharedPtr<FHierarchyWidget> ChildItem = MakeShareable(new FHierarchyWidget(FWidgetReference::FromTemplate(BPEd, Child), BPEd));
+				TSharedPtr<FHierarchyWidget> ChildItem = MakeShareable(new FHierarchyWidget(BPEd->GetReferenceFromTemplate(Child), BPEd));
 				Children.Add(ChildItem);
 			}
 		}
diff --git a/Engine/Source/Editor/UMGEditor/Private/WidgetBlueprintEditor.cpp b/Engine/Source/Editor/UMGEditor/Private/WidgetBlueprintEditor.cpp
--- a/Engine/Source/Editor/UMGEditor/Private/WidgetBlueprintEditor.cpp
+++ b/Engine/Source/Editor/UMGEditor/Private/WidgetBlueprintEditor.cpp
@@ -214,25 +214,25 @@ void FWidgetBlueprintEditor::OnObjectsReplaced(const TMap<UObject*, UObject*>& R
 	// Objects being replaced could force a selection change
 	OnSelectedWidgetsChanging.Broadcast();
 
-	TSet<FWidgetReference> TempSelection;
-	for ( const FWidgetReference& WidgetRef : SelectedWidgets )
+	// Remove dead references and update references
+	for ( int32 HandleIndex = WidgetHandlePool.Num() - 1; HandleIndex >= 0; HandleIndex-- )
 	{
-		UWidget* OldTemplate = WidgetRef.GetTemplate();
+		TSharedPtr<FWidgetHandle> Ref = WidgetHandlePool[HandleIndex].Pin();
 
-		UObject* const* NewObject = ReplacementMap.Find(OldTemplate);
-		if ( NewObject )
+		if ( Ref.IsValid() )
 		{
-			UWidget* NewTemplate = Cast<UWidget>(*NewObject);
-			TempSelection.Add(FWidgetReference::FromTemplate(SharedThis(this), NewTemplate));
+			UObject* const* NewObject = ReplacementMap.Find(Ref->Widget.Get());
+			if ( NewObject )
+			{
+				Ref->Widget = Cast<UWidget>(*NewObject);
+			}
 		}
 		else
 		{
-			TempSelection.Add(WidgetRef);
+			WidgetHandlePool.RemoveAtSwap(HandleIndex);
 		}
 	}
 
-	SelectedWidgets = TempSelection;
-
 	// Fire the selection updated event to ensure everyone is watching the same widgets.
 	OnSelectedWidgetsChanged.Broadcast();
 }
@@ -513,6 +513,31 @@ UUserWidget* FWidgetBlueprintEditor::GetPreview() const
 	return PreviewWidgetPtr.Get();
 }
 
+FWidgetReference FWidgetBlueprintEditor::GetReferenceFromTemplate(UWidget* TemplateWidget)
+{
+	TSharedRef<FWidgetHandle> Reference = MakeShareable(new FWidgetHandle(TemplateWidget));
+	WidgetHandlePool.Add(Reference);
+
+	return FWidgetReference(SharedThis(this), Reference);
+}
+
+FWidgetReference FWidgetBlueprintEditor::GetReferenceFromPreview(UWidget* PreviewWidget)
+{
+	UUserWidget* PreviewRoot = GetPreview();
+	if ( PreviewRoot )
+	{
+		UWidgetBlueprint* Blueprint = GetWidgetBlueprintObj();
+
+		if ( PreviewWidget )
+		{
+			FString Name = PreviewWidget->GetName();
+			return GetReferenceFromTemplate(Blueprint->WidgetTree->FindWidget(Name));
+		}
+	}
+
+	return FWidgetReference(SharedThis(this), TSharedPtr<FWidgetHandle>());
+}
+
 TSharedPtr<ISequencer>& FWidgetBlueprintEditor::GetSequencer()
 {
 	if(!Sequencer.IsValid())
diff --git a/Engine/Source/Editor/UMGEditor/Private/WidgetBlueprintEditor.h b/Engine/Source/Editor/UMGEditor/Private/WidgetBlueprintEditor.h
--- a/Engine/Source/Editor/UMGEditor/Private/WidgetBlueprintEditor.h
+++ b/Engine/Source/Editor/UMGEditor/Private/WidgetBlueprintEditor.h
@@ -40,6 +40,12 @@ public:
 	/** @return The preview widget. */
 	UUserWidget* GetPreview() const;
 
+	/** Creates a widget reference using the template. */
+	FWidgetReference GetReferenceFromTemplate(UWidget* TemplateWidget);
+
+	/** Creates a widget reference using the preview.  Which is used to lookup the stable template pointer. */
+	FWidgetReference GetReferenceFromPreview(UWidget* PreviewWidget);
+
 	/** @return The sequencer used to create widget animations */
 	TSharedPtr<ISequencer>& GetSequencer();
 
@@ -156,4 +162,7 @@ private:
 
 	/** The toolbar builder associated with this editor */
 	TSharedPtr<class FWidgetBlueprintEditorToolbar> WidgetToolbar;
+
+	/** The widget references out in the ether that may need to be updated after being issued. */
+	TArray< TWeakPtr<FWidgetHandle> > WidgetHandlePool;
 };
diff --git a/Engine/Source/Editor/UMGEditor/Private/WidgetBlueprintEditorUtils.cpp b/Engine/Source/Editor/UMGEditor/Private/WidgetBlueprintEditorUtils.cpp
--- a/Engine/Source/Editor/UMGEditor/Private/WidgetBlueprintEditorUtils.cpp
+++ b/Engine/Source/Editor/UMGEditor/Private/WidgetBlueprintEditorUtils.cpp
@@ -89,7 +89,7 @@ bool FWidgetBlueprintEditorUtils::RenameWidget(TSharedRef<FWidgetBlueprintEditor
 
 
 		// Rename Preview before renaming the template widget so the preview widget can be found
-		UWidget* WidgetPreview = FWidgetReference::FromTemplate(BlueprintEditor, Widget).GetPreview();
+		UWidget* WidgetPreview = BlueprintEditor->GetReferenceFromTemplate(Widget).GetPreview();
 		if(WidgetPreview)
 		{
 			WidgetPreview->Rename(*NewNameStr);
diff --git a/Engine/Source/Editor/UMGEditor/Private/WidgetReference.cpp b/Engine/Source/Editor/UMGEditor/Private/WidgetReference.cpp
--- a/Engine/Source/Editor/UMGEditor/Private/WidgetReference.cpp
+++ b/Engine/Source/Editor/UMGEditor/Private/WidgetReference.cpp
@@ -4,71 +4,57 @@
 
 #define LOCTEXT_NAMESPACE "UMG"
 
+FWidgetHandle::FWidgetHandle(UWidget* Widget)
+	: Widget(Widget)
+{
+
+}
+
 // FWidgetReference
 
 FWidgetReference::FWidgetReference()
 	: WidgetEditor()
-	, TemplateWidget(NULL)
-{
-}
-
-FWidgetReference::FWidgetReference(TSharedPtr<FWidgetBlueprintEditor> WidgetEditor, UWidget* TemplateWidget)
-	: WidgetEditor(WidgetEditor)
-	, TemplateWidget(TemplateWidget)
+	, TemplateHandle()
 {
-	// register for any objects replaced
-	GEditor->OnObjectsReplaced().AddRaw(this, &FWidgetReference::OnObjectsReplaced);
-}
 
-FWidgetReference::~FWidgetReference()
-{
-	GEditor->OnObjectsReplaced().RemoveAll(this);
 }
 
-FWidgetReference FWidgetReference::FromTemplate(TSharedPtr<FWidgetBlueprintEditor> InWidgetEditor, UWidget* TemplateWidget)
+FWidgetReference::FWidgetReference(TSharedPtr<FWidgetBlueprintEditor> WidgetEditor, TSharedPtr< FWidgetHandle > TemplateHandle)
+	: WidgetEditor(WidgetEditor)
+	, TemplateHandle(TemplateHandle)
 {
-	return FWidgetReference(InWidgetEditor, TemplateWidget);
+	
 }
 
-FWidgetReference FWidgetReference::FromPreview(TSharedPtr<FWidgetBlueprintEditor> InWidgetEditor, UWidget* PreviewWidget)
+bool FWidgetReference::IsValid() const
 {
-	if ( InWidgetEditor.IsValid() )
+	if ( TemplateHandle.IsValid() )
 	{
-		UUserWidget* PreviewRoot = InWidgetEditor->GetPreview();
-		if ( PreviewRoot )
-		{
-			UWidgetBlueprint* Blueprint = InWidgetEditor->GetWidgetBlueprintObj();
-
-			if ( PreviewWidget )
-			{
-				FString Name = PreviewWidget->GetName();
-				return FromTemplate(InWidgetEditor, Blueprint->WidgetTree->FindWidget(Name));
-			}
-		}
+		return TemplateHandle->Widget.Get() != NULL && GetPreview();
 	}
 
-	return FWidgetReference();
-}
-
-bool FWidgetReference::IsValid() const
-{
-	return TemplateWidget.Get() != NULL && GetPreview();
+	return false;
 }
 
 UWidget* FWidgetReference::GetTemplate() const
 {
-	return TemplateWidget.Get();
+	if ( TemplateHandle.IsValid() )
+	{
+		return TemplateHandle->Widget.Get();
+	}
+
+	return NULL;
 }
 
 UWidget* FWidgetReference::GetPreview() const
 {
-	if ( WidgetEditor.IsValid() )
+	if ( WidgetEditor.IsValid() && TemplateHandle.IsValid() )
 	{
 		UUserWidget* PreviewRoot = WidgetEditor.Pin()->GetPreview();
 
-		if ( PreviewRoot && TemplateWidget.Get() )
+		if ( PreviewRoot && TemplateHandle->Widget.Get() )
 		{
-			UWidget* PreviewWidget = PreviewRoot->GetHandleFromName(TemplateWidget.Get()->GetName());
+			UWidget* PreviewWidget = PreviewRoot->GetHandleFromName(TemplateHandle->Widget.Get()->GetName());
 			return PreviewWidget;
 		}
 	}
@@ -76,13 +62,4 @@ UWidget* FWidgetReference::GetPreview() const
 	return NULL;
 }
 
-void FWidgetReference::OnObjectsReplaced(const TMap<UObject*, UObject*>& ReplacementMap)
-{
-	UObject* const* NewObject = ReplacementMap.Find(TemplateWidget.Get());
-	if ( NewObject )
-	{
-		TemplateWidget = Cast<UWidget>(*NewObject);
-	}
-}
-
 #undef LOCTEXT_NAMESPACE
diff --git a/Engine/Source/Editor/UMGEditor/Public/DesignerExtension.h b/Engine/Source/Editor/UMGEditor/Public/DesignerExtension.h
--- a/Engine/Source/Editor/UMGEditor/Public/DesignerExtension.h
+++ b/Engine/Source/Editor/UMGEditor/Public/DesignerExtension.h
@@ -16,6 +16,7 @@ namespace EExtensionLayoutLocation
 {
 	enum Type
 	{
+		/** Slate unit position relative from the parent. */
 		Absolute,
 
 		TopLeft,
@@ -38,10 +39,11 @@ namespace EExtensionLayoutLocation
 class UMGEDITOR_API FDesignerSurfaceElement : public TSharedFromThis < FDesignerSurfaceElement >
 {
 public:
-	FDesignerSurfaceElement(TSharedRef<SWidget> InWidget, EExtensionLayoutLocation::Type InLocation, FVector2D InOffset = FVector2D(0,0))
+	FDesignerSurfaceElement(TSharedRef<SWidget> InWidget, EExtensionLayoutLocation::Type InLocation, TAttribute<FVector2D> InOffset = FVector2D(0, 0), TAttribute<FVector2D> InAlignment = FVector2D(0, 0))
 		: Widget(InWidget)
 		, Location(InLocation)
 		, Offset(InOffset)
+		, Alignment(InAlignment)
 	{
 	}
 
@@ -58,7 +60,7 @@ public:
 	}
 
 	/** Sets the offset after laid out in that location */
-	void SetOffset(FVector2D InOffset)
+	void SetOffset(TAttribute<FVector2D> InOffset)
 	{
 		Offset = InOffset;
 	}
@@ -66,7 +68,19 @@ public:
 	/** Gets the offset after being laid out. */
 	FVector2D GetOffset() const
 	{
-		return Offset;
+		return Offset.Get();
+	}
+
+	/** Sets the alignment to use, a normalized value representing position inside the parent. */
+	void SetAlignment(TAttribute<FVector2D> InAlignment)
+	{
+		Alignment = InAlignment;
+	}
+
+	/** Gets the alignment to use, a normalized value representing position inside the parent. */
+	FVector2D GetAlignment() const
+	{
+		return Alignment.Get();
 	}
 
 protected:
@@ -74,7 +88,9 @@ protected:
 
 	EExtensionLayoutLocation::Type Location;
 
-	FVector2D Offset;
+	TAttribute<FVector2D> Offset;
+
+	TAttribute<FVector2D> Alignment;
 };
 
 /**
diff --git a/Engine/Source/Editor/UMGEditor/Public/WidgetReference.h b/Engine/Source/Editor/UMGEditor/Public/WidgetReference.h
--- a/Engine/Source/Editor/UMGEditor/Public/WidgetReference.h
+++ b/Engine/Source/Editor/UMGEditor/Public/WidgetReference.h
@@ -5,6 +5,19 @@
 class UWidget;
 class SWidget;
 class UWidgetBlueprint;
+class FWidgetBlueprintEditor;
+struct FWidgetReference;
+
+struct UMGEDITOR_API FWidgetHandle
+{
+	friend FWidgetReference;
+	friend FWidgetBlueprintEditor;
+
+private:
+	FWidgetHandle(UWidget* Widget);
+
+	TWeakObjectPtr<UWidget> Widget;
+};
 
 /**
  * The Widget reference is a useful way to hold onto the selection in a way that allows for up to date access to the current preview object.
@@ -12,23 +25,11 @@ class UWidgetBlueprint;
  */
 struct UMGEDITOR_API FWidgetReference
 {
-public:
-	/**
-	 * Creates a widget reference using the template.
-	 */
-	static FWidgetReference FromTemplate(TSharedPtr<class FWidgetBlueprintEditor> InWidgetEditor, UWidget* TemplateWidget);
-
-	/**
-	 * Creates a widget reference using the preview.  Which is used to lookup the stable template pointer.
-	 */
-	static FWidgetReference FromPreview(TSharedPtr<class FWidgetBlueprintEditor> InWidgetEditor, UWidget* PreviewWidget);
+	friend FWidgetBlueprintEditor;
 
-	/** Constructor for the null widget reference.  Not intended for normal use. */
+public:
 	FWidgetReference();
 
-	/** Destructor */
-	~FWidgetReference();
-
 	/** @returns true if both the template and the preview pointers are valid. */
 	bool IsValid() const;
 
@@ -41,20 +42,23 @@ public:
 	/** Checks if widget reference is the same as nother widget reference, based on the template pointers. */
 	bool operator==( const FWidgetReference& Other ) const
 	{
-		return TemplateWidget.Get() == Other.TemplateWidget.Get();
+		if ( TemplateHandle.IsValid() && Other.TemplateHandle.IsValid() )
+		{
+			return TemplateHandle->Widget.Get() == Other.TemplateHandle->Widget.Get();
+		}
+
+		return false;
 	}
 
 private:
-	FWidgetReference(TSharedPtr<class FWidgetBlueprintEditor> WidgetEditor, UWidget* TemplateWidget);
+	FWidgetReference(TSharedPtr<FWidgetBlueprintEditor> WidgetEditor, TSharedPtr< FWidgetHandle > TemplateHandle);
 
 	void OnObjectsReplaced(const TMap<UObject*, UObject*>& ReplacementMap);
 
 private:
+	TWeakPtr<FWidgetBlueprintEditor> WidgetEditor;
 
-	TWeakPtr<class FWidgetBlueprintEditor> WidgetEditor;
-
-	//TODO UMG We probably want to make this a weak object reference.
-	TWeakObjectPtr<UWidget> TemplateWidget;
+	TSharedPtr< FWidgetHandle > TemplateHandle;
 };
 
 inline uint32 GetTypeHash(const struct FWidgetReference& WidgetRef)
diff --git a/Engine/Source/Runtime/UMG/Private/Components/CanvasPanelSlot.cpp b/Engine/Source/Runtime/UMG/Private/Components/CanvasPanelSlot.cpp
--- a/Engine/Source/Runtime/UMG/Private/Components/CanvasPanelSlot.cpp
+++ b/Engine/Source/Runtime/UMG/Private/Components/CanvasPanelSlot.cpp
@@ -153,16 +153,7 @@ void UCanvasPanelSlot::PreEditChange(UProperty* PropertyAboutToChange)
 {
 	Super::PreEditChange(PropertyAboutToChange);
 
-	// Get the current location
-	if ( UCanvasPanel* Canvas = Cast<UCanvasPanel>(Parent) )
-	{
-		FGeometry Geometry;
-		if ( Canvas->GetGeometryForSlot(this, Geometry) )
-		{
-			PreEditGeometry = Geometry;
-			PreEditLayoutData = LayoutData;
-		}
-	}
+	SaveBaseLayout();
 }
 
 void UCanvasPanelSlot::PostEditChangeChainProperty(struct FPropertyChangedChainEvent& PropertyChangedEvent)
@@ -170,89 +161,120 @@ void UCanvasPanelSlot::PostEditChangeChainProperty(struct FPropertyChangedChainE
 	SynchronizeProperties();
 
 	static FName AnchorsProperty(TEXT("Anchors"));
-	
+
 	FEditPropertyChain::TDoubleLinkedListNode* AnchorNode = PropertyChangedEvent.PropertyChain.GetHead()->GetNextNode();
-	if (!AnchorNode)
+	if ( !AnchorNode )
 	{
 		return;
 	}
-	
+
 	FEditPropertyChain::TDoubleLinkedListNode* LayoutDataNode = AnchorNode->GetNextNode();
-	
-	if (!LayoutDataNode)
+
+	if ( !LayoutDataNode )
 	{
 		return;
 	}
-	
+
 	UProperty* AnchorProperty = LayoutDataNode->GetValue();
 
 	if ( AnchorProperty && AnchorProperty->GetFName() == AnchorsProperty )
 	{
-		// Ensure we have a parent canvas
-		if ( UCanvasPanel* Canvas = Cast<UCanvasPanel>(Parent) )
+		RebaseLayout();
+	}
+
+	Super::PostEditChangeProperty(PropertyChangedEvent);
+}
+
+void UCanvasPanelSlot::SaveBaseLayout()
+{
+	// Get the current location
+	if ( UCanvasPanel* Canvas = Cast<UCanvasPanel>(Parent) )
+	{
+		FGeometry Geometry;
+		if ( Canvas->GetGeometryForSlot(this, Geometry) )
+		{
+			PreEditGeometry = Geometry;
+			PreEditLayoutData = LayoutData;
+		}
+	}
+}
+
+void UCanvasPanelSlot::RebaseLayout(bool PreserveSize)
+{
+	// Ensure we have a parent canvas
+	if ( UCanvasPanel* Canvas = Cast<UCanvasPanel>(Parent) )
+	{
+		FGeometry Geometry;
+		if ( Canvas->GetGeometryForSlot(this, Geometry) )
 		{
-			FGeometry Geometry;
-			if ( Canvas->GetGeometryForSlot(this, Geometry) )
+			// Calculate the default anchor offset, ie where would this control be laid out if no offset were provided.
+			FVector2D CanvasSize = Canvas->GetCanvasWidget()->GetCachedGeometry().Size;
+			FMargin AnchorPositions = FMargin(
+				LayoutData.Anchors.Minimum.X * CanvasSize.X,
+				LayoutData.Anchors.Minimum.Y * CanvasSize.Y,
+				LayoutData.Anchors.Maximum.X * CanvasSize.X,
+				LayoutData.Anchors.Maximum.Y * CanvasSize.Y);
+			FVector2D DefaultAnchorPosition = FVector2D(AnchorPositions.Left, AnchorPositions.Top);
+
+			// Determine the amount that would be offset from the anchor position if alignment was applied.
+			FVector2D AlignmentOffset = LayoutData.Alignment * PreEditGeometry.Size;
+
+			// Determine where the widget's new position needs to be to maintain a stable location when the anchors change.
+			FVector2D LeftTopDelta = PreEditGeometry.Position - DefaultAnchorPosition;
+
+			// Adjust the size to remain constant
+			if ( !LayoutData.Anchors.IsStretchedHorizontal() && PreEditLayoutData.Anchors.IsStretchedHorizontal() )
 			{
-				// Calculate the default anchor offset, ie where would this control be laid out if no offset were provided.
-				FVector2D CanvasSize = Canvas->GetCanvasWidget()->GetCachedGeometry().Size;
-				FMargin AnchorPositions = FMargin(
-					LayoutData.Anchors.Minimum.X * CanvasSize.X,
-					LayoutData.Anchors.Minimum.Y * CanvasSize.Y,
-					LayoutData.Anchors.Maximum.X * CanvasSize.X,
-					LayoutData.Anchors.Maximum.Y * CanvasSize.Y);
-				FVector2D DefaultAnchorPosition = FVector2D(AnchorPositions.Left, AnchorPositions.Top);
-
-				// Determine the amount that would be offset from the anchor position if alignment was applied.
-				FVector2D AlignmentOffset = LayoutData.Alignment * PreEditGeometry.Size;
-
-				// Determine where the widget's new position needs to be to maintain a stable location when the anchors change.
-				FVector2D PositionDelta = PreEditGeometry.Position - DefaultAnchorPosition;// ( Geometry.Position - DefaultAnchorOffset );
-
-				// Adjust the size to remain constant
-				if ( !LayoutData.Anchors.IsStretchedHorizontal() && PreEditLayoutData.Anchors.IsStretchedHorizontal() )
-				{
-					// Adjust the position to remain constant
-					LayoutData.Offsets.Left = PositionDelta.X + AlignmentOffset.X;
-					LayoutData.Offsets.Right = PreEditGeometry.Size.X;
-				}
-				else if ( LayoutData.Anchors.IsStretchedHorizontal() )
-				{
-					// Adjust the position to remain constant
-					LayoutData.Offsets.Left = PositionDelta.X;
-					LayoutData.Offsets.Right = AnchorPositions.Right - ( PositionDelta.X + PreEditGeometry.Size.X );
-				}
-				else
-				{
-					// Adjust the position to remain constant
-					LayoutData.Offsets.Left = PositionDelta.X + AlignmentOffset.X;
-				}
-
-				if ( !LayoutData.Anchors.IsStretchedVertical() && PreEditLayoutData.Anchors.IsStretchedVertical() )
-				{
-					// Adjust the position to remain constant
-					LayoutData.Offsets.Top = PositionDelta.Y + AlignmentOffset.Y;
-					LayoutData.Offsets.Bottom = PreEditGeometry.Size.Y;
-				}
-				else if ( LayoutData.Anchors.IsStretchedVertical() )
-				{
-					// Adjust the position to remain constant
-					LayoutData.Offsets.Top = PositionDelta.Y;
-					LayoutData.Offsets.Bottom = AnchorPositions.Bottom - ( PositionDelta.Y + PreEditGeometry.Size.Y );
-				}
-				else
-				{
-					// Adjust the position to remain constant
-					LayoutData.Offsets.Top = PositionDelta.Y + AlignmentOffset.Y;
-				}
+				// Adjust the position to remain constant
+				LayoutData.Offsets.Left = LeftTopDelta.X + AlignmentOffset.X;
+				LayoutData.Offsets.Right = PreEditGeometry.Size.X;
+			}
+			else if ( !PreserveSize && LayoutData.Anchors.IsStretchedHorizontal() && !PreEditLayoutData.Anchors.IsStretchedHorizontal() )
+			{
+				// Adjust the position to remain constant
+				LayoutData.Offsets.Left = 0;
+				LayoutData.Offsets.Right = 0;
+			}
+			else if ( LayoutData.Anchors.IsStretchedHorizontal() )
+			{
+				// Adjust the position to remain constant
+				LayoutData.Offsets.Left = LeftTopDelta.X;
+				LayoutData.Offsets.Right = AnchorPositions.Right - ( AnchorPositions.Left + LayoutData.Offsets.Left + PreEditGeometry.Size.X );
+			}
+			else
+			{
+				// Adjust the position to remain constant
+				LayoutData.Offsets.Left = LeftTopDelta.X + AlignmentOffset.X;
 			}
 
-			// Apply the changes to the properties.
-			SynchronizeProperties();
+			if ( !LayoutData.Anchors.IsStretchedVertical() && PreEditLayoutData.Anchors.IsStretchedVertical() )
+			{
+				// Adjust the position to remain constant
+				LayoutData.Offsets.Top = LeftTopDelta.Y + AlignmentOffset.Y;
+				LayoutData.Offsets.Bottom = PreEditGeometry.Size.Y;
+			}
+			else if ( !PreserveSize && LayoutData.Anchors.IsStretchedVertical() && !PreEditLayoutData.Anchors.IsStretchedVertical() )
+			{
+				// Adjust the position to remain constant
+				LayoutData.Offsets.Top = 0;
+				LayoutData.Offsets.Bottom = 0;
+			}
+			else if ( LayoutData.Anchors.IsStretchedVertical() )
+			{
+				// Adjust the position to remain constant
+				LayoutData.Offsets.Top = LeftTopDelta.Y;
+				LayoutData.Offsets.Bottom = AnchorPositions.Bottom - ( AnchorPositions.Top + LayoutData.Offsets.Top + PreEditGeometry.Size.Y );
+			}
+			else
+			{
+				// Adjust the position to remain constant
+				LayoutData.Offsets.Top = LeftTopDelta.Y + AlignmentOffset.Y;
+			}
 		}
-	}
 
-	Super::PostEditChangeProperty(PropertyChangedEvent);
+		// Apply the changes to the properties.
+		SynchronizeProperties();
+	}
 }
 
 #endif
\ No newline at end of file
diff --git a/Engine/Source/Runtime/UMG/Public/Components/CanvasPanelSlot.h b/Engine/Source/Runtime/UMG/Public/Components/CanvasPanelSlot.h
--- a/Engine/Source/Runtime/UMG/Public/Components/CanvasPanelSlot.h
+++ b/Engine/Source/Runtime/UMG/Public/Components/CanvasPanelSlot.h
@@ -77,6 +77,12 @@ class UMG_API UCanvasPanelSlot : public UPanelSlot
 	virtual void PreEditChange(UProperty* PropertyAboutToChange) override;
 	virtual void PostEditChangeChainProperty(struct FPropertyChangedChainEvent& PropertyChangedEvent);
 	// End of UObject interface
+
+	/** Stores the current layout information about the slot and parent canvas. */
+	void SaveBaseLayout();
+
+	/** Compares the saved base layout against the current state.  Updates the necessary properties to maintain a stable position. */
+	void RebaseLayout(bool PreserveSize = true);
 #endif
 
 private:
