UMG - Fixing a bug when clicking on the vertical bits of the alignment gizmo, jumping to the wrong anchor location.

diff --git a/Engine/Source/Editor/UMGEditor/Private/Extensions/CanvasSlotExtension.cpp b/Engine/Source/Editor/UMGEditor/Private/Extensions/CanvasSlotExtension.cpp
--- a/Engine/Source/Editor/UMGEditor/Private/Extensions/CanvasSlotExtension.cpp
+++ b/Engine/Source/Editor/UMGEditor/Private/Extensions/CanvasSlotExtension.cpp
@@ -150,21 +150,23 @@ EVisibility FCanvasSlotExtension::GetAnchorVisibility(EAnchorWidget::Type Anchor
 	for ( const FWidgetReference& Selection : SelectionCache )
 	{
 		UWidget* PreviewWidget = Selection.GetPreview();
-
-		UCanvasPanelSlot* PreviewCanvasSlot = CastChecked<UCanvasPanelSlot>(PreviewWidget->Slot);
-		switch ( AnchorType )
+		if ( PreviewWidget->Slot )
 		{
-		case EAnchorWidget::Center:
-			return PreviewCanvasSlot->LayoutData.Anchors.Minimum == PreviewCanvasSlot->LayoutData.Anchors.Maximum ? EVisibility::Visible : EVisibility::Collapsed;
-		case EAnchorWidget::Left:
-		case EAnchorWidget::Right:
-			return PreviewCanvasSlot->LayoutData.Anchors.Minimum.Y == PreviewCanvasSlot->LayoutData.Anchors.Maximum.Y ? EVisibility::Visible : EVisibility::Collapsed;
-		case EAnchorWidget::Top:
-		case EAnchorWidget::Bottom:
-			return PreviewCanvasSlot->LayoutData.Anchors.Minimum.X == PreviewCanvasSlot->LayoutData.Anchors.Maximum.X ? EVisibility::Visible : EVisibility::Collapsed;
-		}
+			UCanvasPanelSlot* PreviewCanvasSlot = CastChecked<UCanvasPanelSlot>(PreviewWidget->Slot);
+			switch ( AnchorType )
+			{
+			case EAnchorWidget::Center:
+				return PreviewCanvasSlot->LayoutData.Anchors.Minimum == PreviewCanvasSlot->LayoutData.Anchors.Maximum ? EVisibility::Visible : EVisibility::Collapsed;
+			case EAnchorWidget::Left:
+			case EAnchorWidget::Right:
+				return PreviewCanvasSlot->LayoutData.Anchors.Minimum.Y == PreviewCanvasSlot->LayoutData.Anchors.Maximum.Y ? EVisibility::Visible : EVisibility::Collapsed;
+			case EAnchorWidget::Top:
+			case EAnchorWidget::Bottom:
+				return PreviewCanvasSlot->LayoutData.Anchors.Minimum.X == PreviewCanvasSlot->LayoutData.Anchors.Maximum.X ? EVisibility::Visible : EVisibility::Collapsed;
+			}
 
-		return EVisibility::Visible;
+			return EVisibility::Visible;
+		}
 	}
 
 	return EVisibility::Collapsed;
@@ -175,26 +177,29 @@ FVector2D FCanvasSlotExtension::GetAnchorAlignment(EAnchorWidget::Type AnchorTyp
 	for ( const FWidgetReference& Selection : SelectionCache )
 	{
 		UWidget* PreviewWidget = Selection.GetPreview();
-		UCanvasPanelSlot* PreviewCanvasSlot = CastChecked<UCanvasPanelSlot>(PreviewWidget->Slot);
-		
-		FVector2D Minimum = PreviewCanvasSlot->LayoutData.Anchors.Minimum;
-		FVector2D Maximum = PreviewCanvasSlot->LayoutData.Anchors.Maximum;
-		
-		switch ( AnchorType )
+		if ( PreviewWidget->Slot )
 		{
-		case EAnchorWidget::Center:
-		case EAnchorWidget::Left:
-		case EAnchorWidget::Top:
-		case EAnchorWidget::TopLeft:
-			return Minimum;
-		case EAnchorWidget::Right:
-		case EAnchorWidget::Bottom:
-		case EAnchorWidget::BottomRight:
-			return Maximum;
-		case EAnchorWidget::TopRight:
-			return  FVector2D(Maximum.X, Minimum.Y);
-		case EAnchorWidget::BottomLeft:
-			return  FVector2D(Minimum.X, Maximum.Y);
+			UCanvasPanelSlot* PreviewCanvasSlot = CastChecked<UCanvasPanelSlot>(PreviewWidget->Slot);
+
+			FVector2D Minimum = PreviewCanvasSlot->LayoutData.Anchors.Minimum;
+			FVector2D Maximum = PreviewCanvasSlot->LayoutData.Anchors.Maximum;
+
+			switch ( AnchorType )
+			{
+			case EAnchorWidget::Center:
+			case EAnchorWidget::Left:
+			case EAnchorWidget::Top:
+			case EAnchorWidget::TopLeft:
+				return Minimum;
+			case EAnchorWidget::Right:
+			case EAnchorWidget::Bottom:
+			case EAnchorWidget::BottomRight:
+				return Maximum;
+			case EAnchorWidget::TopRight:
+				return  FVector2D(Maximum.X, Minimum.Y);
+			case EAnchorWidget::BottomLeft:
+				return  FVector2D(Minimum.X, Maximum.Y);
+			}
 		}
 	}
 
@@ -256,12 +261,21 @@ void FCanvasSlotExtension::Paint(const TSet< FWidgetReference >& Selection, cons
 
 FReply FCanvasSlotExtension::HandleAnchorBeginDrag(const FGeometry& Geometry, const FPointerEvent& Event, EAnchorWidget::Type AnchorType)
 {
+	BeginTransaction(LOCTEXT("MoveAnchor", "Move Anchor"));
+
 	bMovingAnchor = true;
+	MouseDownPosition = Event.GetScreenSpacePosition();
+
+	UCanvasPanelSlot* PreviewCanvasSlot = Cast<UCanvasPanelSlot>(SelectionCache[0].GetPreview()->Slot);
+	BeginAnchors = PreviewCanvasSlot->LayoutData.Anchors;
+
 	return FReply::Handled().CaptureMouse(AnchorWidgets[(int32)AnchorType].ToSharedRef());
 }
 
 FReply FCanvasSlotExtension::HandleAnchorEndDrag(const FGeometry& Geometry, const FPointerEvent& Event, EAnchorWidget::Type AnchorType)
 {
+	EndTransaction();
+
 	bMovingAnchor = false;
 	return FReply::Handled().ReleaseMouseCapture();
 }
@@ -270,9 +284,7 @@ FReply FCanvasSlotExtension::HandleAnchorDragging(const FGeometry& Geometry, con
 {
 	if ( bMovingAnchor && !Event.GetCursorDelta().IsZero() )
 	{
-		float InverseScale = ( 1.0f / Designer->GetPreviewScale() );
-
-		FVector2D DeltaScaled = Event.GetCursorDelta() * InverseScale;
+		float InverseSize = 1.0f / Designer->GetPreviewScale();
 
 		for ( FWidgetReference& Selection : SelectionCache )
 		{
@@ -284,17 +296,19 @@ FReply FCanvasSlotExtension::HandleAnchorDragging(const FGeometry& Geometry, con
 				FGeometry Geometry;
 				if ( Canvas->GetGeometryForSlot(PreviewCanvasSlot, Geometry) )
 				{
-					// Calculate the default anchor offset, ie where would this control be laid out if no offset were provided.
-					FVector2D CanvasSize = Canvas->GetCanvasWidget()->GetCachedGeometry().Size;
-
-					FVector2D AnchorDelta = DeltaScaled / CanvasSize;
+					FGeometry CanvasGeometry = Canvas->GetCanvasWidget()->GetCachedGeometry();
+					FVector2D StartLocalPosition = CanvasGeometry.AbsoluteToLocal(MouseDownPosition);
+					FVector2D NewLocalPosition = CanvasGeometry.AbsoluteToLocal(Event.GetScreenSpacePosition());
+					FVector2D LocalPositionDelta = NewLocalPosition - StartLocalPosition;
 
+					FVector2D AnchorDelta = LocalPositionDelta / CanvasGeometry.Size;
+					
 					PreviewCanvasSlot->SaveBaseLayout();
 					switch ( AnchorType )
 					{
 					case EAnchorWidget::Center:
-						PreviewCanvasSlot->LayoutData.Anchors.Maximum += AnchorDelta;
-						PreviewCanvasSlot->LayoutData.Anchors.Minimum += AnchorDelta;
+						PreviewCanvasSlot->LayoutData.Anchors.Maximum = BeginAnchors.Maximum + AnchorDelta;
+						PreviewCanvasSlot->LayoutData.Anchors.Minimum = BeginAnchors.Minimum + AnchorDelta;
 						PreviewCanvasSlot->LayoutData.Anchors.Minimum.X = FMath::Clamp(PreviewCanvasSlot->LayoutData.Anchors.Minimum.X, 0.0f, 1.0f);
 						PreviewCanvasSlot->LayoutData.Anchors.Maximum.X = FMath::Clamp(PreviewCanvasSlot->LayoutData.Anchors.Maximum.X, 0.0f, 1.0f);
 						PreviewCanvasSlot->LayoutData.Anchors.Minimum.Y = FMath::Clamp(PreviewCanvasSlot->LayoutData.Anchors.Minimum.Y, 0.0f, 1.0f);
@@ -307,7 +321,7 @@ FReply FCanvasSlotExtension::HandleAnchorDragging(const FGeometry& Geometry, con
 					case EAnchorWidget::Left:
 					case EAnchorWidget::TopLeft:
 					case EAnchorWidget::BottomLeft:
-						PreviewCanvasSlot->LayoutData.Anchors.Minimum.X += AnchorDelta.X;
+						PreviewCanvasSlot->LayoutData.Anchors.Minimum.X = BeginAnchors.Minimum.X + AnchorDelta.X;
 						PreviewCanvasSlot->LayoutData.Anchors.Minimum.X = FMath::Clamp(PreviewCanvasSlot->LayoutData.Anchors.Minimum.X, 0.0f, PreviewCanvasSlot->LayoutData.Anchors.Maximum.X);
 						break;
 					}
@@ -317,7 +331,7 @@ FReply FCanvasSlotExtension::HandleAnchorDragging(const FGeometry& Geometry, con
 					case EAnchorWidget::Right:
 					case EAnchorWidget::TopRight:
 					case EAnchorWidget::BottomRight:
-						PreviewCanvasSlot->LayoutData.Anchors.Maximum.X += AnchorDelta.X;
+						PreviewCanvasSlot->LayoutData.Anchors.Maximum.X = BeginAnchors.Maximum.X + AnchorDelta.X;
 						PreviewCanvasSlot->LayoutData.Anchors.Maximum.X = FMath::Clamp(PreviewCanvasSlot->LayoutData.Anchors.Maximum.X, PreviewCanvasSlot->LayoutData.Anchors.Minimum.X, 1.0f);
 						break;
 					}
@@ -327,7 +341,7 @@ FReply FCanvasSlotExtension::HandleAnchorDragging(const FGeometry& Geometry, con
 					case EAnchorWidget::Top:
 					case EAnchorWidget::TopLeft:
 					case EAnchorWidget::TopRight:
-						PreviewCanvasSlot->LayoutData.Anchors.Minimum.Y += AnchorDelta.Y;
+						PreviewCanvasSlot->LayoutData.Anchors.Minimum.Y = BeginAnchors.Minimum.Y + AnchorDelta.Y;
 						PreviewCanvasSlot->LayoutData.Anchors.Minimum.Y = FMath::Clamp(PreviewCanvasSlot->LayoutData.Anchors.Minimum.Y, 0.0f, PreviewCanvasSlot->LayoutData.Anchors.Maximum.Y);
 						break;
 					}
@@ -337,7 +351,7 @@ FReply FCanvasSlotExtension::HandleAnchorDragging(const FGeometry& Geometry, con
 					case EAnchorWidget::Bottom:
 					case EAnchorWidget::BottomLeft:
 					case EAnchorWidget::BottomRight:
-						PreviewCanvasSlot->LayoutData.Anchors.Maximum.Y += AnchorDelta.Y;
+						PreviewCanvasSlot->LayoutData.Anchors.Maximum.Y = BeginAnchors.Maximum.Y + AnchorDelta.Y;
 						PreviewCanvasSlot->LayoutData.Anchors.Maximum.Y = FMath::Clamp(PreviewCanvasSlot->LayoutData.Anchors.Maximum.Y, PreviewCanvasSlot->LayoutData.Anchors.Minimum.Y, 1.0f);
 						break;
 					}
diff --git a/Engine/Source/Editor/UMGEditor/Private/Extensions/CanvasSlotExtension.h b/Engine/Source/Editor/UMGEditor/Private/Extensions/CanvasSlotExtension.h
--- a/Engine/Source/Editor/UMGEditor/Private/Extensions/CanvasSlotExtension.h
+++ b/Engine/Source/Editor/UMGEditor/Private/Extensions/CanvasSlotExtension.h
@@ -65,9 +65,18 @@ private:
 
 private:
 
+	/** */
 	TSharedPtr<SBorder> MoveHandle;
 
+	/** */
 	TArray< TSharedPtr<SWidget> > AnchorWidgets;
 
+	/** */
 	bool bMovingAnchor;
+
+	/** */
+	FVector2D MouseDownPosition;
+
+	/** */
+	FAnchors BeginAnchors;
 };
\ No newline at end of file
