UMG - Now locking movement inside some containers (Canvas), until you press shift to move the widget outside the bounds of the parent.  Added a message to the designer.

diff --git a/Engine/Source/Editor/UMGEditor/Private/Designer/SDesignerView.cpp b/Engine/Source/Editor/UMGEditor/Private/Designer/SDesignerView.cpp
--- a/Engine/Source/Editor/UMGEditor/Private/Designer/SDesignerView.cpp
+++ b/Engine/Source/Editor/UMGEditor/Private/Designer/SDesignerView.cpp
@@ -9,6 +9,7 @@
 
 #include "WidgetTemplateDragDropOp.h"
 #include "SZoomPan.h"
+#include "SDisappearingBar.h"
 
 #define LOCTEXT_NAMESPACE "UMG"
 
@@ -23,13 +24,24 @@ public:
 
 	FWidgetReference Widget;
 
-	static TSharedRef<FSelectedWidgetDragDropOp> New(FWidgetReference InWidget);
+	bool bStayingInParent;
+	FWidgetReference ParentWidget;
+
+	static TSharedRef<FSelectedWidgetDragDropOp> New(TSharedPtr<FWidgetBlueprintEditor> Editor, FWidgetReference InWidget);
 };
 
-TSharedRef<FSelectedWidgetDragDropOp> FSelectedWidgetDragDropOp::New(FWidgetReference InWidget)
+TSharedRef<FSelectedWidgetDragDropOp> FSelectedWidgetDragDropOp::New(TSharedPtr<FWidgetBlueprintEditor> Editor, FWidgetReference InWidget)
 {
+	bool bStayInParent = false;
+	if ( UPanelWidget* PanelTemplate = InWidget.GetTemplate()->GetParent() )
+	{
+		bStayInParent = PanelTemplate->LockToPanelOnDrag();
+	}
+
 	TSharedRef<FSelectedWidgetDragDropOp> Operation = MakeShareable(new FSelectedWidgetDragDropOp());
 	Operation->Widget = InWidget;
+	Operation->bStayingInParent = bStayInParent;
+	Operation->ParentWidget = Editor->GetReferenceFromTemplate(InWidget.GetTemplate()->GetParent());
 	Operation->DefaultHoverText = FText::FromString( InWidget.GetTemplate()->GetLabel() );
 	Operation->CurrentHoverText = FText::FromString( InWidget.GetTemplate()->GetLabel() );
 	Operation->Construct();
@@ -106,6 +118,8 @@ void SDesignerView::Construct(const FArguments& InArgs, TSharedPtr<FWidgetBluepr
 	BlueprintEditor = InBlueprintEditor;
 	CurrentHandle = DH_NONE;
 
+	DesignerMessage = EDesignerMessage::None;
+
 	SetStartupResolution();
 
 	ResolutionTextFade = FCurveSequence(0.0f, 1.0f);
@@ -233,6 +247,37 @@ void SDesignerView::Construct(const FArguments& InArgs, TSharedPtr<FWidgetBluepr
 					]
 				]
 			]
+
+			// Info Bar, displays heads up information about some actions.
+			+ SOverlay::Slot()
+			.HAlign(HAlign_Fill)
+			.VAlign(VAlign_Bottom)
+			.Padding(FMargin(0, 0, 0, 200))
+			[
+				SNew(SDisappearingBar)
+				[
+					SNew(SHorizontalBox)
+
+					+ SHorizontalBox::Slot()
+					.FillWidth(1.0f)
+					.Padding(0)
+					[
+						SNew(SBorder)
+						.BorderImage(FEditorStyle::GetBrush("ToolPanel.GroupBorder"))
+						.BorderBackgroundColor(FLinearColor(1, 1, 1, 0.75))
+						.HAlign(HAlign_Center)
+						.VAlign(VAlign_Center)
+						.Padding(FMargin(0, 5))
+						.Visibility(this, &SDesignerView::GetInfoBarVisibility)
+						[
+							SNew(STextBlock)
+							.TextStyle(FEditorStyle::Get(), "Graph.ZoomText")
+							.Text(this, &SDesignerView::GetInfoBarText)
+						]
+					]
+				]
+			]
+
 			// Bottom bar to show current resolution & AR
 			+ SOverlay::Slot()
 			.HAlign(HAlign_Fill)
@@ -319,6 +364,27 @@ FSlateRect SDesignerView::ComputeAreaBounds() const
 	return FSlateRect(0, 0, PreviewWidth, PreviewHeight);
 }
 
+EVisibility SDesignerView::GetInfoBarVisibility() const
+{
+	if ( DesignerMessage != EDesignerMessage::None && bMovingExistingWidget )
+	{
+		return EVisibility::Visible;
+	}
+
+	return EVisibility::Hidden;
+}
+
+FText SDesignerView::GetInfoBarText() const
+{
+	switch ( DesignerMessage )
+	{
+	case EDesignerMessage::MoveFromParent:
+		return LOCTEXT("PressShiftToMove", "Press SHIFT to move the widget out of the current parent");
+	}
+
+	return FText::GetEmpty();
+}
+
 void SDesignerView::OnEditorSelectionChanged()
 {
 	TSet<FWidgetReference> PendingSelectedWidgets = BlueprintEditor.Pin()->GetSelectedWidgets();
@@ -359,10 +425,15 @@ void SDesignerView::OnEditorSelectionChanged()
 	CreateExtensionWidgetsForSelection();
 }
 
+void SDesignerView::ClearExtensionWidgets()
+{
+	ExtensionWidgetCanvas->ClearChildren();
+}
+
 void SDesignerView::CreateExtensionWidgetsForSelection()
 {
 	// Remove all the current extension widgets
-	ExtensionWidgetCanvas->ClearChildren();
+	ClearExtensionWidgets();
 
 	TArray<FWidgetReference> Selected;
 	if ( SelectedWidget.IsValid() )
@@ -1034,7 +1105,9 @@ FReply SDesignerView::OnDragDetected(const FGeometry& MyGeometry, const FPointer
 {
 	if ( SelectedWidget.IsValid() )
 	{
-		return FReply::Handled().BeginDragDrop(FSelectedWidgetDragDropOp::New(SelectedWidget));
+		ClearExtensionWidgets();
+
+		return FReply::Handled().BeginDragDrop(FSelectedWidgetDragDropOp::New(BlueprintEditor.Pin(), SelectedWidget));
 	}
 
 	return FReply::Unhandled();
@@ -1229,6 +1302,35 @@ UWidget* SDesignerView::ProcessDropAndAddWidget(const FGeometry& MyGeometry, con
 	{
 		SelectedDragDropOp->SetCursorOverride(TOptional<EMouseCursor::Type>());
 
+		// If they've pressed shift, and we were staying in the parent, disable that
+		// and adjust the designer message to no longer warn.
+		if ( DragDropEvent.IsShiftDown() && SelectedDragDropOp->bStayingInParent )
+		{
+			SelectedDragDropOp->bStayingInParent = false;
+			DesignerMessage = EDesignerMessage::None;
+		}
+
+		// If we're staying in the parent we started in, replace the parent found under the cursor with
+		// the original one, also update the arranged widget data so that our layout calculations are accurate.
+		if ( SelectedDragDropOp->bStayingInParent )
+		{
+			DesignerMessage = EDesignerMessage::MoveFromParent;
+
+			if ( UPanelWidget* LockedParent = Cast<UPanelWidget>(bIsPreview ? SelectedDragDropOp->ParentWidget.GetPreview() : SelectedDragDropOp->ParentWidget.GetTemplate()) )
+			{
+				TSharedPtr<SWidget> NewParentWidget = LockedParent->GetCachedWidget();
+
+				if ( NewParentWidget.IsValid() )
+				{
+					Target = LockedParent;
+
+					const bool bSuccess = FDesignTimeUtils::GetArrangedWidget(NewParentWidget.ToSharedRef(), ArrangedWidget);
+					check(bSuccess);
+				}
+			}
+		}
+
+		// If the widget being hovered over is a panel, attempt to place it into that panel.
 		if ( Target && Target->IsA(UPanelWidget::StaticClass()) )
 		{
 			UPanelWidget* NewParent = Cast<UPanelWidget>(Target);
@@ -1319,6 +1421,9 @@ FReply SDesignerView::OnDrop(const FGeometry& MyGeometry, const FDragDropEvent&
 	{
 		FBlueprintEditorUtils::MarkBlueprintAsStructurallyModified(BP);
 
+		// Regenerate extension widgets now that we've finished moving or placing the widget.
+		CreateExtensionWidgetsForSelection();
+
 		return FReply::Handled();
 	}
 	
diff --git a/Engine/Source/Editor/UMGEditor/Private/Designer/SDesignerView.h b/Engine/Source/Editor/UMGEditor/Private/Designer/SDesignerView.h
--- a/Engine/Source/Editor/UMGEditor/Private/Designer/SDesignerView.h
+++ b/Engine/Source/Editor/UMGEditor/Private/Designer/SDesignerView.h
@@ -8,6 +8,15 @@
 #include "DesignerExtension.h"
 #include "IUMGDesigner.h"
 
+namespace EDesignerMessage
+{
+	enum Type
+	{
+		None,
+		MoveFromParent,
+	};
+}
+
 class FDesignerExtension;
 
 /**
@@ -70,8 +79,12 @@ private:
 	/** Updates the designer to display the latest preview widget */
 	void UpdatePreviewWidget();
 
+	void ClearExtensionWidgets();
 	void CreateExtensionWidgetsForSelection();
 
+	EVisibility GetInfoBarVisibility() const;
+	FText GetInfoBarText() const;
+
 	/** Displays the context menu when you right click */
 	void ShowContextMenu(const FGeometry& MyGeometry, const FPointerEvent& MouseEvent);
 
@@ -215,4 +228,7 @@ private:
 
 	/**  */
 	FWeakWidgetPath SelectedWidgetPath;
+
+	/** */
+	EDesignerMessage::Type DesignerMessage;
 };
diff --git a/Engine/Source/Editor/UMGEditor/Private/Designer/SDisappearingBar.cpp b/Engine/Source/Editor/UMGEditor/Private/Designer/SDisappearingBar.cpp
new file mode 100644
--- /dev/null
+++ b/Engine/Source/Editor/UMGEditor/Private/Designer/SDisappearingBar.cpp
@@ -0,0 +1,54 @@
+// Copyright 1998-2014 Epic Games, Inc. All Rights Reserved.
+
+#include "UMGEditorPrivatePCH.h"
+
+#include "SDisappearingBar.h"
+
+#define LOCTEXT_NAMESPACE "UMG"
+
+/////////////////////////////////////////////////////
+// SDisappearingBar
+
+void SDisappearingBar::Construct(const FArguments& InArgs)
+{
+	FadeCurve = FCurveSequence(0, 0.25f);
+	FadeCurve.Play();
+
+	bDisappear = false;
+
+	ColorAndOpacity = TAttribute<FLinearColor>::Create(TAttribute<FLinearColor>::FGetter::CreateSP(this, &SDisappearingBar::GetFadeColorAndOpacity));
+
+	ChildSlot
+	[
+		InArgs._Content.Widget
+	];
+}
+
+void SDisappearingBar::Tick(const FGeometry& AllottedGeometry, const double InCurrentTime, const float InDeltaTime)
+{
+	if ( AllottedGeometry.IsUnderLocation(FSlateApplication::Get().GetCursorPos()) )
+	{
+		if ( !bDisappear )
+		{
+			FadeCurve.Play();
+		}
+
+		bDisappear = true;
+	}
+	else
+	{
+		if ( bDisappear )
+		{
+			FadeCurve.PlayReverse();
+		}
+
+		bDisappear = false;
+	}
+}
+
+FLinearColor SDisappearingBar::GetFadeColorAndOpacity() const
+{
+	return FLinearColor(1, 1, 1, 1 - FadeCurve.GetLerp());
+}
+
+#undef LOCTEXT_NAMESPACE
diff --git a/Engine/Source/Editor/UMGEditor/Private/Designer/SDisappearingBar.h b/Engine/Source/Editor/UMGEditor/Private/Designer/SDisappearingBar.h
new file mode 100644
--- /dev/null
+++ b/Engine/Source/Editor/UMGEditor/Private/Designer/SDisappearingBar.h
@@ -0,0 +1,28 @@
+// Copyright 1998-2014 Epic Games, Inc. All Rights Reserved.
+
+#pragma once
+
+#include "SCompoundWidget.h"
+
+/**
+ * A widget that shows another widget as long as the mouse isn't hovering over it.
+ */
+class SDisappearingBar : public SCompoundWidget
+{
+public:
+
+	SLATE_BEGIN_ARGS(SDisappearingBar) {}
+		/** Slot for this designers content (optional) */
+		SLATE_DEFAULT_SLOT(FArguments, Content)
+	SLATE_END_ARGS()
+
+	void Construct(const FArguments& InArgs);
+	
+	void Tick(const FGeometry& AllottedGeometry, const double InCurrentTime, const float InDeltaTime) override;
+
+private:
+	FLinearColor GetFadeColorAndOpacity() const;
+
+	bool bDisappear;
+	FCurveSequence FadeCurve;
+};
diff --git a/Engine/Source/Runtime/UMG/Public/Blueprint/UserWidget.h b/Engine/Source/Runtime/UMG/Public/Blueprint/UserWidget.h
--- a/Engine/Source/Runtime/UMG/Public/Blueprint/UserWidget.h
+++ b/Engine/Source/Runtime/UMG/Public/Blueprint/UserWidget.h
@@ -277,10 +277,6 @@ public:
 	UPROPERTY(EditAnywhere, BlueprintReadWrite, AdvancedDisplay, Category="Behavior")
 	bool bSupportsKeyboardFocus;
 
-	/** Called when the visibility changes. */
-	UPROPERTY(BlueprintAssignable)
-	FOnVisibilityChangedEvent OnVisibilityChanged;
-
 	UPROPERTY()
 	FMargin Padding;
 
diff --git a/Engine/Source/Runtime/UMG/Public/Components/CanvasPanel.h b/Engine/Source/Runtime/UMG/Public/Components/CanvasPanel.h
--- a/Engine/Source/Runtime/UMG/Public/Components/CanvasPanel.h
+++ b/Engine/Source/Runtime/UMG/Public/Components/CanvasPanel.h
@@ -26,6 +26,13 @@ public:
 	// UWidget interface
 	virtual const FSlateBrush* GetEditorIcon() override;
 	// End UWidget interface
+
+	// UWidget interface
+	virtual bool LockToPanelOnDrag() const override
+	{
+		return true;
+	}
+	// End UWidget interface
 #endif
 
 protected:
diff --git a/Engine/Source/Runtime/UMG/Public/Components/PanelWidget.h b/Engine/Source/Runtime/UMG/Public/Components/PanelWidget.h
--- a/Engine/Source/Runtime/UMG/Public/Components/PanelWidget.h
+++ b/Engine/Source/Runtime/UMG/Public/Components/PanelWidget.h
@@ -70,6 +70,11 @@ public:
 	virtual void ReleaseNativeWidget() override;
 
 #if WITH_EDITOR
+	virtual bool LockToPanelOnDrag() const
+	{
+		return false;
+	}
+
 	virtual void ConnectEditorData() override
 	{
 		for ( UPanelSlot* Slot : Slots )
