UMG - Fixing a bug when dropping on a locked parent.

diff --git a/Engine/Source/Editor/UMGEditor/Private/Designer/SDesignerView.cpp b/Engine/Source/Editor/UMGEditor/Private/Designer/SDesignerView.cpp
--- a/Engine/Source/Editor/UMGEditor/Private/Designer/SDesignerView.cpp
+++ b/Engine/Source/Editor/UMGEditor/Private/Designer/SDesignerView.cpp
@@ -1316,13 +1316,13 @@ UWidget* SDesignerView::ProcessDropAndAddWidget(const FGeometry& MyGeometry, con
 		{
 			DesignerMessage = EDesignerMessage::MoveFromParent;
 
-			if ( UPanelWidget* LockedParent = Cast<UPanelWidget>(bIsPreview ? SelectedDragDropOp->ParentWidget.GetPreview() : SelectedDragDropOp->ParentWidget.GetTemplate()) )
+			if ( UPanelWidget* LockedPreviewParent = Cast<UPanelWidget>(SelectedDragDropOp->ParentWidget.GetPreview()) )
 			{
-				TSharedPtr<SWidget> NewParentWidget = LockedParent->GetCachedWidget();
+				TSharedPtr<SWidget> NewParentWidget = LockedPreviewParent->GetCachedWidget();
 
 				if ( NewParentWidget.IsValid() )
 				{
-					Target = LockedParent;
+					Target = bIsPreview ? SelectedDragDropOp->ParentWidget.GetPreview() : SelectedDragDropOp->ParentWidget.GetTemplate();
 
 					const bool bSuccess = FDesignTimeUtils::GetArrangedWidget(NewParentWidget.ToSharedRef(), ArrangedWidget);
 					check(bSuccess);
