UMG - Clearing the message that appears in the designer on drop.

diff --git a/Engine/Source/Editor/UMGEditor/Private/Designer/SDesignerView.cpp b/Engine/Source/Editor/UMGEditor/Private/Designer/SDesignerView.cpp
--- a/Engine/Source/Editor/UMGEditor/Private/Designer/SDesignerView.cpp
+++ b/Engine/Source/Editor/UMGEditor/Private/Designer/SDesignerView.cpp
@@ -1419,7 +1419,6 @@ FReply SDesignerView::OnDrop(const FGeometry& MyGeometry, const FDragDropEvent&
 {
 	bMouseDown = false;
 	bMovingExistingWidget = false;
-	DesignerMessage = EDesignerMessage::None;
 
 	UWidgetBlueprint* BP = GetBlueprint();
 	
@@ -1443,8 +1442,12 @@ FReply SDesignerView::OnDrop(const FGeometry& MyGeometry, const FDragDropEvent&
 		// Regenerate extension widgets now that we've finished moving or placing the widget.
 		CreateExtensionWidgetsForSelection();
 
+		DesignerMessage = EDesignerMessage::None;
+
 		return FReply::Handled();
 	}
+
+	DesignerMessage = EDesignerMessage::None;
 	
 	return FReply::Unhandled();
 }
