Fixed a crash on widget blueprint compilation on Mac (Mac implementation of FTextInputMethodChangeNotifier() needs to find a window which makes Slate call SDesignerView::GetExtensionPosition() and a couple of more functions while a preview widget is NULL)

diff --git a/Engine/Source/Editor/UMGEditor/Private/Designer/SDesignerView.cpp b/Engine/Source/Editor/UMGEditor/Private/Designer/SDesignerView.cpp
--- a/Engine/Source/Editor/UMGEditor/Private/Designer/SDesignerView.cpp
+++ b/Engine/Source/Editor/UMGEditor/Private/Designer/SDesignerView.cpp
@@ -1,4 +1,4 @@
-﻿// Copyright 1998-2014 Epic Games, Inc. All Rights Reserved.
+// Copyright 1998-2014 Epic Games, Inc. All Rights Reserved.
 
 #include "UMGEditorPrivatePCH.h"
 #include "SDesignerView.h"
@@ -467,7 +467,8 @@ FVector2D SDesignerView::GetExtensionPosition(TSharedRef<FDesignerSurfaceElement
 	{
 		FWidgetReference ParentRef = BlueprintEditor.Pin()->GetReferenceFromTemplate(SelectedWidget.GetTemplate()->GetParent());
 
-		TSharedPtr<SWidget> PreviewSlateWidget = ParentRef.GetPreview()->GetCachedWidget();
+		UWidget* Preview = ParentRef.GetPreview();
+		TSharedPtr<SWidget> PreviewSlateWidget = Preview ? Preview->GetCachedWidget() : NULL;
 		if ( PreviewSlateWidget.IsValid() )
 		{
 			FWidgetPath WidgetPath;
diff --git a/Engine/Source/Editor/UMGEditor/Private/Extensions/CanvasSlotExtension.cpp b/Engine/Source/Editor/UMGEditor/Private/Extensions/CanvasSlotExtension.cpp
--- a/Engine/Source/Editor/UMGEditor/Private/Extensions/CanvasSlotExtension.cpp
+++ b/Engine/Source/Editor/UMGEditor/Private/Extensions/CanvasSlotExtension.cpp
@@ -133,7 +133,7 @@ EVisibility FCanvasSlotExtension::GetAnchorVisibility(EAnchorWidget::Type Anchor
 	for ( const FWidgetReference& Selection : SelectionCache )
 	{
 		UWidget* PreviewWidget = Selection.GetPreview();
-		if ( PreviewWidget->Slot )
+		if ( PreviewWidget && PreviewWidget->Slot )
 		{
 			UCanvasPanelSlot* PreviewCanvasSlot = CastChecked<UCanvasPanelSlot>(PreviewWidget->Slot);
 			switch ( AnchorType )
@@ -160,7 +160,7 @@ FVector2D FCanvasSlotExtension::GetAnchorAlignment(EAnchorWidget::Type AnchorTyp
 	for ( const FWidgetReference& Selection : SelectionCache )
 	{
 		UWidget* PreviewWidget = Selection.GetPreview();
-		if ( PreviewWidget->Slot )
+		if ( PreviewWidget && PreviewWidget->Slot )
 		{
 			UCanvasPanelSlot* PreviewCanvasSlot = CastChecked<UCanvasPanelSlot>(PreviewWidget->Slot);
 
