UMG - Adding render transforms to the base widget class.  Beginning work on the designer interface for them.  Lots of style refactoring for a few widgets to make it far easier to change their appearance at runtime.  Adding slate utility functions to accessing geometry data in blueprints.  The math struct customization can now apply min/max UI slider clamps on their components.  Disabling the preview on the slate brush customization.  Will re-introduce later when we can provide instanced struct customizations.  The thumbnails in the details panel now have a hover effect to indicate that they are double clickable.  Adding a Delta metadata on numeric properties so that if you want the slider to move in specific increments you can now.

diff --git a/Engine/Source/Editor/DetailCustomizations/Private/MathStructCustomizations.cpp b/Engine/Source/Editor/DetailCustomizations/Private/MathStructCustomizations.cpp
--- a/Engine/Source/Editor/DetailCustomizations/Private/MathStructCustomizations.cpp
+++ b/Engine/Source/Editor/DetailCustomizations/Private/MathStructCustomizations.cpp
@@ -66,7 +66,7 @@ void FMathStructCustomization::MakeHeaderRow( TSharedRef<class IPropertyHandle>&
 		HorizontalBox->AddSlot()
 		.Padding( FMargin(0.0f, 2.0f, bLastChild ? 0.0f : 3.0f, 2.0f ) )
 		[
-			MakeChildWidget( ChildHandle )
+			MakeChildWidget(StructPropertyHandle, ChildHandle)
 		];
 	}
 
@@ -131,8 +131,79 @@ void FMathStructCustomization::GetSortedChildren( TSharedRef<IPropertyHandle> St
 }
 
 template<typename NumericType>
-TSharedRef<SWidget> FMathStructCustomization::MakeNumericWidget(TSharedRef<IPropertyHandle>& PropertyHandle)
+void ExtractNumericMetadata(TSharedRef<IPropertyHandle>& PropertyHandle, TOptional<NumericType>& MinValue, TOptional<NumericType>& MaxValue, TOptional<NumericType>& SliderMinValue, TOptional<NumericType>& SliderMaxValue, NumericType& SliderExponent, NumericType& Delta)
 {
+	UProperty* Property = PropertyHandle->GetProperty();
+
+	const FString& MetaUIMinString = Property->GetMetaData(TEXT("UIMin"));
+	const FString& MetaUIMaxString = Property->GetMetaData(TEXT("UIMax"));
+	const FString& SliderExponentString = Property->GetMetaData(TEXT("SliderExponent"));
+	const FString& DeltaString = Property->GetMetaData(TEXT("Delta"));
+	const FString& ClampMinString = Property->GetMetaData(TEXT("ClampMin"));
+	const FString& ClampMaxString = Property->GetMetaData(TEXT("ClampMax"));
+
+	// If no UIMin/Max was specified then use the clamp string
+	const FString& UIMinString = MetaUIMinString.Len() ? MetaUIMinString : ClampMinString;
+	const FString& UIMaxString = MetaUIMaxString.Len() ? MetaUIMaxString : ClampMaxString;
+
+	NumericType ClampMin = TNumericLimits<NumericType>::Lowest();
+	NumericType ClampMax = TNumericLimits<NumericType>::Max();
+
+	if ( !ClampMinString.IsEmpty() )
+	{
+		TTypeFromString<NumericType>::FromString(ClampMin, *ClampMinString);
+	}
+
+	if ( !ClampMaxString.IsEmpty() )
+	{
+		TTypeFromString<NumericType>::FromString(ClampMax, *ClampMaxString);
+	}
+
+	NumericType UIMin = TNumericLimits<NumericType>::Lowest();
+	NumericType UIMax = TNumericLimits<NumericType>::Max();
+	TTypeFromString<NumericType>::FromString(UIMin, *UIMinString);
+	TTypeFromString<NumericType>::FromString(UIMax, *UIMaxString);
+
+	SliderExponent = NumericType(1);
+	if ( SliderExponentString.Len() )
+	{
+		TTypeFromString<NumericType>::FromString(SliderExponent, *SliderExponentString);
+	}
+
+	Delta = NumericType(0);
+	if ( DeltaString.Len() )
+	{
+		TTypeFromString<NumericType>::FromString(Delta, *DeltaString);
+	}
+
+	if ( ClampMin >= ClampMax && ( ClampMinString.Len() || ClampMaxString.Len() ) )
+	{
+		//UE_LOG(LogPropertyNode, Warning, TEXT("Clamp Min (%s) >= Clamp Max (%s) for Ranged Numeric"), *ClampMinString, *ClampMaxString);
+	}
+
+	const NumericType ActualUIMin = FMath::Max(UIMin, ClampMin);
+	const NumericType ActualUIMax = FMath::Min(UIMax, ClampMax);
+
+	MinValue = ClampMinString.Len() ? ClampMin : TOptional<NumericType>();
+	MaxValue = ClampMaxString.Len() ? ClampMax : TOptional<NumericType>();
+	SliderMinValue = ( UIMinString.Len() ) ? ActualUIMin : TOptional<NumericType>();
+	SliderMaxValue = ( UIMaxString.Len() ) ? ActualUIMax : TOptional<NumericType>();
+
+	if ( ActualUIMin >= ActualUIMax && ( MetaUIMinString.Len() || MetaUIMaxString.Len() ) )
+	{
+		//UE_LOG(LogPropertyNode, Warning, TEXT("UI Min (%s) >= UI Max (%s) for Ranged Numeric"), *UIMinString, *UIMaxString);
+	}
+}
+
+template<typename NumericType>
+TSharedRef<SWidget> FMathStructCustomization::MakeNumericWidget(
+	TSharedRef<IPropertyHandle>& StructurePropertyHandle,
+	TSharedRef<IPropertyHandle>& PropertyHandle)
+{
+	TOptional<NumericType> MinValue, MaxValue, SliderMinValue, SliderMaxValue;
+	NumericType SliderExponent, Delta;
+	ExtractNumericMetadata(StructurePropertyHandle, MinValue, MaxValue, SliderMinValue, SliderMaxValue, SliderExponent, Delta);
+
 	TWeakPtr<IPropertyHandle> WeakHandlePtr = PropertyHandle;
 
 	return 
@@ -148,10 +219,12 @@ TSharedRef<SWidget> FMathStructCustomization::MakeNumericWidget(TSharedRef<IProp
 		.LabelVAlign(VAlign_Center)
 		// Only allow spin on handles with one object.  Otherwise it is not clear what value to spin
 		.AllowSpin( PropertyHandle->GetNumOuterObjects() == 1 )
-		.MinValue(TOptional<NumericType>())
-		.MaxValue(TOptional<NumericType>())
-		.MaxSliderValue(TOptional<NumericType>())
-		.MinSliderValue(TOptional<NumericType>())
+		.MinValue(MinValue)
+		.MaxValue(MaxValue)
+		.MinSliderValue(SliderMinValue)
+		.MaxSliderValue(SliderMaxValue)
+		.SliderExponent(SliderExponent)
+		.Delta(Delta)
 		.Label()
 		[
 			SNew( STextBlock )
@@ -160,23 +233,25 @@ TSharedRef<SWidget> FMathStructCustomization::MakeNumericWidget(TSharedRef<IProp
 		];
 }
 
-TSharedRef<SWidget> FMathStructCustomization::MakeChildWidget( TSharedRef<IPropertyHandle>& PropertyHandle )
+TSharedRef<SWidget> FMathStructCustomization::MakeChildWidget( 
+	TSharedRef<IPropertyHandle>& StructurePropertyHandle,
+	TSharedRef<IPropertyHandle>& PropertyHandle )
 {
 	const UClass* PropertyClass = PropertyHandle->GetPropertyClass();
 	
 	if (PropertyClass == UFloatProperty::StaticClass())
 	{
-		return MakeNumericWidget<float>(PropertyHandle);
+		return MakeNumericWidget<float>(StructurePropertyHandle, PropertyHandle);
 	}
 	
 	if (PropertyClass == UIntProperty::StaticClass())
 	{
-		return MakeNumericWidget<int32>(PropertyHandle);
+		return MakeNumericWidget<int32>(StructurePropertyHandle, PropertyHandle);
 	}
 
 	if (PropertyClass == UByteProperty::StaticClass())
 	{
-		return MakeNumericWidget<uint8>(PropertyHandle);
+		return MakeNumericWidget<uint8>(StructurePropertyHandle, PropertyHandle);
 	}
 
 	check(0); // Unsupported class
diff --git a/Engine/Source/Editor/DetailCustomizations/Private/MathStructCustomizations.h b/Engine/Source/Editor/DetailCustomizations/Private/MathStructCustomizations.h
--- a/Engine/Source/Editor/DetailCustomizations/Private/MathStructCustomizations.h
+++ b/Engine/Source/Editor/DetailCustomizations/Private/MathStructCustomizations.h
@@ -41,9 +41,12 @@ protected:
 	/**
 	 * Constructs a widget for the property handle 
 	 *
+	 * @param StructurePropertyHandle	handle of the struct property
 	 * @param PropertyHandle	Child handle of the struct property
 	 */
-	virtual TSharedRef<SWidget> MakeChildWidget( TSharedRef<IPropertyHandle>& PropertyHandle );
+	virtual TSharedRef<SWidget> MakeChildWidget( 
+		TSharedRef<IPropertyHandle>& StructurePropertyHandle,
+		TSharedRef<IPropertyHandle>& PropertyHandle );
 
 	/**
 	 * Gets the value as a float for the provided property handle
@@ -109,7 +112,7 @@ private:
 	void OnEndSliderMovement( NumericType NewValue );
 
 	template<typename NumericType>
-	TSharedRef<SWidget> MakeNumericWidget( TSharedRef<IPropertyHandle>& PropertyHandle );
+	TSharedRef<SWidget> MakeNumericWidget(TSharedRef<IPropertyHandle>& StructurePropertyHandle, TSharedRef<IPropertyHandle>& PropertyHandle);
 
 protected:
 	/** All the sorted children of the struct that should be displayed */
diff --git a/Engine/Source/Editor/DetailCustomizations/Private/SlateBrushCustomization.cpp b/Engine/Source/Editor/DetailCustomizations/Private/SlateBrushCustomization.cpp
--- a/Engine/Source/Editor/DetailCustomizations/Private/SlateBrushCustomization.cpp
+++ b/Engine/Source/Editor/DetailCustomizations/Private/SlateBrushCustomization.cpp
@@ -806,7 +806,6 @@ private:
 	void OnBrushResourceChanged()
 	{
 		CachePropertyValues();
-		ImageSizeProperty->SetValue( CachedTextureSize );
 		UpdatePreviewImageSize();
 		UpdateMarginLinePositions();
 	}
@@ -942,15 +941,72 @@ private:
 const float SSlateBrushPreview::ImagePadding = 5.0f;
 const float SSlateBrushPreview::BorderHitSize = 8.0f;
 
+// SSlateBrushStaticPreview
+////////////////////////////////////////////////////////////////////////////////
+
+class SSlateBrushStaticPreview : public SCompoundWidget
+{
+	SLATE_BEGIN_ARGS(SSlateBrushStaticPreview) {}
+
+	SLATE_END_ARGS()
+
+	void Construct( const FArguments& InArgs, TSharedPtr<IPropertyHandle> InResourceObjectProperty )
+	{
+		ResourceObjectProperty = InResourceObjectProperty;
+
+		ChildSlot
+		[
+			SNew(SBorder)
+			.BorderImage(this, &SSlateBrushStaticPreview::GetPropertyBrush)
+			[
+				SNew(SSpacer)
+				.Size(FVector2D(1,1))
+			]
+		];
+	}
+
+	void Tick(const FGeometry& AllottedGeometry, const double InCurrentTime, const float InDeltaTime)
+	{
+		TArray<void*> RawData;
+		ResourceObjectProperty->AccessRawData(RawData);
+
+		check(RawData[0] != NULL);
+		TemporaryBrush = *static_cast<FSlateBrush*>( RawData[0] );
+	}
+
+	const FSlateBrush* SSlateBrushStaticPreview::GetPropertyBrush() const
+	{
+		return &TemporaryBrush;
+	}
+
+private:
+
+	/**
+	 * Temporary brush data used to store the structure returned from the property handy so that we have stable
+	 * pointer to give to slate.
+	 */
+	FSlateBrush TemporaryBrush;
+
+	TSharedPtr<IPropertyHandle> ResourceObjectProperty;
+};
+
+// SBrushResourceObjectBox
+////////////////////////////////////////////////////////////////////////////////
+
 class SBrushResourceObjectBox : public SCompoundWidget
 {
 	SLATE_BEGIN_ARGS( SBrushResourceObjectBox ) {}
 	
 	SLATE_END_ARGS()
 
-	void Construct( const FArguments& InArgs, IStructCustomizationUtils* StructCustomizationUtils, TSharedPtr<IPropertyHandle> InResourceObjectProperty )
+	void Construct(const FArguments& InArgs, IStructCustomizationUtils* StructCustomizationUtils, TSharedPtr<IPropertyHandle> InResourceObjectProperty, TSharedPtr<IPropertyHandle> InImageSizeProperty)
 	{
 		ResourceObjectProperty = InResourceObjectProperty;
+		ImageSizeProperty = InImageSizeProperty;
+
+		FSimpleDelegate OnBrushResourceChangedDelegate = FSimpleDelegate::CreateSP(this, &SBrushResourceObjectBox::OnBrushResourceChanged);
+		ResourceObjectProperty->SetOnPropertyValueChanged(OnBrushResourceChangedDelegate);
+
 		ChildSlot
 		[
 			SNew(SVerticalBox)
@@ -961,6 +1017,7 @@ class SBrushResourceObjectBox : public SCompoundWidget
 				.PropertyHandle(InResourceObjectProperty)
 				.ThumbnailPool(StructCustomizationUtils->GetThumbnailPool())
 				.OnShouldFilterAsset(this, &SBrushResourceObjectBox::OnFilterAssetPicker)
+				.OnObjectChanged(this, &SBrushResourceObjectBox::OnAssetPicked)
 			]
 			+ SVerticalBox::Slot()
 			.AutoHeight()
@@ -996,6 +1053,31 @@ class SBrushResourceObjectBox : public SCompoundWidget
 
 private:
 
+	void OnBrushResourceChanged()
+	{
+		UObject* ResourceObject;
+		FPropertyAccess::Result Result = ResourceObjectProperty->GetValue(ResourceObject);
+		if ( Result == FPropertyAccess::Success )
+		{
+			FVector2D CachedTextureSize;
+
+			TArray<void*> RawData;
+			ImageSizeProperty->AccessRawData(RawData);
+			if ( RawData.Num() > 0 && RawData[0] != NULL )
+			{
+				CachedTextureSize = *static_cast<FVector2D*>( RawData[0] );
+			}
+
+			UTexture2D* BrushTexture = Cast<UTexture2D>(ResourceObject);
+			if ( BrushTexture )
+			{
+				CachedTextureSize = FVector2D(BrushTexture->GetSizeX(), BrushTexture->GetSizeY());
+			}
+
+			ImageSizeProperty->SetValue(CachedTextureSize);
+		}
+	}
+
 	/** Called when the asset picker needs to be filtered */
 	bool OnFilterAssetPicker(const FAssetData& InAssetData) const
 	{
@@ -1003,16 +1085,45 @@ private:
 		return !Class->IsChildOf(UTexture2D::StaticClass()) && !Class->IsChildOf(UMaterialInterface::StaticClass());
 	}
 
+	/**
+	 * When the asset is picked if it's a material being used in the UI we 
+	 * automatically set the bUsedWithUI flag if it isn't already set.
+	 */
+	void OnAssetPicked(const FAssetData& InAssetData) const
+	{
+		UObject* Resource = nullptr;
+
+		if ( ResourceObjectProperty->GetValue(Resource) == FPropertyAccess::Success && Resource && Resource->IsA<UMaterialInterface>() )
+		{
+			UMaterialInterface* MaterialInterface = Cast<UMaterialInterface>(Resource);
+			UMaterial* BaseMaterial = MaterialInterface->GetBaseMaterial();
+			if ( BaseMaterial && !BaseMaterial->bUsedWithUI )
+			{
+				bool bNeedsRecompile = true;
+				BaseMaterial->SetMaterialUsage(bNeedsRecompile, MATUSAGE_UI);
+			}
+		}
+	}
+
 private:
 	TSharedPtr<IPropertyHandle> ResourceObjectProperty;
+	TSharedPtr<IPropertyHandle> ImageSizeProperty;
 	TSharedPtr<SErrorText> ResourceErrorText;
 };
 
+// FSlateBrushStructCustomization
+////////////////////////////////////////////////////////////////////////////////
+
 TSharedRef<IPropertyTypeCustomization> FSlateBrushStructCustomization::MakeInstance() 
 {
 	return MakeShareable( new FSlateBrushStructCustomization() );
 }
 
+FSlateBrushStructCustomization::FSlateBrushStructCustomization()
+	: bSlimView(true)
+{
+}
+
 void FSlateBrushStructCustomization::CustomizeHeader( TSharedRef<IPropertyHandle> StructPropertyHandle, FDetailWidgetRow& HeaderRow, IPropertyTypeCustomizationUtils& StructCustomizationUtils )
 {
 	bool ShowOnlyInnerProperties = StructPropertyHandle->GetProperty()->HasMetaData(TEXT("ShowOnlyInnerProperties"));
@@ -1023,6 +1134,10 @@ void FSlateBrushStructCustomization::CustomizeHeader( TSharedRef<IPropertyHandle
 		.NameContent()
 		[
 			StructPropertyHandle->CreatePropertyNameWidget()
+		]
+		.ValueContent()
+		[
+			SNew(SSlateBrushStaticPreview, StructPropertyHandle)
 		];
 	}
 }
@@ -1036,7 +1151,7 @@ void FSlateBrushStructCustomization::CustomizeChildren( TSharedRef<IPropertyHand
 	TSharedPtr<IPropertyHandle> MarginProperty = StructPropertyHandle->GetChildHandle( TEXT("Margin") );
 	TSharedPtr<IPropertyHandle> TintProperty = StructPropertyHandle->GetChildHandle( TEXT("TintColor") );
 	TSharedPtr<IPropertyHandle> ResourceObjectProperty = StructPropertyHandle->GetChildHandle( TEXT("ResourceObject") );
-
+	
 	StructBuilder.AddChildContent( NSLOCTEXT( "SlateBrushCustomization", "ResourceObjectFilterString", "Resource" ).ToString() )
 	.NameContent()
 	[
@@ -1046,7 +1161,7 @@ void FSlateBrushStructCustomization::CustomizeChildren( TSharedRef<IPropertyHand
 	.MinDesiredWidth(250.0f)
 	.MaxDesiredWidth(0.0f)
 	[
-		SNew( SBrushResourceObjectBox, &StructCustomizationUtils, ResourceObjectProperty )
+		SNew( SBrushResourceObjectBox, &StructCustomizationUtils, ResourceObjectProperty, ImageSizeProperty )
 	];
 
 	StructBuilder.AddChildProperty( ImageSizeProperty.ToSharedRef() );
@@ -1057,50 +1172,53 @@ void FSlateBrushStructCustomization::CustomizeChildren( TSharedRef<IPropertyHand
 	StructBuilder.AddChildProperty( MarginProperty.ToSharedRef() )
 	.Visibility( TAttribute<EVisibility>::Create( TAttribute<EVisibility>::FGetter::CreateSP( this, &FSlateBrushStructCustomization::GetMarginPropertyVisibility ) ) );
 
-	// Create the Slate Brush Preview widget and add the Preview group
-	TArray<void*> RawData;
-	StructPropertyHandle->AccessRawData( RawData );
-
-	// Can only display the preview with one brush
-	if( RawData.Num() == 1 )
+	// Don't show the preview area when in slim view mode.
+	if ( !bSlimView )
 	{
-		FSlateBrush* Brush = static_cast<FSlateBrush*>(RawData[ 0 ]);
+		// Create the Slate Brush Preview widget and add the Preview group
+		TArray<void*> RawData;
+		StructPropertyHandle->AccessRawData(RawData);
 
-		TSharedRef<SSlateBrushPreview> Preview = SNew( SSlateBrushPreview )
-			.DrawAsProperty( DrawAsProperty )
-			.TilingProperty( TilingProperty )
-			.ImageSizeProperty( ImageSizeProperty )
-			.MarginProperty( MarginProperty )
-			.ResourceObjectProperty( ResourceObjectProperty )
-			.SlateBrush( Brush );
+		// Can only display the preview with one brush
+		if ( RawData.Num() == 1 )
+		{
+			FSlateBrush* Brush = static_cast<FSlateBrush*>( RawData[0] );
 
-		IDetailGroup& PreviewGroup = StructBuilder.AddChildGroup( TEXT( "Preview" ), TEXT("") );
+			TSharedRef<SSlateBrushPreview> Preview = SNew(SSlateBrushPreview)
+				.DrawAsProperty(DrawAsProperty)
+				.TilingProperty(TilingProperty)
+				.ImageSizeProperty(ImageSizeProperty)
+				.MarginProperty(MarginProperty)
+				.ResourceObjectProperty(ResourceObjectProperty)
+				.SlateBrush(Brush);
 
-		PreviewGroup
-			.HeaderRow()
-			.NameContent()
-			[
-				StructPropertyHandle->CreatePropertyNameWidget( NSLOCTEXT( "UnrealEd", "Preview", "Preview" ).ToString(), false )
-			]
-		.ValueContent()
-			.MinDesiredWidth( 1 )
-			.MaxDesiredWidth( 4096 )
-			[
-				Preview->GenerateAlignmentComboBoxes()
-			];
+			IDetailGroup& PreviewGroup = StructBuilder.AddChildGroup(TEXT("Preview"), TEXT(""));
 
-		PreviewGroup
-			.AddWidgetRow()
+			PreviewGroup
+				.HeaderRow()
+				.NameContent()
+				[
+					StructPropertyHandle->CreatePropertyNameWidget(NSLOCTEXT("UnrealEd", "Preview", "Preview").ToString(), false)
+				]
 			.ValueContent()
-			.MinDesiredWidth( 1 )
-			.MaxDesiredWidth( 4096 )
-			[
-				Preview
-			];
+				.MinDesiredWidth(1)
+				.MaxDesiredWidth(4096)
+				[
+					Preview->GenerateAlignmentComboBoxes()
+				];
+
+			PreviewGroup
+				.AddWidgetRow()
+				.ValueContent()
+				.MinDesiredWidth(1)
+				.MaxDesiredWidth(4096)
+				[
+					Preview
+				];
+		}
 	}
 }
 
-
 EVisibility FSlateBrushStructCustomization::GetTilingPropertyVisibility() const
 {
 	uint8 DrawAsType;
diff --git a/Engine/Source/Editor/DetailCustomizations/Private/SlateBrushCustomization.h b/Engine/Source/Editor/DetailCustomizations/Private/SlateBrushCustomization.h
--- a/Engine/Source/Editor/DetailCustomizations/Private/SlateBrushCustomization.h
+++ b/Engine/Source/Editor/DetailCustomizations/Private/SlateBrushCustomization.h
@@ -8,6 +8,8 @@ class FSlateBrushStructCustomization : public IPropertyTypeCustomization
 public:
 	static TSharedRef<IPropertyTypeCustomization> MakeInstance();
 
+	FSlateBrushStructCustomization();
+
 	/** IPropertyTypeCustomization interface */
 	virtual void CustomizeHeader( TSharedRef<class IPropertyHandle> StructPropertyHandle, class FDetailWidgetRow& HeaderRow, IPropertyTypeCustomizationUtils& StructCustomizationUtils ) override;
 
@@ -24,11 +26,13 @@ private:
 	 */
 	EVisibility GetMarginPropertyVisibility() const;
 
-
 	/** Slate Brush DrawAs property */
 	TSharedPtr<IPropertyHandle> DrawAsProperty;
 
 	/** Error text to display if the resource object is not valid*/
 	TSharedPtr<SErrorText> ResourceErrorText;
+
+	/**  */
+	bool bSlimView;
 };
 
diff --git a/Engine/Source/Editor/EditorStyle/Private/SlateEditorStyle.cpp b/Engine/Source/Editor/EditorStyle/Private/SlateEditorStyle.cpp
--- a/Engine/Source/Editor/EditorStyle/Private/SlateEditorStyle.cpp
+++ b/Engine/Source/Editor/EditorStyle/Private/SlateEditorStyle.cpp
@@ -5686,6 +5686,14 @@ void FSlateEditorStyle::FStyle::SetupAutomationStyles()
 
 void FSlateEditorStyle::FStyle::SetupUMGEditorStyles()
 {
+	Set("WidgetDesigner.LayoutTransform", new IMAGE_BRUSH("Icons/UMG/Layout_TransformMode_16x", Icon16x16));
+	Set("WidgetDesigner.LayoutTransform.Small", new IMAGE_BRUSH("Icons/UMG/Layout_TransformMode_16x", Icon16x16));
+	Set("WidgetDesigner.RenderTransform", new IMAGE_BRUSH("Icons/UMG/Render_TransformMode_16x", Icon16x16));
+	Set("WidgetDesigner.RenderTransform.Small", new IMAGE_BRUSH("Icons/UMG/Render_TransformMode_16x", Icon16x16));
+
+	Set("WidgetDesigner.ZoomToFit", new IMAGE_BRUSH("Icons/UMG/Fit_16x", Icon16x16));
+	Set("WidgetDesigner.ZoomToFit.Small", new IMAGE_BRUSH("Icons/UMG/Fit_16x", Icon16x16));
+
 	Set("UMGEditor.ZoomToFit", new IMAGE_BRUSH("Icons/UMG/Fit_16x", Icon16x16, FLinearColor(.05f, .05f, .05f, 1.f)));
 
 	Set("UMGEditor.AnchorGizmo.Center", new IMAGE_BRUSH("Icons/UMG/AnchorGizmo/center", Icon16x16));
diff --git a/Engine/Source/Editor/PropertyEditor/Private/PropertyCustomizationHelpers.cpp b/Engine/Source/Editor/PropertyEditor/Private/PropertyCustomizationHelpers.cpp
--- a/Engine/Source/Editor/PropertyEditor/Private/PropertyCustomizationHelpers.cpp
+++ b/Engine/Source/Editor/PropertyEditor/Private/PropertyCustomizationHelpers.cpp
@@ -281,6 +281,7 @@ void SObjectPropertyEntryBox::Construct( const FArguments& InArgs )
 	OnObjectChanged = InArgs._OnObjectChanged;
 
 	bool bDisplayThumbnail = false;
+	FIntPoint ThumbnailSize(64, 64);
 
 	if( InArgs._PropertyHandle.IsValid() && InArgs._PropertyHandle->IsValidHandle() )
 	{
@@ -293,6 +294,18 @@ void SObjectPropertyEntryBox::Construct( const FArguments& InArgs )
 			bDisplayThumbnail = DisplayThumbnailString == TEXT("true");
 		}
 
+		// check if the property metadata has an override to the thumbnail size
+		FString ThumbnailSizeString = PropertyHandle->GetProperty()->GetMetaData(TEXT("ThumbnailSize"));
+		if ( ThumbnailSizeString.Len() > 0 )
+		{
+			FVector2D ParsedVector;
+			if ( ParsedVector.InitFromString(ThumbnailSizeString) )
+			{
+				ThumbnailSize.X = (int32)ParsedVector.X;
+				ThumbnailSize.Y = (int32)ParsedVector.Y;
+			}
+		}
+
 		// if being used with an object property, check the allowed class is valid for the property
 		UObjectPropertyBase* ObjectProperty = Cast<UObjectPropertyBase>(PropertyHandle->GetProperty());
 		if (ObjectProperty != NULL)
@@ -317,6 +330,7 @@ void SObjectPropertyEntryBox::Construct( const FArguments& InArgs )
 				.OnShouldFilterAsset(InArgs._OnShouldFilterAsset)
 				.AllowClear(InArgs._AllowClear)
 				.PropertyHandle(PropertyHandle)
+				.ThumbnailSize(ThumbnailSize)
 		]
 	];
 }
diff --git a/Engine/Source/Editor/PropertyEditor/Private/UserInterface/PropertyEditor/SPropertyEditorAsset.cpp b/Engine/Source/Editor/PropertyEditor/Private/UserInterface/PropertyEditor/SPropertyEditorAsset.cpp
--- a/Engine/Source/Editor/PropertyEditor/Private/UserInterface/PropertyEditor/SPropertyEditorAsset.cpp
+++ b/Engine/Source/Editor/PropertyEditor/Private/UserInterface/PropertyEditor/SPropertyEditorAsset.cpp
@@ -116,9 +116,9 @@ void SPropertyEditorAsset::Construct( const FArguments& InArgs, const TSharedPtr
 				.Padding( 0.0f, 0.0f, 2.0f, 0.0f )
 				.AutoWidth()
 				[
-					SNew(SBorder)
-					.Padding(2.0f)
-					.BorderImage( FEditorStyle::GetBrush("PropertyEditor.AssetThumbnailShadow") )
+					SAssignNew( ThumbnailBorder, SBorder )
+					.Padding( 3.0f )
+					.BorderImage( this, &SPropertyEditorAsset::GetThumbnailBorder )
 					.OnMouseDoubleClick( this, &SPropertyEditorAsset::OnAssetThumbnailDoubleClick )
 					[
 						SNew( SBox )
@@ -148,7 +148,7 @@ void SPropertyEditorAsset::Construct( const FArguments& InArgs, const TSharedPtr
 			.IsEnabled( this, &SPropertyEditorAsset::CanEdit )
 			+ SHorizontalBox::Slot()
 			[
-				SAssignNew(AssetComboButton, SComboButton)
+				SAssignNew( AssetComboButton, SComboButton )
 				.ToolTipText( this, &SPropertyEditorAsset::OnGetToolTip )
 				.ButtonStyle( FEditorStyle::Get(), "PropertyEditor.AssetComboStyle" )
 				.ForegroundColor(FEditorStyle::GetColor("PropertyEditor.AssetName.ColorAndOpacity"))
@@ -224,9 +224,21 @@ void SPropertyEditorAsset::GetDesiredWidth( float& OutMinDesiredWidth, float &Ou
 	OutMaxDesiredWidth = 0.0f;
 }
 
+const FSlateBrush* SPropertyEditorAsset::GetThumbnailBorder() const
+{
+	if ( ThumbnailBorder->IsHovered() )
+	{
+		return FEditorStyle::GetBrush("PropertyEditor.AssetThumbnailLight");
+	}
+	else
+	{
+		return FEditorStyle::GetBrush("PropertyEditor.AssetThumbnailShadow");
+	}
+}
+
 void SPropertyEditorAsset::Tick( const FGeometry& AllottedGeometry, const double InCurrentTime, const float InDeltaTime )
 {
-	if(AssetThumbnail.IsValid())
+	if( AssetThumbnail.IsValid() )
 	{
 		// Ensure the thumbnail is up to date
 		FObjectOrAssetData Value;
diff --git a/Engine/Source/Editor/PropertyEditor/Private/UserInterface/PropertyEditor/SPropertyEditorAsset.h b/Engine/Source/Editor/PropertyEditor/Private/UserInterface/PropertyEditor/SPropertyEditorAsset.h
--- a/Engine/Source/Editor/PropertyEditor/Private/UserInterface/PropertyEditor/SPropertyEditorAsset.h
+++ b/Engine/Source/Editor/PropertyEditor/Private/UserInterface/PropertyEditor/SPropertyEditorAsset.h
@@ -88,6 +88,9 @@ private:
 	 */
 	bool ShouldDisplayThumbnail( const FArguments& InArgs );
 
+	/** Gets the border brush to show around the thumbnail, changes when the user hovers on it. */
+	const FSlateBrush* GetThumbnailBorder() const;
+
 	/** 
 	 * Get the content to be displayed in the asset/actor picker menu 
 	 * @returns the widget for the menu content
@@ -233,6 +236,9 @@ private:
 	/** Main combobutton */
 	TSharedPtr< class SComboButton > AssetComboButton;
 
+	/** The border surrounding the thumbnail image. */
+	TSharedPtr< class SBorder > ThumbnailBorder;
+
 	/** The property editor, if any */
 	TSharedPtr<class FPropertyEditor> PropertyEditor;
 
diff --git a/Engine/Source/Editor/PropertyEditor/Private/UserInterface/PropertyEditor/SPropertyEditorNumeric.h b/Engine/Source/Editor/PropertyEditor/Private/UserInterface/PropertyEditor/SPropertyEditorNumeric.h
--- a/Engine/Source/Editor/PropertyEditor/Private/UserInterface/PropertyEditor/SPropertyEditorNumeric.h
+++ b/Engine/Source/Editor/PropertyEditor/Private/UserInterface/PropertyEditor/SPropertyEditorNumeric.h
@@ -28,6 +28,7 @@ public:
 		const FString& MetaUIMinString = Property->GetMetaData(TEXT("UIMin"));
 		const FString& MetaUIMaxString = Property->GetMetaData(TEXT("UIMax"));
 		const FString& SliderExponentString = Property->GetMetaData(TEXT("SliderExponent"));
+		const FString& DeltaString = Property->GetMetaData(TEXT("Delta"));
 		const FString& ClampMinString = Property->GetMetaData(TEXT("ClampMin"));
 		const FString& ClampMaxString = Property->GetMetaData(TEXT("ClampMax"));
 
@@ -58,6 +59,12 @@ public:
 		{
 			TTypeFromString<NumericType>::FromString( SliderExponent, *SliderExponentString );
 		}
+
+		NumericType Delta = NumericType(0);
+		if ( DeltaString.Len() )
+		{
+			TTypeFromString<NumericType>::FromString( Delta, *DeltaString );
+		}
 		
 		if (ClampMin >= ClampMax && (ClampMinString.Len() || ClampMaxString.Len()))
 		{
@@ -93,6 +100,7 @@ public:
 				.MinSliderValue(SliderMinValue)
 				.MaxSliderValue(SliderMaxValue)
 				.SliderExponent(SliderExponent)
+				.Delta(Delta)
 				.UndeterminedString( LOCTEXT("MultipleValues", "Multiple Values").ToString() )
 				.OnValueChanged( this, &SPropertyEditorNumeric<NumericType>::OnValueChanged )
 				.OnValueCommitted( this, &SPropertyEditorNumeric<NumericType>::OnValueCommitted )
@@ -216,7 +224,7 @@ private:
 	}
 	
 	/**
- 	 * Called when the slider begins to move.  We create a transaction here to undo the property
+	 * Called when the slider begins to move.  We create a transaction here to undo the property
 	 */
 	void OnBeginSliderMovement()
 	{
@@ -227,7 +235,7 @@ private:
 
 
 	/**
- 	 * Called when the slider stops moving.  We end the previously created transation
+	 * Called when the slider stops moving.  We end the previously created transation
 	 */
 	void OnEndSliderMovement( NumericType NewValue )
 	{
diff --git a/Engine/Source/Editor/UMGEditor/Private/Designer/DesignerCommands.cpp b/Engine/Source/Editor/UMGEditor/Private/Designer/DesignerCommands.cpp
new file mode 100644
--- /dev/null
+++ b/Engine/Source/Editor/UMGEditor/Private/Designer/DesignerCommands.cpp
@@ -0,0 +1,10 @@
+// Copyright 1998-2014 Epic Games, Inc. All Rights Reserved.
+
+#include "UMGEditorPrivatePCH.h"
+#include "DesignerCommands.h"
+
+void FDesignerCommands::RegisterCommands()
+{
+	UI_COMMAND( LayoutTransform, "Layout Transform Mode", "Adjust widget layout transform", EUserInterfaceActionType::ToggleButton, FInputGesture(EKeys::W) );
+	UI_COMMAND( RenderTransform, "Render Transform Mode", "Adjust widget render transform", EUserInterfaceActionType::ToggleButton, FInputGesture(EKeys::E) );
+}
diff --git a/Engine/Source/Editor/UMGEditor/Private/Designer/DesignerCommands.h b/Engine/Source/Editor/UMGEditor/Private/Designer/DesignerCommands.h
new file mode 100644
--- /dev/null
+++ b/Engine/Source/Editor/UMGEditor/Private/Designer/DesignerCommands.h
@@ -0,0 +1,33 @@
+// Copyright 1998-2014 Epic Games, Inc. All Rights Reserved.
+
+#pragma once 
+
+/**
+ * 
+ */
+class FDesignerCommands : public TCommands<FDesignerCommands>
+{
+public:
+	FDesignerCommands()
+		: TCommands<FDesignerCommands>
+		(
+			TEXT("WidgetDesigner"), // Context name for fast lookup
+			NSLOCTEXT("Contexts", "DesignerCommands", "Common Designer Commands"), // Localized context name for displaying
+			NAME_None, // Parent
+			FEditorStyle::GetStyleSetName() // Icon Style Set
+		)
+	{
+	}
+	
+	/** Layout Mode */
+	TSharedPtr< FUICommandInfo > LayoutTransform;
+
+	/** Render Mode */
+	TSharedPtr< FUICommandInfo > RenderTransform;
+
+public:
+	/** Registers our commands with the binding system */
+	virtual void RegisterCommands() override;
+};
+
+
diff --git a/Engine/Source/Editor/UMGEditor/Private/Designer/SDesignerToolBar.cpp b/Engine/Source/Editor/UMGEditor/Private/Designer/SDesignerToolBar.cpp
new file mode 100644
--- /dev/null
+++ b/Engine/Source/Editor/UMGEditor/Private/Designer/SDesignerToolBar.cpp
@@ -0,0 +1,46 @@
+// Copyright 1998-2014 Epic Games, Inc. All Rights Reserved.
+
+#include "UMGEditorPrivatePCH.h"
+
+#include "SDesignerToolBar.h"
+#include "DesignerCommands.h"
+
+#define LOCTEXT_NAMESPACE "UMG"
+
+void SDesignerToolBar::Construct( const FArguments& InArgs )
+{
+	CommandList = InArgs._CommandList;
+
+	ChildSlot
+	[
+		MakeToolBar(InArgs._Extenders)
+	];
+
+	SViewportToolBar::Construct(SViewportToolBar::FArguments());
+}
+
+TSharedRef< SWidget > SDesignerToolBar::MakeToolBar(const TSharedPtr< FExtender > InExtenders)
+{
+	FToolBarBuilder ToolbarBuilder( CommandList, FMultiBoxCustomization::None, InExtenders );
+
+	// Use a custom style
+	FName ToolBarStyle = "ViewportMenu";
+	ToolbarBuilder.SetStyle(&FEditorStyle::Get(), ToolBarStyle);
+	ToolbarBuilder.SetLabelVisibility(EVisibility::Collapsed);
+
+	// Transform controls cannot be focusable as it fights with the press space to change transform mode feature
+	ToolbarBuilder.SetIsFocusable( false );
+
+	ToolbarBuilder.BeginSection("Transform");
+	ToolbarBuilder.BeginBlockGroup();
+	{
+		ToolbarBuilder.AddToolBarButton( FDesignerCommands::Get().LayoutTransform, NAME_None, TAttribute<FText>(), TAttribute<FText>(), TAttribute<FSlateIcon>(), "LayoutTransform" );
+		ToolbarBuilder.AddToolBarButton( FDesignerCommands::Get().RenderTransform, NAME_None, TAttribute<FText>(), TAttribute<FText>(), TAttribute<FSlateIcon>(), "RenderTransform" );
+	}
+	ToolbarBuilder.EndBlockGroup();
+	ToolbarBuilder.EndSection();
+	
+	return ToolbarBuilder.MakeWidget();
+}
+
+#undef LOCTEXT_NAMESPACE
diff --git a/Engine/Source/Editor/UMGEditor/Private/Designer/SDesignerToolBar.h b/Engine/Source/Editor/UMGEditor/Private/Designer/SDesignerToolBar.h
new file mode 100644
--- /dev/null
+++ b/Engine/Source/Editor/UMGEditor/Private/Designer/SDesignerToolBar.h
@@ -0,0 +1,31 @@
+// Copyright 1998-2014 Epic Games, Inc. All Rights Reserved.
+
+#pragma once
+
+#include "SViewportToolBar.h"
+
+/**
+ * 
+ */
+class SDesignerToolBar : public SViewportToolBar
+{
+public:
+	SLATE_BEGIN_ARGS( SDesignerToolBar ){}
+		SLATE_ARGUMENT( TSharedPtr<FUICommandList>, CommandList )
+		SLATE_ARGUMENT( TSharedPtr<FExtender>, Extenders )
+	SLATE_END_ARGS()
+
+	void Construct( const FArguments& InArgs );
+
+	/**
+	 * Static: Creates a widget for the main tool bar
+	 *
+	 * @return	New widget
+	 */
+	TSharedRef< SWidget > MakeToolBar( const TSharedPtr< FExtender > InExtenders );	
+
+private:
+
+	/** Command list */
+	TSharedPtr<FUICommandList> CommandList;
+};
\ No newline at end of file
diff --git a/Engine/Source/Editor/UMGEditor/Private/Designer/SDesignerView.cpp b/Engine/Source/Editor/UMGEditor/Private/Designer/SDesignerView.cpp
--- a/Engine/Source/Editor/UMGEditor/Private/Designer/SDesignerView.cpp
+++ b/Engine/Source/Editor/UMGEditor/Private/Designer/SDesignerView.cpp
@@ -10,6 +10,8 @@
 #include "WidgetTemplateDragDropOp.h"
 #include "SZoomPan.h"
 #include "SDisappearingBar.h"
+#include "SDesignerToolBar.h"
+#include "STransformHandle.h"
 
 #define LOCTEXT_NAMESPACE "UMG"
 
@@ -116,9 +118,9 @@ void SDesignerView::Construct(const FArguments& InArgs, TSharedPtr<FWidgetBluepr
 	DropPreviewWidget = NULL;
 	DropPreviewParent = NULL;
 	BlueprintEditor = InBlueprintEditor;
-	CurrentHandle = DH_NONE;
 
 	DesignerMessage = EDesignerMessage::None;
+	TransformMode = ETransformMode::Layout;
 
 	SetStartupResolution();
 
@@ -130,18 +132,6 @@ void SDesignerView::Construct(const FArguments& InArgs, TSharedPtr<FWidgetBluepr
 	bMouseDown = false;
 	bMovingExistingWidget = false;
 
-	DragDirections.Init((int32)DH_MAX);
-	DragDirections[DH_TOP_LEFT] = FVector2D(-1, -1);
-	DragDirections[DH_TOP_CENTER] = FVector2D(0, -1);
-	DragDirections[DH_TOP_RIGHT] = FVector2D(1, -1);
-
-	DragDirections[DH_MIDDLE_LEFT] = FVector2D(-1, 0);
-	DragDirections[DH_MIDDLE_RIGHT] = FVector2D(1, 0);
-
-	DragDirections[DH_BOTTOM_LEFT] = FVector2D(-1, 1);
-	DragDirections[DH_BOTTOM_CENTER] = FVector2D(0, 1);
-	DragDirections[DH_BOTTOM_RIGHT] = FVector2D(1, 1);
-
 	// TODO UMG - Register these with the module through some public interface to allow for new extensions to be registered.
 	Register(MakeShareable(new FVerticalSlotExtension()));
 	Register(MakeShareable(new FHorizontalSlotExtension()));
@@ -152,6 +142,8 @@ void SDesignerView::Construct(const FArguments& InArgs, TSharedPtr<FWidgetBluepr
 	UWidgetBlueprint* Blueprint = GetBlueprint();
 	Blueprint->OnChanged().AddSP(this, &SDesignerView::OnBlueprintChanged);
 
+	BindCommands();
+
 	SDesignSurface::Construct(SDesignSurface::FArguments()
 		.AllowContinousZoomInterpolation(false)
 		.Content()
@@ -216,6 +208,14 @@ void SDesignerView::Construct(const FArguments& InArgs, TSharedPtr<FWidgetBluepr
 					.Size(FVector2D(1, 1))
 				]
 
+				+ SHorizontalBox::Slot()
+				.AutoWidth()
+				.Padding(5.0f)
+				[
+					SNew(SDesignerToolBar)
+					.CommandList(CommandList)
+				]
+
 				+ SHorizontalBox::Slot()
 				.AutoWidth()
 				.Padding(5.0f)
@@ -309,6 +309,42 @@ SDesignerView::~SDesignerView()
 	}
 }
 
+void SDesignerView::BindCommands()
+{
+	CommandList = MakeShareable(new FUICommandList);
+
+	const FDesignerCommands& Commands = FDesignerCommands::Get();
+
+	CommandList->MapAction(
+		Commands.LayoutTransform,
+		FExecuteAction::CreateSP(this, &SDesignerView::SetTransformMode, ETransformMode::Layout),
+		FCanExecuteAction::CreateSP(this, &SDesignerView::CanSetTransformMode, ETransformMode::Layout),
+		FIsActionChecked::CreateSP(this, &SDesignerView::IsTransformModeActive, ETransformMode::Layout)
+		);
+
+	CommandList->MapAction(
+		Commands.RenderTransform,
+		FExecuteAction::CreateSP(this, &SDesignerView::SetTransformMode, ETransformMode::Render),
+		FCanExecuteAction::CreateSP(this, &SDesignerView::CanSetTransformMode, ETransformMode::Render),
+		FIsActionChecked::CreateSP(this, &SDesignerView::IsTransformModeActive, ETransformMode::Render)
+		);
+}
+
+void SDesignerView::SetTransformMode(ETransformMode::Type InTransformMode)
+{
+	TransformMode = InTransformMode;
+}
+
+bool SDesignerView::CanSetTransformMode(ETransformMode::Type InTransformMode) const
+{
+	return true;
+}
+
+bool SDesignerView::IsTransformModeActive(ETransformMode::Type InTransformMode) const
+{
+	return TransformMode == InTransformMode;
+}
+
 void SDesignerView::SetStartupResolution()
 {
 	// Use previously set resolution (or create new entries using default values)
@@ -337,6 +373,16 @@ float SDesignerView::GetPreviewScale() const
 	return GetZoomAmount() * GetPreviewDPIScale();
 }
 
+FWidgetReference SDesignerView::GetSelectedWidget() const
+{
+	return SelectedWidget;
+}
+
+ETransformMode::Type SDesignerView::GetTransformMode() const
+{
+	return TransformMode;
+}
+
 FOptionalSize SDesignerView::GetPreviewWidth() const
 {
 	return (float)PreviewWidth;
@@ -436,24 +482,39 @@ void SDesignerView::CreateExtensionWidgetsForSelection()
 
 	TArray< TSharedRef<FDesignerSurfaceElement> > ExtensionElements;
 
-	// Build extension widgets for new selection
-	for ( TSharedRef<FDesignerExtension>& Ext : DesignerExtensions )
+	if ( Selected.Num() > 0 )
 	{
-		if ( Ext->CanExtendSelection(Selected) )
+		// Add transform handles
+		ExtensionElements.Add(MakeShareable(new FDesignerSurfaceElement(SNew(STransformHandle, this, ETransformDirection::TopLeft), EExtensionLayoutLocation::TopLeft, FVector2D(-10, -10))));
+		ExtensionElements.Add(MakeShareable(new FDesignerSurfaceElement(SNew(STransformHandle, this, ETransformDirection::TopCenter), EExtensionLayoutLocation::TopCenter, FVector2D(-5, -10))));
+		ExtensionElements.Add(MakeShareable(new FDesignerSurfaceElement(SNew(STransformHandle, this, ETransformDirection::TopRight), EExtensionLayoutLocation::TopRight, FVector2D(0, -10))));
+
+		ExtensionElements.Add(MakeShareable(new FDesignerSurfaceElement(SNew(STransformHandle, this, ETransformDirection::CenterLeft), EExtensionLayoutLocation::CenterLeft, FVector2D(-10, -5))));
+		ExtensionElements.Add(MakeShareable(new FDesignerSurfaceElement(SNew(STransformHandle, this, ETransformDirection::CenterRight), EExtensionLayoutLocation::CenterRight, FVector2D(0, -5))));
+
+		ExtensionElements.Add(MakeShareable(new FDesignerSurfaceElement(SNew(STransformHandle, this, ETransformDirection::BottomLeft), EExtensionLayoutLocation::BottomLeft, FVector2D(-10, 0))));
+		ExtensionElements.Add(MakeShareable(new FDesignerSurfaceElement(SNew(STransformHandle, this, ETransformDirection::BottomCenter), EExtensionLayoutLocation::BottomCenter, FVector2D(-5, 0))));
+		ExtensionElements.Add(MakeShareable(new FDesignerSurfaceElement(SNew(STransformHandle, this, ETransformDirection::BottomRight), EExtensionLayoutLocation::BottomRight, FVector2D(0, 0))));
+
+		// Build extension widgets for new selection
+		for ( TSharedRef<FDesignerExtension>& Ext : DesignerExtensions )
 		{
-			Ext->ExtendSelection(Selected, ExtensionElements);
+			if ( Ext->CanExtendSelection(Selected) )
+			{
+				Ext->ExtendSelection(Selected, ExtensionElements);
+			}
 		}
-	}
 
-	// Add Widgets to designer surface
-	for ( TSharedRef<FDesignerSurfaceElement>& ExtElement : ExtensionElements )
-	{
-		ExtensionWidgetCanvas->AddSlot()
-			.Position(TAttribute<FVector2D>::Create(TAttribute<FVector2D>::FGetter::CreateSP(this, &SDesignerView::GetExtensionPosition, ExtElement)))
-			.Size(TAttribute<FVector2D>::Create(TAttribute<FVector2D>::FGetter::CreateSP(this, &SDesignerView::GetExtensionSize, ExtElement)))
-			[
-				ExtElement->GetWidget()
-			];
+		// Add Widgets to designer surface
+		for ( TSharedRef<FDesignerSurfaceElement>& ExtElement : ExtensionElements )
+		{
+			ExtensionWidgetCanvas->AddSlot()
+				.Position(TAttribute<FVector2D>::Create(TAttribute<FVector2D>::FGetter::CreateSP(this, &SDesignerView::GetExtensionPosition, ExtElement)))
+				.Size(TAttribute<FVector2D>::Create(TAttribute<FVector2D>::FGetter::CreateSP(this, &SDesignerView::GetExtensionSize, ExtElement)))
+				[
+					ExtElement->GetWidget()
+				];
+		}
 	}
 }
 
@@ -618,14 +679,6 @@ FReply SDesignerView::OnMouseButtonDown(const FGeometry& MyGeometry, const FPoin
 {
 	SDesignSurface::OnMouseButtonDown(MyGeometry, MouseEvent);
 
-	CurrentHandle = HitTestDragHandles(MyGeometry, MouseEvent);
-
-	if ( CurrentHandle != DH_NONE )
-	{
-		BeginTransaction(LOCTEXT("ResizeWidget", "Resize Widget"));
-		return FReply::Handled().PreventThrottling().CaptureMouse(AsShared());
-	}
-
 	//TODO UMG Undoable Selection
 	FArrangedWidget ArrangedWidget(SNullWidget::NullWidget, FGeometry());
 	FWidgetReference NewSelectedWidget = GetWidgetAtCursor(MyGeometry, MouseEvent, ArrangedWidget);
@@ -656,12 +709,6 @@ FReply SDesignerView::OnMouseButtonUp(const FGeometry& MyGeometry, const FPointe
 		bMouseDown = false;
 		bMovingExistingWidget = false;
 		DesignerMessage = EDesignerMessage::None;
-
-		if ( CurrentHandle != DH_NONE )
-		{
-			EndTransaction();
-			CurrentHandle = DH_NONE;
-		}
 	}
 	else if ( MouseEvent.GetEffectingButton() == EKeys::RightMouseButton )
 	{
@@ -689,33 +736,6 @@ FReply SDesignerView::OnMouseMove(const FGeometry& MyGeometry, const FPointerEve
 		return SurfaceHandled;
 	}
 
-	// Update the resizing based on the mouse movement
-	if ( CurrentHandle != DH_NONE )
-	{
-		if ( SelectedWidget.IsValid() )
-		{
-			//@TODO UMG - Implement some system to query slots to know if they support dragging so we can provide visual feedback by hiding handles that wouldn't work and such.
-			//SelectedTemplate->Slot->CanResize(DH_TOP_CENTER)
-
-			UWidget* Template = SelectedWidget.GetTemplate();
-			UWidget* Preview = SelectedWidget.GetPreview();
-
-			if ( Preview->Slot )
-			{
-				Preview->Slot->Resize(DragDirections[CurrentHandle], MouseEvent.GetCursorDelta() * ( 1.0f / GetPreviewScale() ));
-			}
-
-			if ( Template->Slot )
-			{
-				Template->Slot->Resize(DragDirections[CurrentHandle], MouseEvent.GetCursorDelta() * ( 1.0f / GetPreviewScale() ));
-			}
-
-			FBlueprintEditorUtils::MarkBlueprintAsModified(GetBlueprint());
-
-			return FReply::Handled().PreventThrottling();
-		}
-	}
-
 	if ( bMouseDown )
 	{
 		if ( SelectedWidget.IsValid() && !bMovingExistingWidget )
@@ -724,23 +744,19 @@ FReply SDesignerView::OnMouseMove(const FGeometry& MyGeometry, const FPointerEve
 			if ( !bIsRootWidget )
 			{
 				bMovingExistingWidget = true;
-				DesignerMessage = EDesignerMessage::None;
 				//Drag selected widgets
 				return FReply::Handled().DetectDrag(AsShared(), EKeys::LeftMouseButton);
 			}
 		}
 	}
 	
-	if ( CurrentHandle == DH_NONE )
+	// Update the hovered widget under the mouse
+	FArrangedWidget ArrangedWidget(SNullWidget::NullWidget, FGeometry());
+	FWidgetReference NewHoveredWidget = GetWidgetAtCursor(MyGeometry, MouseEvent, ArrangedWidget);
+	if ( !( NewHoveredWidget == HoveredWidget ) )
 	{
-		// Update the hovered widget under the mouse
-		FArrangedWidget ArrangedWidget(SNullWidget::NullWidget, FGeometry());
-		FWidgetReference NewHoveredWidget = GetWidgetAtCursor(MyGeometry, MouseEvent, ArrangedWidget);
-		if ( !( NewHoveredWidget == HoveredWidget ) )
-		{
-			HoveredWidget = NewHoveredWidget;
-			HoverTime = 0;
-		}
+		HoveredWidget = NewHoveredWidget;
+		HoverTime = 0;
 	}
 
 	return FReply::Unhandled();
@@ -762,7 +778,12 @@ FReply SDesignerView::OnKeyDown(const FGeometry& MyGeometry, const FKeyboardEven
 {
 	BlueprintEditor.Pin()->PasteDropLocation = FVector2D(0, 0);
 
-	if ( BlueprintEditor.Pin()->WidgetCommandList->ProcessCommandBindings(InKeyboardEvent) )
+	if ( BlueprintEditor.Pin()->DesignerCommandList->ProcessCommandBindings(InKeyboardEvent) )
+	{
+		return FReply::Handled();
+	}
+
+	if ( CommandList->ProcessCommandBindings(InKeyboardEvent) )
 	{
 		return FReply::Handled();
 	}
@@ -874,137 +895,11 @@ int32 SDesignerView::OnPaint(const FPaintArgs& Args, const FGeometry& AllottedGe
 			ESlateDrawEffect::None,
 			Tint
 		);
-
-		DrawDragHandles(ArrangedWidget.Geometry, MyClippingRect, OutDrawElements, LayerId, InWidgetStyle);
 	}
 
 	return LayerId;
 }
 
-void SDesignerView::DrawDragHandles(const FGeometry& SelectionGeometry, const FSlateRect& MyClippingRect, FSlateWindowElementList& OutDrawElements, int32 LayerId, const FWidgetStyle& InWidgetStyle) const
-{
-	if ( SelectedWidget.IsValid() && SelectedWidget.GetTemplate()->Slot )
-	{
-		float X = 0.0f;
-		float Y = 0.0f;
-		float Width = SelectionGeometry.Size.X;
-		float Height = SelectionGeometry.Size.Y;
-
-		// @TODO UMG Handles should come from the slot/container to express how its slots can be transformed.
-		TArray<FVector2D> Handles;
-		Handles.Add(FVector2D(X, Y));					// Top - Left
-		Handles.Add(FVector2D(X + Width * 0.5f, Y));	// Top - Middle
-		Handles.Add(FVector2D(X + Width, Y));			// Top - Right
-
-		Handles.Add(FVector2D(X, Y + Height * 0.5f));			// Middle - Left
-		Handles.Add(FVector2D(X + Width, Y + Height * 0.5f));	// Middle - Right
-
-		Handles.Add(FVector2D(X, Y + Height));					// Bottom - Left
-		Handles.Add(FVector2D(X + Width * 0.5f, Y + Height));	// Bottom - Middle
-		Handles.Add(FVector2D(X + Width, Y + Height));			// Bottom - Right
-
-		const FVector2D HandleSize = FVector2D(10, 10);
-
-		// @TODO UMG - Don't use the curve editors brushes
-		const FSlateBrush* KeyBrush = FEditorStyle::GetBrush("CurveEd.CurveKey");
-		FLinearColor KeyColor = InWidgetStyle.GetColorAndOpacityTint();// IsEditingEnabled() ? InWidgetStyle.GetColorAndOpacityTint() : FLinearColor(0.1f, 0.1f, 0.1f, 1.f);
-
-		for ( int32 HandleIndex = 0; HandleIndex < Handles.Num(); HandleIndex++ )
-		{
-			const FVector2D& Handle = Handles[HandleIndex];
-			if ( !SelectedWidget.GetTemplate()->Slot->CanResize(DragDirections[HandleIndex]) )
-			{
-				// This isn't a valid handle
-				continue;
-			}
-
-			FSlateDrawElement::MakeBox(
-				OutDrawElements,
-				++LayerId,
-				SelectionGeometry.ToPaintGeometry(HandleSize, FSlateLayoutTransform(FVector2D(Handle.X - HandleSize.X * 0.5f, Handle.Y - HandleSize.Y * 0.5f))),
-				KeyBrush,
-				MyClippingRect,
-				ESlateDrawEffect::None,
-				KeyBrush->GetTint(InWidgetStyle) * InWidgetStyle.GetColorAndOpacityTint() * KeyColor
-				);
-		}
-	}
-}
-
-SDesignerView::DragHandle SDesignerView::HitTestDragHandles(const FGeometry& AllottedGeometry, const FPointerEvent& PointerEvent) const
-{
-	FVector2D LocalPointer = AllottedGeometry.AbsoluteToLocal(PointerEvent.GetScreenSpacePosition());
-
-	if ( SelectedWidget.IsValid() && SelectedWidget.GetTemplate()->Slot && SelectedSlateWidget.IsValid() )
-	{
-		TSharedRef<SWidget> Widget = SelectedSlateWidget.Pin().ToSharedRef();
-
-		FArrangedWidget ArrangedWidget(SNullWidget::NullWidget, FGeometry());
-		FDesignTimeUtils::GetArrangedWidget(Widget, ArrangedWidget);
-
-		FVector2D WidgetLocalPosition = AllottedGeometry.AbsoluteToLocal(ArrangedWidget.Geometry.AbsolutePosition);
-		float X = WidgetLocalPosition.X;
-		float Y = WidgetLocalPosition.Y;
-		float Width = ArrangedWidget.Geometry.Size.X * GetZoomAmount() * GetPreviewDPIScale();
-		float Height = ArrangedWidget.Geometry.Size.Y * GetZoomAmount() * GetPreviewDPIScale();
-
-		TArray<FVector2D> Handles;
-		Handles.Add(FVector2D(X, Y));					// Top - Left
-		Handles.Add(FVector2D(X + Width * 0.5f, Y));	// Top - Middle
-		Handles.Add(FVector2D(X + Width, Y));			// Top - Right
-
-		Handles.Add(FVector2D(X, Y + Height * 0.5f));			// Middle - Left
-		Handles.Add(FVector2D(X + Width, Y + Height * 0.5f));	// Middle - Right
-
-		Handles.Add(FVector2D(X, Y + Height));					// Bottom - Left
-		Handles.Add(FVector2D(X + Width * 0.5f, Y + Height));	// Bottom - Middle
-		Handles.Add(FVector2D(X + Width, Y + Height));			// Bottom - Right
-
-		const FVector2D HandleSize = FVector2D(10, 10);
-
-		int32 i = 0;
-		for ( FVector2D Handle : Handles )
-		{
-			FSlateRect Rect(FVector2D(Handle.X - HandleSize.X * 0.5f, Handle.Y - HandleSize.Y * 0.5f), Handle + HandleSize);
-			if ( Rect.ContainsPoint(LocalPointer) )
-			{
-				if ( !SelectedWidget.GetTemplate()->Slot->CanResize(DragDirections[i]) )
-				{
-					// This isn't a valid handle
-					break;
-				}
-
-				return (DragHandle)i;
-			}
-			i++;
-		}
-	}
-
-	return DH_NONE;
-}
-
-FCursorReply SDesignerView::OnCursorQuery(const FGeometry& MyGeometry, const FPointerEvent& CursorEvent) const
-{
-	// Hit test the drag handles
-	switch ( HitTestDragHandles(MyGeometry, CursorEvent) )
-	{
-	case DH_TOP_LEFT:
-	case DH_BOTTOM_RIGHT:
-		return FCursorReply::Cursor(EMouseCursor::ResizeSouthEast);
-	case DH_TOP_RIGHT:
-	case DH_BOTTOM_LEFT:
-		return FCursorReply::Cursor(EMouseCursor::ResizeSouthWest);
-	case DH_TOP_CENTER:
-	case DH_BOTTOM_CENTER:
-		return FCursorReply::Cursor(EMouseCursor::ResizeUpDown);
-	case DH_MIDDLE_LEFT:
-	case DH_MIDDLE_RIGHT:
-		return FCursorReply::Cursor(EMouseCursor::ResizeLeftRight);
-	}
-
-	return FCursorReply::Unhandled();
-}
-
 void SDesignerView::UpdatePreviewWidget()
 {
 	UUserWidget* LatestPreviewWidget = BlueprintEditor.Pin()->GetPreview();
diff --git a/Engine/Source/Editor/UMGEditor/Private/Designer/SDesignerView.h b/Engine/Source/Editor/UMGEditor/Private/Designer/SDesignerView.h
--- a/Engine/Source/Editor/UMGEditor/Private/Designer/SDesignerView.h
+++ b/Engine/Source/Editor/UMGEditor/Private/Designer/SDesignerView.h
@@ -41,8 +41,6 @@ public:
 
 	virtual FReply OnKeyDown(const FGeometry& MyGeometry, const FKeyboardEvent& InKeyboardEvent) override;
 
-	virtual FCursorReply OnCursorQuery(const FGeometry& MyGeometry, const FPointerEvent& CursorEvent) const override;
-
 	//virtual void OnArrangeChildren(const FGeometry& AllottedGeometry, FArrangedChildren& ArrangedChildren) const;
 
 	virtual int32 OnPaint(const FPaintArgs& Args, const FGeometry& AllottedGeometry, const FSlateRect& MyClippingRect, FSlateWindowElementList& OutDrawElements, int32 LayerId, const FWidgetStyle& InWidgetStyle, bool bParentEnabled) const override;
@@ -58,7 +56,9 @@ public:
 	void Register(TSharedRef<FDesignerExtension> Extension);
 
 	// IUMGDesigner interface
-	float GetPreviewScale() const override;
+	virtual float GetPreviewScale() const override;
+	virtual FWidgetReference GetSelectedWidget() const override;
+	virtual ETransformMode::Type GetTransformMode() const override;
 	// End of IUMGDesigner interface
 
 private:
@@ -126,36 +126,17 @@ private:
 	static const uint32 DefaultResolutionHeight;
 	static const FString DefaultAspectRatio;
 
-	enum DragHandle
-	{
-		DH_NONE = -1,
-
-		DH_TOP_LEFT = 0,
-		DH_TOP_CENTER,
-		DH_TOP_RIGHT,
-
-		DH_MIDDLE_LEFT,
-		DH_MIDDLE_RIGHT,
-
-		DH_BOTTOM_LEFT,
-		DH_BOTTOM_CENTER,
-		DH_BOTTOM_RIGHT,
-		
-		DH_MAX,
-	};
-
-	TArray< FVector2D > DragDirections;
-
 	/** Extensions for the designer to allow for custom widgets to be inserted onto the design surface as selection changes. */
 	TArray< TSharedRef<FDesignerExtension> > DesignerExtensions;
 
 private:
-	void DrawDragHandles(const FGeometry& SelectionGeometry, const FSlateRect& MyClippingRect, FSlateWindowElementList& OutDrawElements, int32 LayerId, const FWidgetStyle& InWidgetStyle) const;
-	DragHandle HitTestDragHandles(const FGeometry& AllottedGeometry, const FPointerEvent& PointerEvent) const;
+	void BindCommands();
 
-	UWidget* ProcessDropAndAddWidget(const FGeometry& MyGeometry, const FDragDropEvent& DragDropEvent, bool bIsPreview);
+	void SetTransformMode(ETransformMode::Type InTransformMode);
+	bool CanSetTransformMode(ETransformMode::Type InTransformMode) const;
+	bool IsTransformModeActive(ETransformMode::Type InTransformMode) const;
 
-private:
+	UWidget* ProcessDropAndAddWidget(const FGeometry& MyGeometry, const FDragDropEvent& DragDropEvent, bool bIsPreview);
 
 	FVector2D GetExtensionPosition(TSharedRef<FDesignerSurfaceElement> ExtensionElement) const;
 
@@ -165,6 +146,9 @@ private:
 	/** A reference to the BP Editor that owns this designer */
 	TWeakPtr<FWidgetBlueprintEditor> BlueprintEditor;
 
+	/** The designer command list */
+	TSharedPtr<FUICommandList> CommandList;
+
 	/** The transaction used to commit undoable actions from resize, move...etc */
 	FScopedTransaction* ScopedTransaction;
 
@@ -181,8 +165,6 @@ private:
 	TSharedPtr<SDPIScaler> PreviewSurface;
 	TSharedPtr<SCanvas> ExtensionWidgetCanvas;
 
-	DragHandle CurrentHandle;
-
 	FVector2D CachedDesignerWidgetLocation;
 	FVector2D CachedDesignerWidgetSize;
 
@@ -231,4 +213,7 @@ private:
 
 	/** */
 	EDesignerMessage::Type DesignerMessage;
+
+	/** */
+	ETransformMode::Type TransformMode;
 };
diff --git a/Engine/Source/Editor/UMGEditor/Private/Designer/SDisappearingBar.cpp b/Engine/Source/Editor/UMGEditor/Private/Designer/SDisappearingBar.cpp
--- a/Engine/Source/Editor/UMGEditor/Private/Designer/SDisappearingBar.cpp
+++ b/Engine/Source/Editor/UMGEditor/Private/Designer/SDisappearingBar.cpp
@@ -12,7 +12,7 @@
 void SDisappearingBar::Construct(const FArguments& InArgs)
 {
 	FadeCurve = FCurveSequence(0, 0.25f);
-	FadeCurve.Play();
+	//FadeCurve.Play();
 
 	bDisappear = false;
 
diff --git a/Engine/Source/Editor/UMGEditor/Private/Designer/STransformHandle.cpp b/Engine/Source/Editor/UMGEditor/Private/Designer/STransformHandle.cpp
new file mode 100644
--- /dev/null
+++ b/Engine/Source/Editor/UMGEditor/Private/Designer/STransformHandle.cpp
@@ -0,0 +1,236 @@
+// Copyright 1998-2014 Epic Games, Inc. All Rights Reserved.
+
+#include "UMGEditorPrivatePCH.h"
+
+#include "ObjectEditorUtils.h"
+#include "STransformHandle.h"
+
+#define LOCTEXT_NAMESPACE "UMG"
+
+/////////////////////////////////////////////////////
+// STransformHandle
+
+void STransformHandle::Construct(const FArguments& InArgs, IUMGDesigner* InDesigner, ETransformDirection::Type InTransformDirection)
+{
+	TransformDirection = InTransformDirection;
+	Designer = InDesigner;
+
+	Action = ETransformAction::None;
+
+	DragDirection = ComputeDragDirection(InTransformDirection);
+	DragOrigin = ComputeOrigin(InTransformDirection);
+
+	ChildSlot
+	[
+		SNew(SBox)
+		.WidthOverride(10)
+		.HeightOverride(10)
+		.Visibility(this, &STransformHandle::GetHandleVisibility)
+		[
+			SNew(SImage)
+			.Image(FEditorStyle::Get().GetBrush("CurveEd.CurveKey"))
+		]
+	];
+}
+
+EVisibility STransformHandle::GetHandleVisibility() const
+{
+	ETransformMode::Type TransformMode = Designer->GetTransformMode();
+
+	switch ( TransformMode )
+	{
+	case ETransformMode::Layout:
+	{
+		FWidgetReference SelectedWidget = Designer->GetSelectedWidget();
+		if ( UPanelSlot* TemplateSlot = SelectedWidget.GetTemplate()->Slot )
+		{
+			if ( TemplateSlot->CanResize(DragDirection) )
+			{
+				return EVisibility::Visible;
+			}
+		}
+		break;
+	}
+	case ETransformMode::Render:
+		return EVisibility::Visible;
+	}
+
+	return EVisibility::Collapsed;
+}
+
+FReply STransformHandle::OnMouseButtonDown(const FGeometry& MyGeometry, const FPointerEvent& MouseEvent)
+{
+	if ( MouseEvent.GetEffectingButton() == EKeys::LeftMouseButton )
+	{
+		Action = ComputeActionAtLocation(MyGeometry, MouseEvent);
+
+		return FReply::Handled().CaptureMouse(SharedThis(this));
+	}
+
+	return FReply::Unhandled();
+}
+
+FReply STransformHandle::OnMouseButtonUp(const FGeometry& MyGeometry, const FPointerEvent& MouseEvent)
+{
+	if ( HasMouseCapture() && MouseEvent.GetEffectingButton() == EKeys::LeftMouseButton )
+	{
+		Action = ETransformAction::None;
+		return FReply::Handled().ReleaseMouseCapture();
+	}
+
+	return FReply::Unhandled();
+}
+
+FReply STransformHandle::OnMouseMove(const FGeometry& MyGeometry, const FPointerEvent& MouseEvent)
+{
+	if ( Action != ETransformAction::None )
+	{
+		FWidgetReference SelectedWidget = Designer->GetSelectedWidget();
+
+		UWidget* Template = SelectedWidget.GetTemplate();
+		UWidget* Preview = SelectedWidget.GetPreview();
+
+		ETransformMode::Type TransformMode = Designer->GetTransformMode();
+		if ( TransformMode == ETransformMode::Layout )
+		{
+			FVector2D TranslateAmount = MouseEvent.GetCursorDelta() * ( 1.0f / Designer->GetPreviewScale() );
+
+			if ( Preview->Slot )
+			{
+				Preview->Slot->Resize(DragDirection, TranslateAmount);
+			}
+
+			if ( Template->Slot )
+			{
+				Template->Slot->Resize(DragDirection, TranslateAmount);
+			}
+		}
+		else if ( TransformMode == ETransformMode::Render )
+		{
+			FWidgetTransform RenderTransform = Preview->RenderTransform;
+
+			if ( Action == ETransformAction::Primary )
+			{
+				FVector2D ScaleDelta = MouseEvent.GetCursorDelta() * 0.05f;
+				RenderTransform.Scale += ScaleDelta;
+			}
+			else
+			{
+				FVector2D RotationDelta = MouseEvent.GetCursorDelta();
+				RenderTransform.Angle += RotationDelta.Size();
+			}
+
+			FObjectEditorUtils::SetPropertyValue<UWidget, FWidgetTransform>(Preview, "RenderTransform", RenderTransform);
+		}
+	}
+
+	return FReply::Unhandled();
+}
+
+FCursorReply STransformHandle::OnCursorQuery(const FGeometry& MyGeometry, const FPointerEvent& MouseEvent) const
+{
+	ETransformAction::Type CurrentAction = Action;
+	if ( CurrentAction == ETransformAction::None )
+	{
+		CurrentAction = ComputeActionAtLocation(MyGeometry, MouseEvent);
+	}
+
+	ETransformMode::Type TransformMode = Designer->GetTransformMode();
+
+	if ( TransformMode == ETransformMode::Layout || CurrentAction == ETransformAction::Primary )
+	{
+		switch ( TransformDirection )
+		{
+		case ETransformDirection::TopLeft:
+		case ETransformDirection::BottomRight:
+			return FCursorReply::Cursor(EMouseCursor::ResizeSouthEast);
+		case ETransformDirection::TopRight:
+		case ETransformDirection::BottomLeft:
+			return FCursorReply::Cursor(EMouseCursor::ResizeSouthWest);
+		case ETransformDirection::TopCenter:
+		case ETransformDirection::BottomCenter:
+			return FCursorReply::Cursor(EMouseCursor::ResizeUpDown);
+		case ETransformDirection::CenterLeft:
+		case ETransformDirection::CenterRight:
+			return FCursorReply::Cursor(EMouseCursor::ResizeLeftRight);
+		}
+	}
+	else
+	{
+		return FCursorReply::Cursor(EMouseCursor::EyeDropper);
+	}
+
+	return FCursorReply::Unhandled();
+}
+
+FVector2D STransformHandle::ComputeDragDirection(ETransformDirection::Type InTransformDirection) const
+{
+	switch ( InTransformDirection )
+	{
+	case ETransformDirection::TopLeft:
+		return FVector2D(-1, -1);
+	case ETransformDirection::TopCenter:
+		return FVector2D(0, -1);
+	case ETransformDirection::TopRight:
+		return FVector2D(1, -1);
+
+	case ETransformDirection::CenterLeft:
+		return FVector2D(-1, 0);
+	case ETransformDirection::CenterRight:
+		return FVector2D(1, 0);
+
+	case ETransformDirection::BottomLeft:
+		return FVector2D(-1, 1);
+	case ETransformDirection::BottomCenter:
+		return FVector2D(0, 1);
+	case ETransformDirection::BottomRight:
+		return FVector2D(1, 1);
+	}
+
+	return FVector2D(0, 0);
+}
+
+FVector2D STransformHandle::ComputeOrigin(ETransformDirection::Type InTransformDirection) const
+{
+	FVector2D Size(10, 10);
+
+	switch ( InTransformDirection )
+	{
+	case ETransformDirection::TopLeft:
+		return Size * FVector2D(1, 1);
+	case ETransformDirection::TopCenter:
+		return Size * FVector2D(0.5, 1);
+	case ETransformDirection::TopRight:
+		return Size * FVector2D(0, 1);
+
+	case ETransformDirection::CenterLeft:
+		return Size * FVector2D(1, 0.5);
+	case ETransformDirection::CenterRight:
+		return Size * FVector2D(0, 0.5);
+
+	case ETransformDirection::BottomLeft:
+		return Size * FVector2D(1, 0);
+	case ETransformDirection::BottomCenter:
+		return Size * FVector2D(0.5, 0);
+	case ETransformDirection::BottomRight:
+		return Size * FVector2D(0, 0);
+	}
+
+	return FVector2D(0, 0);
+}
+
+ETransformAction::Type STransformHandle::ComputeActionAtLocation(const FGeometry& MyGeometry, const FPointerEvent& MouseEvent) const
+{
+	FVector2D LocalPosition = MyGeometry.AbsoluteToLocal(MouseEvent.GetScreenSpacePosition());
+	FVector2D GrabOriginOffset = LocalPosition - DragOrigin;
+	if ( GrabOriginOffset.Size() < 6 )
+	{
+		return ETransformAction::Primary;
+	}
+	else
+	{
+		return ETransformAction::Secondary;
+	}
+}
+
+#undef LOCTEXT_NAMESPACE
diff --git a/Engine/Source/Editor/UMGEditor/Private/Designer/STransformHandle.h b/Engine/Source/Editor/UMGEditor/Private/Designer/STransformHandle.h
new file mode 100644
--- /dev/null
+++ b/Engine/Source/Editor/UMGEditor/Private/Designer/STransformHandle.h
@@ -0,0 +1,69 @@
+// Copyright 1998-2014 Epic Games, Inc. All Rights Reserved.
+
+#pragma once
+
+#include "SCompoundWidget.h"
+
+namespace ETransformDirection
+{
+	enum Type
+	{
+		TopLeft = 0,
+		TopCenter,
+		TopRight,
+
+		CenterLeft,
+		CenterRight,
+
+		BottomLeft,
+		BottomCenter,
+		BottomRight,
+
+		MAX
+	};
+}
+
+
+namespace ETransformAction
+{
+	enum Type
+	{
+		None,
+		Primary,
+		Secondary
+	};
+}
+
+class STransformHandle : public SCompoundWidget
+{
+public:
+
+	SLATE_BEGIN_ARGS(STransformHandle) {}
+	SLATE_END_ARGS()
+
+	void Construct(const FArguments& InArgs, IUMGDesigner* InDesigner, ETransformDirection::Type InTransformDirection);
+
+	// SWidget interface
+	virtual FReply OnMouseButtonDown(const FGeometry& MyGeometry, const FPointerEvent& MouseEvent) override;
+	virtual FReply OnMouseButtonUp(const FGeometry& MyGeometry, const FPointerEvent& MouseEvent) override;
+	virtual FReply OnMouseMove(const FGeometry& MyGeometry, const FPointerEvent& MouseEvent) override;
+	virtual FCursorReply OnCursorQuery(const FGeometry& MyGeometry, const FPointerEvent& CursorEvent) const override;
+	// End SWidget interface
+	
+protected:
+	EVisibility GetHandleVisibility() const;
+
+	ETransformAction::Type ComputeActionAtLocation(const FGeometry& MyGeometry, const FPointerEvent& MouseEvent) const;
+
+protected:
+	FVector2D ComputeDragDirection(ETransformDirection::Type InTransformDirection) const;
+	FVector2D ComputeOrigin(ETransformDirection::Type InTransformDirection) const;
+
+protected:
+	IUMGDesigner* Designer;
+	ETransformDirection::Type TransformDirection;
+	ETransformAction::Type Action;
+
+	FVector2D DragDirection;
+	FVector2D DragOrigin;
+};
diff --git a/Engine/Source/Editor/UMGEditor/Private/Details/SWidgetDetailsView.cpp b/Engine/Source/Editor/UMGEditor/Private/Details/SWidgetDetailsView.cpp
--- a/Engine/Source/Editor/UMGEditor/Private/Details/SWidgetDetailsView.cpp
+++ b/Engine/Source/Editor/UMGEditor/Private/Details/SWidgetDetailsView.cpp
@@ -104,8 +104,8 @@ SWidgetDetailsView::~SWidgetDetailsView()
 
 void SWidgetDetailsView::RegisterCustomizations()
 {
-	FOnGetDetailCustomizationInstance LayoutDelegateDetails = FOnGetDetailCustomizationInstance::CreateStatic(&FBlueprintWidgetCustomization::MakeInstance, BlueprintEditor.Pin().ToSharedRef(), BlueprintEditor.Pin()->GetBlueprintObj());
-	PropertyView->RegisterInstancedCustomPropertyLayout(UWidget::StaticClass(), LayoutDelegateDetails);
+	PropertyView->RegisterInstancedCustomPropertyLayout(UWidget::StaticClass(), FOnGetDetailCustomizationInstance::CreateStatic(&FBlueprintWidgetCustomization::MakeInstance, BlueprintEditor.Pin().ToSharedRef(), BlueprintEditor.Pin()->GetBlueprintObj()));
+	//PropertyView->RegisterInstancedCustomPropertyLayout(FSlateBrush::StaticStruct()->StaticClass(), FOnGetDetailCustomizationInstance::CreateStatic(&FSlateBrushStructCustomization::MakeInstance));
 
 	static FName PropertyEditor("PropertyEditor");
 	FPropertyEditorModule& PropertyModule = FModuleManager::GetModuleChecked<FPropertyEditorModule>(PropertyEditor);
@@ -338,6 +338,8 @@ void SWidgetDetailsView::NotifyPreChange(FEditPropertyChain* PropertyAboutToChan
 
 void SWidgetDetailsView::NotifyPostChange(const FPropertyChangedEvent& PropertyChangedEvent, FEditPropertyChain* PropertyThatChanged)
 {
+	const static FName DesignerRebuildName("DesignerRebuild");
+
 	if ( PropertyChangedEvent.ChangeType != EPropertyChangeType::Interactive && !BlueprintEditor.Pin()->GetSequencer()->IsAutoKeyEnabled() )
 	{
 		TSharedPtr<FWidgetBlueprintEditor> Editor = BlueprintEditor.Pin();
@@ -345,6 +347,13 @@ void SWidgetDetailsView::NotifyPostChange(const FPropertyChangedEvent& PropertyC
 		const bool bIsModify = false;
 		Editor->MigrateFromChain(PropertyThatChanged, bIsModify);
 	}
+
+	// If the property that changed is marked as "DesignerRebuild" we invalidate
+	// the preview.
+	if ( PropertyChangedEvent.Property->GetBoolMetaData(DesignerRebuildName) )
+	{
+		BlueprintEditor.Pin()->InvalidatePreview();
+	}
 }
 
 #undef LOCTEXT_NAMESPACE
diff --git a/Engine/Source/Editor/UMGEditor/Private/Hierarchy/SHierarchyView.cpp b/Engine/Source/Editor/UMGEditor/Private/Hierarchy/SHierarchyView.cpp
--- a/Engine/Source/Editor/UMGEditor/Private/Hierarchy/SHierarchyView.cpp
+++ b/Engine/Source/Editor/UMGEditor/Private/Hierarchy/SHierarchyView.cpp
@@ -105,7 +105,7 @@ FReply SHierarchyView::OnKeyDown(const FGeometry& MyGeometry, const FKeyboardEve
 {
 	BlueprintEditor.Pin()->PasteDropLocation = FVector2D(0, 0);
 
-	if ( BlueprintEditor.Pin()->WidgetCommandList->ProcessCommandBindings(InKeyboardEvent) )
+	if ( BlueprintEditor.Pin()->DesignerCommandList->ProcessCommandBindings(InKeyboardEvent) )
 	{
 		return FReply::Handled();
 	}
diff --git a/Engine/Source/Editor/UMGEditor/Private/UMGEditorModule.cpp b/Engine/Source/Editor/UMGEditor/Private/UMGEditorModule.cpp
--- a/Engine/Source/Editor/UMGEditor/Private/UMGEditorModule.cpp
+++ b/Engine/Source/Editor/UMGEditor/Private/UMGEditorModule.cpp
@@ -13,6 +13,8 @@
 #include "Editor/Sequencer/Public/ISequencerModule.h"
 #include "Animation/MarginTrackEditor.h"
 
+#include "DesignerCommands.h"
+
 const FName UMGEditorAppIdentifier = FName(TEXT("UMGEditorApp"));
 
 class FUMGEditorModule : public IUMGEditorModule
@@ -28,6 +30,8 @@ public:
 	{
 		FModuleManager::LoadModuleChecked<IUMGModule>("UMG");
 
+		FDesignerCommands::Register();
+
 		MenuExtensibilityManager = MakeShareable(new FExtensibilityManager());
 		ToolBarExtensibilityManager = MakeShareable(new FExtensibilityManager());
 
diff --git a/Engine/Source/Editor/UMGEditor/Private/WidgetBlueprintEditor.cpp b/Engine/Source/Editor/UMGEditor/Private/WidgetBlueprintEditor.cpp
--- a/Engine/Source/Editor/UMGEditor/Private/WidgetBlueprintEditor.cpp
+++ b/Engine/Source/Editor/UMGEditor/Private/WidgetBlueprintEditor.cpp
@@ -7,6 +7,7 @@
 #include "MovieScene.h"
 #include "Editor/Sequencer/Public/ISequencerModule.h"
 #include "Animation/UMGSequencerObjectBindingManager.h"
+#include "ObjectEditorUtils.h"
 
 #include "PropertyCustomizationHelpers.h"
 
@@ -23,6 +24,7 @@ FWidgetBlueprintEditor::FWidgetBlueprintEditor()
 	: PreviewScene(FPreviewScene::ConstructionValues().AllowAudioPlayback(true).ShouldSimulatePhysics(true))
 	, PreviewBlueprint(NULL)
 {
+	
 }
 
 FWidgetBlueprintEditor::~FWidgetBlueprintEditor()
@@ -64,24 +66,24 @@ void FWidgetBlueprintEditor::InitWidgetBlueprintEditor(const EToolkitMode::Type
 
 	SequencerObjectBindingManager->InitPreviewObjects();
 
-	WidgetCommandList = MakeShareable(new FUICommandList);
+	DesignerCommandList = MakeShareable(new FUICommandList);
 
-	WidgetCommandList->MapAction(FGenericCommands::Get().Delete,
+	DesignerCommandList->MapAction(FGenericCommands::Get().Delete,
 		FExecuteAction::CreateSP(this, &FWidgetBlueprintEditor::DeleteSelectedWidgets),
 		FCanExecuteAction::CreateSP(this, &FWidgetBlueprintEditor::CanDeleteSelectedWidgets)
 		);
 
-	WidgetCommandList->MapAction(FGenericCommands::Get().Copy,
+	DesignerCommandList->MapAction(FGenericCommands::Get().Copy,
 		FExecuteAction::CreateSP(this, &FWidgetBlueprintEditor::CopySelectedWidgets),
 		FCanExecuteAction::CreateSP(this, &FWidgetBlueprintEditor::CanCopySelectedWidgets)
 		);
 
-	WidgetCommandList->MapAction(FGenericCommands::Get().Cut,
+	DesignerCommandList->MapAction(FGenericCommands::Get().Cut,
 		FExecuteAction::CreateSP(this, &FWidgetBlueprintEditor::CutSelectedWidgets),
 		FCanExecuteAction::CreateSP(this, &FWidgetBlueprintEditor::CanCutSelectedWidgets)
 		);
 
-	WidgetCommandList->MapAction(FGenericCommands::Get().Paste,
+	DesignerCommandList->MapAction(FGenericCommands::Get().Paste,
 		FExecuteAction::CreateSP(this, &FWidgetBlueprintEditor::PasteWidgets),
 		FCanExecuteAction::CreateSP(this, &FWidgetBlueprintEditor::CanPasteWidgets)
 		);
@@ -191,21 +193,18 @@ const TSet< TWeakObjectPtr<UObject> >& FWidgetBlueprintEditor::GetSelectedObject
 	return SelectedObjects;
 }
 
+void FWidgetBlueprintEditor::InvalidatePreview()
+{
+	bPreviewInvalidated = true;
+}
+
 void FWidgetBlueprintEditor::OnBlueprintChanged(UBlueprint* InBlueprint)
 {
 	FBlueprintEditor::OnBlueprintChanged(InBlueprint);
 
 	if ( InBlueprint )
 	{
-		// Rebuilding the preview can force objects to be recreated, so the selection may need to be updated.
-		OnSelectedWidgetsChanging.Broadcast();
-
-		UpdatePreview(InBlueprint, true);
-
-		CleanSelection();
-
-		// Fire the selection updated event to ensure everyone is watching the same widgets.
-		OnSelectedWidgetsChanged.Broadcast();
+		RefreshPreview();
 	}
 }
 
@@ -351,45 +350,11 @@ void FWidgetBlueprintEditor::Tick(float DeltaTime)
 
 	// Note: The weak ptr can become stale if the actor is reinstanced due to a Blueprint change, etc. In that case we 
 	//       look to see if we can find the new instance in the preview world and then update the weak ptr.
-	if ( PreviewWidgetPtr.IsStale(true) )
-	{
-		UpdatePreview(GetWidgetBlueprintObj(), true);
-	}
-}
-
-static bool MigratePropertyValue(UObject* SourceObject, UObject* DestinationObject, UProperty* MemberProperty)
-{
-	FString SourceValue;
-
-	// Get the property addresses for the source and destination objects.
-	uint8* SourceAddr = MemberProperty->ContainerPtrToValuePtr<uint8>(SourceObject);
-	uint8* DestionationAddr = MemberProperty->ContainerPtrToValuePtr<uint8>(DestinationObject);
-
-	if ( SourceAddr == NULL || DestionationAddr == NULL )
-	{
-		return false;
-	}
-
-	// Get the current value from the source object.
-	MemberProperty->ExportText_Direct(SourceValue, SourceAddr, SourceAddr, NULL, PPF_Localized);
-
-	const bool bNotifyObjectOfChange = true;
-
-	if ( !DestinationObject->HasAnyFlags(RF_ClassDefaultObject) && bNotifyObjectOfChange )
-	{
-		DestinationObject->PreEditChange(MemberProperty);
-	}
-
-	// Set the value on the destination object.
-	MemberProperty->ImportText(*SourceValue, DestionationAddr, 0, DestinationObject);
-
-	if ( !DestinationObject->HasAnyFlags(RF_ClassDefaultObject) && bNotifyObjectOfChange )
+	if ( PreviewWidgetPtr.IsStale(true) || bPreviewInvalidated )
 	{
-		FPropertyChangedEvent PropertyEvent(MemberProperty);
-		DestinationObject->PostEditChangeProperty(PropertyEvent);
+		bPreviewInvalidated = false;
+		RefreshPreview();
 	}
-
-	return true;
 }
 
 static bool MigratePropertyValue(UObject* SourceObject, UObject* DestinationObject, FEditPropertyChain::TDoubleLinkedListNode* PropertyChainNode, UProperty* MemberProperty, bool bIsModify)
@@ -410,10 +375,10 @@ static bool MigratePropertyValue(UObject* SourceObject, UObject* DestinationObje
 			UBoolProperty* EditConditionProperty = PropertyCustomizationHelpers::GetEditConditionProperty(MemberProperty, bDummyNegate);
 			if ( EditConditionProperty != NULL )
 			{
-				MigratePropertyValue(SourceObject, DestinationObject, EditConditionProperty);
+				FObjectEditorUtils::MigratePropertyValue(SourceObject, EditConditionProperty, DestinationObject, EditConditionProperty);
 			}
 
-			return MigratePropertyValue(SourceObject, DestinationObject, MemberProperty);
+			return FObjectEditorUtils::MigratePropertyValue(SourceObject, MemberProperty, DestinationObject, MemberProperty);
 		}
 	}
 	
@@ -614,6 +579,19 @@ void FWidgetBlueprintEditor::ChangeViewedAnimation( UWidgetAnimation& InAnimatio
 	}
 }
 
+void FWidgetBlueprintEditor::RefreshPreview()
+{
+	// Rebuilding the preview can force objects to be recreated, so the selection may need to be updated.
+	OnSelectedWidgetsChanging.Broadcast();
+
+	UpdatePreview(GetWidgetBlueprintObj(), true);
+
+	CleanSelection();
+
+	// Fire the selection updated event to ensure everyone is watching the same widgets.
+	OnSelectedWidgetsChanged.Broadcast();
+}
+
 void FWidgetBlueprintEditor::DestroyPreview()
 {
 	UUserWidget* PreviewActor = GetPreview();
@@ -623,21 +601,6 @@ void FWidgetBlueprintEditor::DestroyPreview()
 
 		PreviewActor->MarkPendingKill();
 	}
-
-	//if ( PreviewBlueprint != NULL )
-	//{
-	//	if ( PreviewBlueprint->SimpleConstructionScript != NULL
-	//		&& PreviewActor == PreviewBlueprint->SimpleConstructionScript->GetComponentEditorActorInstance() )
-	//	{
-	//		// Ensure that all editable component references are cleared
-	//		PreviewBlueprint->SimpleConstructionScript->ClearEditorComponentReferences();
-
-	//		// Clear the reference to the preview actor instance
-	//		PreviewBlueprint->SimpleConstructionScript->SetComponentEditorActorInstance(NULL);
-	//	}
-
-	//	PreviewBlueprint = NULL;
-	//}
 }
 
 void FWidgetBlueprintEditor::UpdatePreview(UBlueprint* InBlueprint, bool bInForceFullUpdate)
diff --git a/Engine/Source/Editor/UMGEditor/Private/WidgetBlueprintEditor.h b/Engine/Source/Editor/UMGEditor/Private/WidgetBlueprintEditor.h
--- a/Engine/Source/Editor/UMGEditor/Private/WidgetBlueprintEditor.h
+++ b/Engine/Source/Editor/UMGEditor/Private/WidgetBlueprintEditor.h
@@ -40,6 +40,12 @@ public:
 	/** @return The preview widget. */
 	UUserWidget* GetPreview() const;
 
+	/** Causes the preview to be destroyed and a new one to be created next tick */
+	void InvalidatePreview();
+
+	/** Immediately rebuilds the preview widget. */
+	void RefreshPreview();
+
 	/** Creates a widget reference using the template. */
 	FWidgetReference GetReferenceFromTemplate(UWidget* TemplateWidget);
 
@@ -84,6 +90,7 @@ public:
 
 	/** Creates a sequencer widget */
 	TSharedRef<SWidget> CreateSequencerWidget();
+
 public:
 	/** Fires whenever the selected set of widgets changing */
 	FOnSelectedWidgetsChanged OnSelectedWidgetsChanging;
@@ -92,7 +99,7 @@ public:
 	FOnSelectedWidgetsChanged OnSelectedWidgetsChanged;
 
 	/** Command list for handling widget actions in the WidgetBlueprintEditor */
-	TSharedPtr< FUICommandList > WidgetCommandList;
+	TSharedPtr< FUICommandList > DesignerCommandList;
 
 	/** Paste Metadata */
 	FVector2D PasteDropLocation;
@@ -165,4 +172,7 @@ private:
 
 	/** The widget references out in the ether that may need to be updated after being issued. */
 	TArray< TWeakPtr<FWidgetHandle> > WidgetHandlePool;
+
+	/** The preview becomes invalid and needs to be rebuilt on the next tick. */
+	bool bPreviewInvalidated;
 };
diff --git a/Engine/Source/Editor/UMGEditor/Private/WidgetBlueprintEditorUtils.cpp b/Engine/Source/Editor/UMGEditor/Private/WidgetBlueprintEditorUtils.cpp
--- a/Engine/Source/Editor/UMGEditor/Private/WidgetBlueprintEditorUtils.cpp
+++ b/Engine/Source/Editor/UMGEditor/Private/WidgetBlueprintEditorUtils.cpp
@@ -176,7 +176,7 @@ void FWidgetBlueprintEditorUtils::CreateWidgetContextMenu(FMenuBuilder& MenuBuil
 	TSet<FWidgetReference> Widgets = BlueprintEditor->GetSelectedWidgets();
 	UWidgetBlueprint* BP = BlueprintEditor->GetWidgetBlueprintObj();
 
-	MenuBuilder.PushCommandList(BlueprintEditor->WidgetCommandList.ToSharedRef());
+	MenuBuilder.PushCommandList(BlueprintEditor->DesignerCommandList.ToSharedRef());
 
 	MenuBuilder.BeginSection("Edit", LOCTEXT("Edit", "Edit"));
 	{
diff --git a/Engine/Source/Editor/UMGEditor/Public/IUMGDesigner.h b/Engine/Source/Editor/UMGEditor/Public/IUMGDesigner.h
--- a/Engine/Source/Editor/UMGEditor/Public/IUMGDesigner.h
+++ b/Engine/Source/Editor/UMGEditor/Public/IUMGDesigner.h
@@ -2,6 +2,20 @@
 
 #pragma once
 
+/**
+ * The logical type of transform that can be applied to a widget.
+ */
+namespace ETransformMode
+{
+	enum Type
+	{
+		/** Allows parent transfers */
+		Layout,
+		/** Only effects the rendered appearance of the widget */
+		Render,
+	};
+}
+
 /**
  * The public interface implemented by the UMG Designer to allow extensions to call methods
  * on the designer.
@@ -11,4 +25,10 @@ class UMGEDITOR_API IUMGDesigner
 public:
 	/** @return the effective preview scale after both the DPI and Zoom scale has been applied. */
 	virtual float GetPreviewScale() const = 0;
+
+	/** @return The currently selected widget. */
+	virtual FWidgetReference GetSelectedWidget() const = 0;
+
+	/** @return Get the transform mode currently in use in the designer. */
+	virtual ETransformMode::Type GetTransformMode() const = 0;
 };
diff --git a/Engine/Source/Runtime/Engine/Private/ObjectEditorUtils.cpp b/Engine/Source/Runtime/Engine/Private/ObjectEditorUtils.cpp
--- a/Engine/Source/Runtime/Engine/Private/ObjectEditorUtils.cpp
+++ b/Engine/Source/Runtime/Engine/Private/ObjectEditorUtils.cpp
@@ -59,6 +59,39 @@ namespace FObjectEditorUtils
 		}
 		return bResult;
 	}
+
+	bool MigratePropertyValue(UObject* SourceObject, UProperty* SourceProperty, UObject* DestinationObject, UProperty* DestinationProperty)
+	{
+		FString SourceValue;
+
+		// Get the property addresses for the source and destination objects.
+		uint8* SourceAddr = SourceProperty->ContainerPtrToValuePtr<uint8>(SourceObject);
+		uint8* DestionationAddr = DestinationProperty->ContainerPtrToValuePtr<uint8>(DestinationObject);
+
+		if ( SourceAddr == NULL || DestionationAddr == NULL )
+		{
+			return false;
+		}
+
+		// Get the current value from the source object.
+		SourceProperty->ExportText_Direct(SourceValue, SourceAddr, SourceAddr, NULL, PPF_Localized);
+
+		if ( !DestinationObject->HasAnyFlags(RF_ClassDefaultObject) )
+		{
+			DestinationObject->PreEditChange(DestinationProperty);
+		}
+
+		// Set the value on the destination object.
+		DestinationProperty->ImportText(*SourceValue, DestionationAddr, 0, DestinationObject);
+
+		if ( !DestinationObject->HasAnyFlags(RF_ClassDefaultObject) )
+		{
+			FPropertyChangedEvent PropertyEvent(DestinationProperty);
+			DestinationObject->PostEditChangeProperty(PropertyEvent);
+		}
+
+		return true;
+	}
 }
 
 #endif // WITH_EDITOR
diff --git a/Engine/Source/Runtime/Engine/Public/ObjectEditorUtils.h b/Engine/Source/Runtime/Engine/Public/ObjectEditorUtils.h
--- a/Engine/Source/Runtime/Engine/Public/ObjectEditorUtils.h
+++ b/Engine/Source/Runtime/Engine/Public/ObjectEditorUtils.h
@@ -1,8 +1,5 @@
 // Copyright 1998-2014 Epic Games, Inc. All Rights Reserved.
 
-#ifndef __ObjectEditorUtils_h__
-#define __ObjectEditorUtils_h__
-
 #pragma once
 
 #if WITH_EDITOR
@@ -36,7 +33,7 @@ namespace FObjectEditorUtils
 	 *
 	 * @return	true if the function is hidden in the specified class.
 	 */
-	ENGINE_API bool IsFunctionHiddenFromClass(  const UFunction* InFunction, const UClass* Class );
+	ENGINE_API bool IsFunctionHiddenFromClass( const UFunction* InFunction, const UClass* Class );
 
 	/**
 	 * Query if a the category a variable belongs to is flagged as hidden from the given class.
@@ -46,8 +43,77 @@ namespace FObjectEditorUtils
 	 *
 	 * @return	true if the functions category is hidden in the specified class.
 	 */
-	ENGINE_API bool IsVariableCategoryHiddenFromClass(  const UProperty* InVariable, const UClass* Class );
+	ENGINE_API bool IsVariableCategoryHiddenFromClass( const UProperty* InVariable, const UClass* Class );
+
+	/**
+	 * Copy the value of a property from source object to a destination object.
+	 *
+	 * @param	SourceObject		The object to copy the property value from.
+	 * @param	SourceProperty		The property on the SourceObject to copy the value from.
+	 * @param	DestinationObject	The object to copy the property value to.
+	 * @param	DestinationProperty	The property on the DestinationObject to copy the value to.
+	 *
+	 * @return true if the value was successfully migrated.
+	 */
+	ENGINE_API bool MigratePropertyValue(UObject* SourceObject, UProperty* SourceProperty, UObject* DestinationObject, UProperty* DestinationProperty);
+
+	/**
+	 * Copy the value of a property from source object to a destination object.
+	 *
+	 * @param	SourceObject			The object to copy the property value from.
+	 * @param	SourcePropertyName		The name of the property on the SourceObject to copy the value from.
+	 * @param	DestinationObject		The object to copy the property value to.
+	 * @param	DestinationPropertyName	The name of the property on the DestinationObject to copy the value to.
+	 *
+	 * @return true if the value was successfully migrated.
+	 */
+	template <typename SourceType, typename DestinationType>
+	bool MigratePropertyValue(SourceType* SourceObject, FName SourcePropertyName, DestinationType* DestinationObject, FName DestinationPropertyName)
+	{
+		UProperty* SourceProperty = FindFieldChecked<UProperty>(SourceType::StaticClass(), SourcePropertyName);
+		UProperty* DestinationProperty = FindFieldChecked<UProperty>(DestinationType::StaticClass(), DestinationPropertyName);
+
+		return FObjectEditorUtils::MigratePropertyValue(SourceObject, SourceProperty, DestinationObject, DestinationProperty);
+	}
+
+	/**
+	 * Set the value on an UObject using reflection.
+	 * @param	Object			The object to copy the value into.
+	 * @param	PropertyName	The name of the property to set.
+	 * @param	Value			The value to assign to the property.
+	 *
+	 * @return true if the value was set correctly
+	 */
+	template <typename ObjectType, typename ValueType>
+	bool SetPropertyValue(ObjectType* Object, FName PropertyName, ValueType Value)
+	{
+		// Get the property addresses for the source and destination objects.
+		UProperty* Property = FindFieldChecked<UProperty>(ObjectType::StaticClass(), PropertyName);
+
+		// Get the property addresses for the object
+		ValueType* SourceAddr = Property->ContainerPtrToValuePtr<ValueType>(Object);
+
+		if ( SourceAddr == NULL )
+		{
+			return false;
+		}
+
+		if ( !Object->HasAnyFlags(RF_ClassDefaultObject) )
+		{
+			Object->PreEditChange(Property);
+		}
+
+		// Set the value on the destination object.
+		*SourceAddr = Value;
+
+		if ( !Object->HasAnyFlags(RF_ClassDefaultObject) )
+		{
+			FPropertyChangedEvent PropertyEvent(Property);
+			Object->PostEditChangeProperty(PropertyEvent);
+		}
+
+		return true;
+	}
 };
-#endif // WITH_EDITOR
 
-#endif// __ObjectEditorUtils_h__
+#endif // WITH_EDITOR
diff --git a/Engine/Source/Runtime/UMG/Private/Components/Border.cpp b/Engine/Source/Runtime/UMG/Private/Components/Border.cpp
--- a/Engine/Source/Runtime/UMG/Private/Components/Border.cpp
+++ b/Engine/Source/Runtime/UMG/Private/Components/Border.cpp
@@ -2,6 +2,7 @@
 
 #include "UMGPrivatePCH.h"
 #include "Slate/SlateBrushAsset.h"
+#include "ObjectEditorUtils.h"
 
 #define LOCTEXT_NAMESPACE "UMG"
 
@@ -13,12 +14,8 @@ UBorder::UBorder(const FPostConstructInitializeProperties& PCIP)
 {
 	bIsVariable = false;
 
-	ContentScale = FVector2D(1.0f, 1.0f);
 	ContentColorAndOpacity = FLinearColor::White;
-
-	DesiredSizeScale = FVector2D(1.0f, 1.0f);
 	BrushColor = FLinearColor::White;
-	ForegroundColor = FLinearColor::Black;
 
 	SBorder::FArguments BorderDefaults;
 	bShowEffectWhenDisabled = BorderDefaults._ShowEffectWhenDisabled.Get();
@@ -47,16 +44,16 @@ void UBorder::SynchronizeProperties()
 {
 	Super::SynchronizeProperties();
 	
-	MyBorder->SetBorderBackgroundColor(BrushColor);
-	MyBorder->SetColorAndOpacity(ContentColorAndOpacity);
-	MyBorder->SetForegroundColor(ForegroundColor);
+	TAttribute<FLinearColor> ContentColorAndOpacityBinding = OPTIONAL_BINDING(FLinearColor, ContentColorAndOpacity);
+	TAttribute<FSlateColor> BrushColorBinding = OPTIONAL_BINDING_CONVERT(FLinearColor, BrushColor, FSlateColor, ConvertLinearColorToSlateColor);
+	TAttribute<const FSlateBrush*> ImageBinding = OPTIONAL_BINDING_CONVERT(FSlateBrush, Background, const FSlateBrush*, ConvertImage);
 	
-	MyBorder->SetContentScale(ContentScale);
-	MyBorder->SetDesiredSizeScale(DesiredSizeScale);
+	MyBorder->SetBorderBackgroundColor(BrushColorBinding);
+	MyBorder->SetColorAndOpacity(ContentColorAndOpacityBinding);
+
+	MyBorder->SetBorderImage(ImageBinding);
 	
 	MyBorder->SetShowEffectWhenDisabled(bShowEffectWhenDisabled);
-	
-	MyBorder->SetBorderImage(GetBorderBrush());
 
 	MyBorder->SetOnMouseButtonDown(BIND_UOBJECT_DELEGATE(FPointerEventHandler, HandleMouseButtonDown));
 	MyBorder->SetOnMouseButtonUp(BIND_UOBJECT_DELEGATE(FPointerEventHandler, HandleMouseButtonUp));
@@ -71,10 +68,18 @@ UClass* UBorder::GetSlotClass() const
 
 void UBorder::OnSlotAdded(UPanelSlot* Slot)
 {
+	// Copy the content properties into the new slot so that it matches what has been setup
+	// so far by the user.
+	UBorderSlot* BorderSlot = CastChecked<UBorderSlot>(Slot);
+	BorderSlot->Padding = Padding;
+	BorderSlot->HorizontalAlignment = HorizontalAlignment;
+	BorderSlot->VerticalAlignment = VerticalAlignment;
+
 	// Add the child to the live slot if it already exists
 	if ( MyBorder.IsValid() )
 	{
-		Cast<UBorderSlot>(Slot)->BuildSlot(MyBorder.ToSharedRef());
+		// Construct the underlying slot.
+		BorderSlot->BuildSlot(MyBorder.ToSharedRef());
 	}
 }
 
@@ -87,34 +92,24 @@ void UBorder::OnSlotRemoved(UPanelSlot* Slot)
 	}
 }
 
-void UBorder::SetBrushColor(FLinearColor Color)
+void UBorder::SetContentColorAndOpacity(FLinearColor Color)
 {
-	BrushColor = Color;
+	ContentColorAndOpacity = Color;
 	if ( MyBorder.IsValid() )
 	{
 		MyBorder->SetColorAndOpacity(Color);
 	}
 }
 
-void UBorder::SetForegroundColor(FLinearColor InForegroundColor)
+void UBorder::SetBrushColor(FLinearColor Color)
 {
-	ForegroundColor = InForegroundColor;
+	BrushColor = Color;
 	if ( MyBorder.IsValid() )
 	{
-		MyBorder->SetForegroundColor(InForegroundColor);
+		MyBorder->SetBorderBackgroundColor(Color);
 	}
 }
 
-const FSlateBrush* UBorder::GetBorderBrush() const
-{
-	if ( Brush == NULL )
-	{
-		return NULL;
-	}
-
-	return &Brush->Brush;
-}
-
 FReply UBorder::HandleMouseButtonDown(const FGeometry& Geometry, const FPointerEvent& MouseEvent)
 {
 	if ( OnMouseButtonDownEvent.IsBound() )
@@ -155,10 +150,86 @@ FReply UBorder::HandleMouseDoubleClick(const FGeometry& Geometry, const FPointer
 	return FReply::Unhandled();
 }
 
+void UBorder::SetBrushFromAsset(USlateBrushAsset* Asset)
+{
+	Background = Asset ? Asset->Brush : FSlateBrush();
+
+	if ( MyBorder.IsValid() )
+	{
+		MyBorder->SetBorderImage(&Background);
+	}
+}
+
+void UBorder::SetBrushFromTexture(UTexture2D* Texture)
+{
+	Background.SetResourceObject(Texture);
+
+	if ( MyBorder.IsValid() )
+	{
+		MyBorder->SetBorderImage(&Background);
+	}
+}
+
+void UBorder::SetBrushFromMaterial(UMaterialInterface* Material)
+{
+	Background.SetResourceObject(Material);
+
+	//TODO UMG Check if the material can be used with the UI
+
+	if ( MyBorder.IsValid() )
+	{
+		MyBorder->SetBorderImage(&Background);
+	}
+}
+
+UMaterialInstanceDynamic* UBorder::GetDynamicMaterial()
+{
+	UMaterialInterface* Material = NULL;
+
+	UObject* Resource = Background.GetResourceObject();
+	Material = Cast<UMaterialInterface>(Resource);
+
+	if ( Material )
+	{
+		UMaterialInstanceDynamic* DynamicMaterial = Cast<UMaterialInstanceDynamic>(Material);
+
+		if ( !DynamicMaterial )
+		{
+			DynamicMaterial = UMaterialInstanceDynamic::Create(Material, this);
+			Background.SetResourceObject(DynamicMaterial);
+
+			if ( MyBorder.IsValid() )
+			{
+				MyBorder->SetBorderImage(&Background);
+			}
+		}
+
+		return DynamicMaterial;
+	}
+
+	//TODO UMG can we do something for textures?  General purpose dynamic material for them?
+
+	return NULL;
+}
+
+const FSlateBrush* UBorder::ConvertImage(TAttribute<FSlateBrush> InImageAsset) const
+{
+	UBorder* MutableThis = const_cast<UBorder*>( this );
+	MutableThis->Background = InImageAsset.Get();
+
+	return &Background;
+}
+
 void UBorder::PostLoad()
 {
 	Super::PostLoad();
 
+	if ( Brush_DEPRECATED != NULL )
+	{
+		Background = Brush_DEPRECATED->Brush;
+		Brush_DEPRECATED = nullptr;
+	}
+
 	if ( GetChildrenCount() > 0 )
 	{
 		//TODO UMG Pre-Release Upgrade, now have slots of their own.  Convert existing slot to new slot.
@@ -178,6 +249,41 @@ void UBorder::PostLoad()
 
 #if WITH_EDITOR
 
+void UBorder::PostEditChangeProperty(struct FPropertyChangedEvent& PropertyChangedEvent)
+{
+	Super::PostEditChangeProperty(PropertyChangedEvent);
+
+	static bool IsReentrant = false;
+
+	if ( !IsReentrant )
+	{
+		IsReentrant = true;
+
+		if ( PropertyChangedEvent.Property )
+		{
+			FName PropertyName = PropertyChangedEvent.Property->GetFName();
+
+			if ( UBorderSlot* Slot = Cast<UBorderSlot>(GetContentSlot()) )
+			{
+				if ( PropertyName == "Padding" )
+				{
+					FObjectEditorUtils::MigratePropertyValue(this, "Padding", Slot, "Padding");
+				}
+				else if ( PropertyName == "HorizontalAlignment" )
+				{
+					FObjectEditorUtils::MigratePropertyValue(this, "HorizontalAlignment", Slot, "HorizontalAlignment");
+				}
+				else if ( PropertyName == "VerticalAlignment" )
+				{
+					FObjectEditorUtils::MigratePropertyValue(this, "VerticalAlignment", Slot, "VerticalAlignment");
+				}
+			}
+		}
+
+		IsReentrant = false;
+	}
+}
+
 const FSlateBrush* UBorder::GetEditorIcon()
 {
 	return FUMGStyle::Get().GetBrush("Widget.Border");
diff --git a/Engine/Source/Runtime/UMG/Private/Components/BorderSlot.cpp b/Engine/Source/Runtime/UMG/Private/Components/BorderSlot.cpp
--- a/Engine/Source/Runtime/UMG/Private/Components/BorderSlot.cpp
+++ b/Engine/Source/Runtime/UMG/Private/Components/BorderSlot.cpp
@@ -69,3 +69,42 @@ void UBorderSlot::SynchronizeProperties()
 	SetHorizontalAlignment(HorizontalAlignment);
 	SetVerticalAlignment(VerticalAlignment);
 }
+
+#if WITH_EDITOR
+
+void UBorderSlot::PostEditChangeProperty(FPropertyChangedEvent& PropertyChangedEvent)
+{
+	Super::PostEditChangeProperty(PropertyChangedEvent);
+
+	static bool IsReentrant = false;
+
+	if ( !IsReentrant )
+	{
+		IsReentrant = true;
+
+		if ( PropertyChangedEvent.Property )
+		{
+			FName PropertyName = PropertyChangedEvent.Property->GetFName();
+
+			if ( UBorder* Border = CastChecked<UBorder>(Parent) )
+			{
+				if ( PropertyName == "Padding" )
+				{
+					FObjectEditorUtils::MigratePropertyValue(this, "Padding", Border, "Padding");
+				}
+				else if ( PropertyName == "HorizontalAlignment" )
+				{
+					FObjectEditorUtils::MigratePropertyValue(this, "HorizontalAlignment", Border, "HorizontalAlignment");
+				}
+				else if ( PropertyName == "VerticalAlignment" )
+				{
+					FObjectEditorUtils::MigratePropertyValue(this, "VerticalAlignment", Border, "VerticalAlignment");
+				}
+			}
+		}
+
+		IsReentrant = false;
+	}
+}
+
+#endif
diff --git a/Engine/Source/Runtime/UMG/Private/Components/Button.cpp b/Engine/Source/Runtime/UMG/Private/Components/Button.cpp
--- a/Engine/Source/Runtime/UMG/Private/Components/Button.cpp
+++ b/Engine/Source/Runtime/UMG/Private/Components/Button.cpp
@@ -7,19 +7,18 @@
 /////////////////////////////////////////////////////
 // UButton
 
+FName UButton::StyleName(TEXT("Style"));
+
 UButton::UButton(const FPostConstructInitializeProperties& PCIP)
 	: Super(PCIP)
 {
-	SButton::FArguments ButtonDefaults;
-
-	Style = NULL;
+	ButtonStyle = PCIP.CreateDefaultSubobject<UButtonWidgetStyle>(this, StyleName);
 
-	DesiredSizeScale = ButtonDefaults._DesiredSizeScale.Get();
-	ContentScale = ButtonDefaults._ContentScale.Get();
+	SButton::FArguments ButtonDefaults;
+	ButtonStyle->ButtonStyle = *ButtonDefaults._ButtonStyle;
 
 	ColorAndOpacity = FLinearColor::White;
 	BackgroundColor = FLinearColor::White;
-	//ForegroundColor = FLinearColor::Black;
 
 	ClickMethod = EButtonClickMethod::DownAndUp;
 	TouchMethod = EButtonTouchMethod::DownAndUp;
@@ -36,7 +35,14 @@ void UButton::ReleaseNativeWidget()
 
 TSharedRef<SWidget> UButton::RebuildWidget()
 {
+	const FButtonStyle* StylePtr = ( Style != nullptr ) ? Style->GetStyle<FButtonStyle>() : nullptr;
+	if ( StylePtr != nullptr )
+	{
+		ButtonStyle->ButtonStyle = *StylePtr;
+	}
+
 	MyButton = SNew(SButton)
+		.ButtonStyle(&ButtonStyle->ButtonStyle)
 		.ClickMethod(ClickMethod)
 		.TouchMethod(TouchMethod)
 		.IsFocusable(IsFocusable);
@@ -53,36 +59,9 @@ void UButton::SynchronizeProperties()
 {
 	Super::SynchronizeProperties();
 
-	TOptional<FSlateSound> OptionalPressedSound;
-	if ( PressedSound.GetResourceObject() )
-	{
-		OptionalPressedSound = PressedSound;
-	}
-
-	TOptional<FSlateSound> OptionalHoveredSound;
-	if ( HoveredSound.GetResourceObject() )
-	{
-		OptionalHoveredSound = HoveredSound;
-	}
-
-	if ( MyStyle.IsSet() )
-	{
-		MyButton->SetButtonStyle(&MyStyle.GetValue());
-	}
-	else
-	{
-		const FButtonStyle* StylePtr = GetStyle();
-		MyButton->SetButtonStyle(StylePtr);
-	}
-
 	MyButton->SetColorAndOpacity( ColorAndOpacity );
 	MyButton->SetBorderBackgroundColor( BackgroundColor );
 	
-	MyButton->SetDesiredSizeScale( DesiredSizeScale );
-	MyButton->SetContentScale( ContentScale );
-	MyButton->SetPressedSound( OptionalPressedSound );
-	MyButton->SetHoveredSound( OptionalHoveredSound );
-
 	MyButton->SetOnClicked(BIND_UOBJECT_DELEGATE(FOnClicked, SlateHandleClicked));
 }
 
@@ -109,52 +88,14 @@ void UButton::OnSlotRemoved(UPanelSlot* Slot)
 	}
 }
 
-void UButton::SetStyle(USlateWidgetStyleAsset* InStyle)
-{
-	Style = InStyle;
-	MyStyle = TOptional<FButtonStyle>();
-
-	const FButtonStyle* StylePtr = GetStyle();
-
-	if ( MyButton.IsValid() )
-	{
-		MyButton->SetButtonStyle(StylePtr);
-	}
-}
-
-const FButtonStyle* UButton::GetStyle() const
-{
-	const FButtonStyle* StylePtr = ( Style != NULL ) ? Style->GetStyle<FButtonStyle>() : NULL;
-	if ( StylePtr == NULL )
-	{
-		SButton::FArguments ButtonDefaults;
-		StylePtr = ButtonDefaults._ButtonStyle;
-	}
-
-	return StylePtr;
-}
-
 void UButton::SetButtonStyle(FButtonStyle InButtonStyle)
 {
-	MyStyle = InButtonStyle;
-
-	if ( MyButton.IsValid() )
-	{
-		MyButton->SetButtonStyle(&MyStyle.GetValue());
-	}
+	ButtonStyle->ButtonStyle = InButtonStyle;
 }
 
 FButtonStyle UButton::GetButtonStyle()
 {
-	// If the dynamic style hasn't been set, default it to a clone of the current
-	// button style asset.
-	if ( !MyStyle.IsSet() )
-	{
-		const FButtonStyle* StylePtr = GetStyle();
-		MyStyle = *StylePtr;
-	}
-	
-	return MyStyle.GetValue();
+	return ButtonStyle->ButtonStyle;
 }
 
 void UButton::SetColorAndOpacity(FLinearColor Color)
diff --git a/Engine/Source/Runtime/UMG/Private/Components/Image.cpp b/Engine/Source/Runtime/UMG/Private/Components/Image.cpp
--- a/Engine/Source/Runtime/UMG/Private/Components/Image.cpp
+++ b/Engine/Source/Runtime/UMG/Private/Components/Image.cpp
@@ -10,11 +10,21 @@
 
 UImage::UImage(const FPostConstructInitializeProperties& PCIP)
 	: Super(PCIP)
-	, Image()
 	, ColorAndOpacity(FLinearColor::White)
 {
 }
 
+void UImage::PostLoad()
+{
+	Super::PostLoad();
+
+	if ( Image_DEPRECATED != NULL )
+	{
+		Brush = Image_DEPRECATED->Brush;
+		Image_DEPRECATED = nullptr;
+	}
+}
+
 void UImage::ReleaseNativeWidget()
 {
 	Super::ReleaseNativeWidget();
@@ -33,7 +43,7 @@ void UImage::SynchronizeProperties()
 	Super::SynchronizeProperties();
 
 	TAttribute<FSlateColor> ColorAndOpacityBinding = OPTIONAL_BINDING(FSlateColor, ColorAndOpacity);
-	TAttribute<const FSlateBrush*> ImageBinding = OPTIONAL_BINDING_CONVERT(USlateBrushAsset*, Image, const FSlateBrush*, ConvertImage);
+	TAttribute<const FSlateBrush*> ImageBinding = OPTIONAL_BINDING_CONVERT(FSlateBrush, Brush, const FSlateBrush*, ConvertImage);
 
 	MyImage->SetImage(ImageBinding);
 	MyImage->SetColorAndOpacity(ColorAndOpacityBinding);
@@ -58,81 +68,43 @@ void UImage::SetOpacity(float InOpacity)
 	}
 }
 
-const FSlateBrush* UImage::ConvertImage(TAttribute<USlateBrushAsset*> InImageAsset) const
-{
-	USlateBrushAsset* ImageAsset = InImageAsset.Get();
-	if ( ImageAsset != Image )
-	{
-		UImage* MutableThis = const_cast< UImage* >( this );
-		MutableThis->DynamicBrush = TOptional<FSlateBrush>();
-		MutableThis->Image = ImageAsset;
-	}
-	else if ( DynamicBrush.IsSet() )
-	{
-		return &DynamicBrush.GetValue();
-	}
-
-	if ( Image == NULL && IsDesignTime() )
-	{
-		SImage::FArguments ImageDefaults;
-		return ImageDefaults._Image.Get();
-	}
-
-	return Image ? &Image->Brush : NULL;
-}
-
-const FSlateBrush* UImage::GetImageBrush() const
-{
-	return ConvertImage(Image);
-}
-
-void UImage::SetImage(USlateBrushAsset* InImage)
+const FSlateBrush* UImage::ConvertImage(TAttribute<FSlateBrush> InImageAsset) const
 {
-	if ( InImage != Image )
-	{
-		DynamicBrush = TOptional<FSlateBrush>();
-	}
+	UImage* MutableThis = const_cast<UImage*>( this );
+	MutableThis->Brush = InImageAsset.Get();
 
-	Image = InImage;
-	if ( MyImage.IsValid() )
-	{
-		MyImage->SetImage(GetImageBrush());
-	}
+	return &Brush;
 }
 
-void UImage::SetImageFromBrush(FSlateBrush InImage)
+void UImage::SetBrushFromAsset(USlateBrushAsset* Asset)
 {
-	DynamicBrush = InImage;
+	Brush = Asset ? Asset->Brush : FSlateBrush();
 
 	if ( MyImage.IsValid() )
 	{
-		MyImage->SetImage(&DynamicBrush.GetValue());
+		MyImage->SetImage(&Brush);
 	}
 }
 
-void UImage::SetImageFromTexture(UTexture2D* Texture)
+void UImage::SetBrushFromTexture(UTexture2D* Texture)
 {
-	FSlateBrush TextureBrush;
-	TextureBrush.SetResourceObject(Texture);
-
-	DynamicBrush = TextureBrush;
+	Brush.SetResourceObject(Texture);
 
 	if ( MyImage.IsValid() )
 	{
-		MyImage->SetImage(&DynamicBrush.GetValue());
+		MyImage->SetImage(&Brush);
 	}
 }
 
-void UImage::SetImageFromMaterial(UMaterialInterface* Material)
+void UImage::SetBrushFromMaterial(UMaterialInterface* Material)
 {
-	FSlateBrush MaterialBrush;
-	MaterialBrush.SetResourceObject(Material);
+	Brush.SetResourceObject(Material);
 
-	DynamicBrush = MaterialBrush;
+	//TODO UMG Check if the material can be used with the UI
 
 	if ( MyImage.IsValid() )
 	{
-		MyImage->SetImage(&DynamicBrush.GetValue());
+		MyImage->SetImage(&Brush);
 	}
 }
 
@@ -140,17 +112,8 @@ UMaterialInstanceDynamic* UImage::GetDynamicMaterial()
 {
 	UMaterialInterface* Material = NULL;
 
-	if ( !DynamicBrush.IsSet() )
-	{
-		const FSlateBrush* Brush = GetImageBrush();
-		UObject* Resource = Brush->GetResourceObject();
-		Material = Cast<UMaterialInterface>(Resource);
-	}
-	else
-	{
-		UObject* Resource = DynamicBrush.GetValue().GetResourceObject();
-		Material = Cast<UMaterialInterface>(Resource);
-	}
+	UObject* Resource = Brush.GetResourceObject();
+	Material = Cast<UMaterialInterface>(Resource);
 
 	if ( Material )
 	{
@@ -159,16 +122,11 @@ UMaterialInstanceDynamic* UImage::GetDynamicMaterial()
 		if ( !DynamicMaterial )
 		{
 			DynamicMaterial = UMaterialInstanceDynamic::Create(Material, this);
-
-			const FSlateBrush* Brush = GetImageBrush();
-			FSlateBrush ClonedBrush = *Brush;
-			ClonedBrush.SetResourceObject(DynamicMaterial);
-
-			DynamicBrush = ClonedBrush;
+			Brush.SetResourceObject(DynamicMaterial);
 
 			if ( MyImage.IsValid() )
 			{
-				MyImage->SetImage(&DynamicBrush.GetValue());
+				MyImage->SetImage(&Brush);
 			}
 		}
 		
diff --git a/Engine/Source/Runtime/UMG/Private/Components/Widget.cpp b/Engine/Source/Runtime/UMG/Private/Components/Widget.cpp
--- a/Engine/Source/Runtime/UMG/Private/Components/Widget.cpp
+++ b/Engine/Source/Runtime/UMG/Private/Components/Widget.cpp
@@ -12,6 +12,59 @@ UWidget::UWidget(const FPostConstructInitializeProperties& PCIP)
 	bIsVariable = true;
 	bDesignTime = false;
 	Visiblity = ESlateVisibility::Visible;
+	RenderTransformPivot = FVector2D(0.5f, 0.5f);
+}
+
+void UWidget::SetRenderTransform(FWidgetTransform Transform)
+{
+	RenderTransform = Transform;
+	UpdateRenderTransform();
+}
+
+void UWidget::SetRenderScale(FVector2D Scale)
+{
+	RenderTransform.Scale = Scale;
+	UpdateRenderTransform();
+}
+
+void UWidget::SetRenderShear(FVector2D Shear)
+{
+	RenderTransform.Shear = Shear;
+	UpdateRenderTransform();
+}
+
+void UWidget::SetRenderAngle(float Angle)
+{
+	RenderTransform.Angle = Angle;
+	UpdateRenderTransform();
+}
+
+void UWidget::SetRenderTranslation(FVector2D Translation)
+{
+	RenderTransform.Translation = Translation;
+	UpdateRenderTransform();
+}
+
+void UWidget::UpdateRenderTransform()
+{
+	if ( MyWidget.IsValid() )
+	{
+		// Compute the M (Shear Slot) = CoTan(90 - SlopeAngle)
+		float ShearX = 1.0f / FMath::Tan(FMath::DegreesToRadians(90 - FMath::Clamp(RenderTransform.Shear.X, -89.0f, 89.0f)));
+		float ShearY = 1.0f / FMath::Tan(FMath::DegreesToRadians(90 - FMath::Clamp(RenderTransform.Shear.Y, -89.0f, 89.0f)));
+
+		FSlateRenderTransform Transform2D = ::Concatenate(FScale2D(RenderTransform.Scale), FShear2D(ShearX, ShearY), FQuat2D(FMath::DegreesToRadians(RenderTransform.Angle)), FVector2D(RenderTransform.Translation));
+		MyWidget.Pin()->SetRenderTransform(Transform2D);
+	}
+}
+
+void UWidget::SetRenderTransformPivot(FVector2D Pivot)
+{
+	RenderTransformPivot = Pivot;
+	if ( MyWidget.IsValid() )
+	{
+		MyWidget.Pin()->SetRenderTransformPivot(Pivot);
+	}
 }
 
 bool UWidget::GetIsEnabled() const
@@ -267,12 +320,6 @@ const FSlateBrush* UWidget::GetEditorIcon()
 	return FUMGStyle::Get().GetBrush("Widget");
 }
 
-TSharedRef<SWidget> UWidget::GetToolboxPreviewWidget() const
-{
-	return SNew(SImage);
-	//.Image( FEditorStyle::GetBrush("UMGEditor.ToolboxPreviewWidget") );
-}
-
 void UWidget::PostEditChangeProperty(FPropertyChangedEvent& PropertyChangedEvent)
 {
 	Super::PostEditChangeProperty(PropertyChangedEvent);
@@ -360,6 +407,9 @@ void UWidget::SynchronizeProperties()
 		SafeWidget->SetVisibility(OPTIONAL_BINDING_CONVERT(ESlateVisibility::Type, Visiblity, EVisibility, ConvertVisibility));
 	}
 
+	UpdateRenderTransform();
+	SafeWidget->SetRenderTransformPivot(RenderTransformPivot);
+
 	if ( !ToolTipText.IsEmpty() )
 	{
 		SafeWidget->SetToolTipText(OPTIONAL_BINDING(FText, ToolTipText));
diff --git a/Engine/Source/Runtime/UMG/Private/SlateBlueprintLibrary.cpp b/Engine/Source/Runtime/UMG/Private/SlateBlueprintLibrary.cpp
new file mode 100644
--- /dev/null
+++ b/Engine/Source/Runtime/UMG/Private/SlateBlueprintLibrary.cpp
@@ -0,0 +1,33 @@
+// Copyright 1998-2014 Epic Games, Inc. All Rights Reserved.
+
+#include "UMGPrivatePCH.h"
+
+#include "Slate/SlateBrushAsset.h"
+#include "SlateBlueprintLibrary.h"
+
+#define LOCTEXT_NAMESPACE "UMG"
+
+/////////////////////////////////////////////////////
+// USlateBlueprintLibrary
+
+USlateBlueprintLibrary::USlateBlueprintLibrary(const FPostConstructInitializeProperties& PCIP)
+: Super(PCIP)
+{
+}
+
+bool USlateBlueprintLibrary::IsUnderLocation(const FGeometry& Geometry, const FVector2D& AbsoluteCoordinate)
+{
+	return Geometry.IsUnderLocation(AbsoluteCoordinate);
+}
+
+FVector2D USlateBlueprintLibrary::AbsoluteToLocal(const FGeometry& Geometry, FVector2D AbsoluteCoordinate)
+{
+	return Geometry.AbsoluteToLocal(AbsoluteCoordinate);
+}
+
+FVector2D USlateBlueprintLibrary::LocalToAbsolute(const FGeometry& Geometry, FVector2D LocalCoordinate)
+{
+	return Geometry.LocalToAbsolute(LocalCoordinate);
+}
+
+#undef LOCTEXT_NAMESPACE
\ No newline at end of file
diff --git a/Engine/Source/Runtime/UMG/Public/Blueprint/SlateBlueprintLibrary.h b/Engine/Source/Runtime/UMG/Public/Blueprint/SlateBlueprintLibrary.h
new file mode 100644
--- /dev/null
+++ b/Engine/Source/Runtime/UMG/Public/Blueprint/SlateBlueprintLibrary.h
@@ -0,0 +1,37 @@
+// Copyright 1998-2014 Epic Games, Inc. All Rights Reserved.
+
+#pragma once
+
+#include "SlateBlueprintLibrary.generated.h"
+
+UCLASS(MinimalAPI)
+class USlateBlueprintLibrary : public UBlueprintFunctionLibrary
+{
+	GENERATED_UCLASS_BODY()
+
+	/**
+	 * Absolute coordinates could be either desktop or window space depending on what space the root of the widget hierarchy is in.
+	 *
+	 * @return true if the provided location in absolute coordinates is within the bounds of this geometry.
+	 */
+	UFUNCTION(BlueprintCallable, Category="User Interface|Geometry")
+	static bool IsUnderLocation(const FGeometry& Geometry, const FVector2D& AbsoluteCoordinate);
+
+	/**
+	 * Absolute coordinates could be either desktop or window space depending on what space the root of the widget hierarchy is in.
+	 *
+	 * @return Transforms AbsoluteCoordinate into the local space of this Geometry.
+	 */
+	UFUNCTION(BlueprintCallable, Category="User Interface|Geometry")
+	static FVector2D AbsoluteToLocal(const FGeometry& Geometry, FVector2D AbsoluteCoordinate);
+
+	/**
+	 * Translates local coordinates into absolute coordinates
+	 *
+	 * Absolute coordinates could be either desktop or window space depending on what space the root of the widget hierarchy is in.
+	 *
+	 * @return  Absolute coordinates
+	 */
+	UFUNCTION(BlueprintCallable, Category="User Interface|Geometry")
+	static FVector2D LocalToAbsolute(const FGeometry& Geometry, FVector2D LocalCoordinate);
+};
diff --git a/Engine/Source/Runtime/UMG/Public/Components/Border.h b/Engine/Source/Runtime/UMG/Public/Components/Border.h
--- a/Engine/Source/Runtime/UMG/Public/Components/Border.h
+++ b/Engine/Source/Runtime/UMG/Public/Components/Border.h
@@ -18,20 +18,32 @@ class UMG_API UBorder : public UContentWidget
 public:
 
 	/** Color and opacity multiplier of content in the border */
-	UPROPERTY(EditDefaultsOnly, Category="Content Layout")
+	UPROPERTY(EditDefaultsOnly, Category="Content")
 	FLinearColor ContentColorAndOpacity;
 
-	/** The scaling factor for the border content */
-	UPROPERTY(EditDefaultsOnly, Category="Content Layout", AdvancedDisplay)
-	FVector2D ContentScale;
+	/** A bindable delegate for the ContentColorAndOpacity. */
+	UPROPERTY()
+	FGetLinearColor ContentColorAndOpacityDelegate;
 
-	/** Image to use for the border */
-	UPROPERTY(EditDefaultsOnly, Category=Appearance, meta=( DisplayThumbnail = "true" ))
-	USlateBrushAsset* Brush;
+	/** The padding area between the slot and the content it contains. */
+	UPROPERTY(EditDefaultsOnly, Category="Content")
+	FMargin Padding;
+
+	/** The alignment of the content horizontally. */
+	UPROPERTY(EditDefaultsOnly, Category="Content")
+	TEnumAsByte<EHorizontalAlignment> HorizontalAlignment;
+
+	/** The alignment of the content vertically. */
+	UPROPERTY(EditDefaultsOnly, Category="Content")
+	TEnumAsByte<EVerticalAlignment> VerticalAlignment;
+
+	/** Brush to drag as the background */
+	UPROPERTY(EditDefaultsOnly, Category=Appearance, meta=(DisplayName="Brush"))
+	FSlateBrush Background;
 
 	/** A bindable delegate for the Brush. */
 	UPROPERTY()
-	FGetSlateBrushAsset BrushDelegate;
+	FGetSlateBrush BackgroundDelegate;
 
 	/** Color and opacity of the actual border image */
 	UPROPERTY(EditDefaultsOnly, Category=Appearance)
@@ -41,18 +53,10 @@ public:
 	UPROPERTY()
 	FGetLinearColor BrushColorDelegate;
 
-	/** The foreground color of text and some glyphs that appear as the border's content. */
-	UPROPERTY(EditDefaultsOnly, Category=Appearance)
-	FLinearColor ForegroundColor;
-
 	/** Whether or not to show the disabled effect when this border is disabled */
 	UPROPERTY(EditDefaultsOnly, Category=Appearance, AdvancedDisplay)
 	bool bShowEffectWhenDisabled;
 
-	/** The scaling factor for the border */
-	UPROPERTY(EditDefaultsOnly, Category=Appearance, AdvancedDisplay)
-	FVector2D DesiredSizeScale;
-
 	UPROPERTY(EditDefaultsOnly, Category=Events)
 	FOnPointerEvent OnMouseButtonDownEvent;
 
@@ -69,11 +73,29 @@ public:
 
 	/**  */
 	UFUNCTION(BlueprintCallable, Category="Appearance")
-	void SetBrushColor(FLinearColor InColorAndOpacity);
+	void SetContentColorAndOpacity(FLinearColor InContentColorAndOpacity);
+
+	/**  */
+	UFUNCTION(BlueprintCallable, Category="Appearance")
+	void SetBrushColor(FLinearColor InBrushColor);
+
+	/**  */
+	UFUNCTION(BlueprintCallable, Category="Appearance")
+	void SetBrushFromAsset(USlateBrushAsset* Asset);
 
 	/**  */
 	UFUNCTION(BlueprintCallable, Category="Appearance")
-	void SetForegroundColor(FLinearColor InForegroundColor);
+	void SetBrushFromTexture(UTexture2D* Texture);
+
+	/**  */
+	UFUNCTION(BlueprintCallable, Category="Appearance")
+	void SetBrushFromMaterial(UMaterialInterface* Material);
+
+	/**  */
+	UFUNCTION(BlueprintCallable, Category="Appearance")
+	UMaterialInstanceDynamic* GetDynamicMaterial();
+
+public:
 
 	// UWidget interface
 	virtual void SynchronizeProperties() override;
@@ -88,6 +110,10 @@ public:
 	// End of UObject interface
 
 #if WITH_EDITOR
+	// UObject interface
+	virtual void PostEditChangeProperty(struct FPropertyChangedEvent& PropertyChangedEvent) override;
+	// End of UObject interface
+
 	virtual const FSlateBrush* GetEditorIcon() override;
 	virtual const FText GetToolboxCategory() override;
 #endif
@@ -113,5 +139,11 @@ protected:
 	FReply HandleMouseMove(const FGeometry& Geometry, const FPointerEvent& MouseEvent);
 	FReply HandleMouseDoubleClick(const FGeometry& Geometry, const FPointerEvent& MouseEvent);
 
-	const FSlateBrush* GetBorderBrush() const;
+	/** Translates the bound brush data and assigns it to the cached brush used by this widget. */
+	const FSlateBrush* ConvertImage(TAttribute<FSlateBrush> InImageAsset) const;
+
+protected:
+	/** Image to use for the border */
+	UPROPERTY()
+	USlateBrushAsset* Brush_DEPRECATED;
 };
diff --git a/Engine/Source/Runtime/UMG/Public/Components/BorderSlot.h b/Engine/Source/Runtime/UMG/Public/Components/BorderSlot.h
--- a/Engine/Source/Runtime/UMG/Public/Components/BorderSlot.h
+++ b/Engine/Source/Runtime/UMG/Public/Components/BorderSlot.h
@@ -44,6 +44,16 @@ public:
 
 	virtual void ReleaseNativeWidget() override;
 
+public:
+
+#if WITH_EDITOR
+
+	// UObject interface
+	virtual void PostEditChangeProperty(struct FPropertyChangedEvent& PropertyChangedEvent) override;
+	// End of UObject interface
+
+#endif
+
 private:
 
 	/** A pointer to the button to allow us to adjust the size, padding...etc at runtime. */
diff --git a/Engine/Source/Runtime/UMG/Public/Components/Button.h b/Engine/Source/Runtime/UMG/Public/Components/Button.h
--- a/Engine/Source/Runtime/UMG/Public/Components/Button.h
+++ b/Engine/Source/Runtime/UMG/Public/Components/Button.h
@@ -2,6 +2,8 @@
 
 #pragma once
 
+#include "ButtonWidgetStyle.h"
+
 #include "Button.generated.h"
 
 DECLARE_DYNAMIC_MULTICAST_DELEGATE(FOnButtonClickedEvent);
@@ -15,17 +17,13 @@ class UMG_API UButton : public UContentWidget
 	GENERATED_UCLASS_BODY()
 
 public:
-	/** Style of the button */
-	UPROPERTY(EditDefaultsOnly, Category=Style, meta=( DisplayThumbnail = "true" ))
+	/** The template style asset, used to seed the mutable instance of the style. */
+	UPROPERTY(EditDefaultsOnly, AdvancedDisplay, DisplayName="Template Style", Category=Appearance, meta=(DesignerRebuild=true))
 	USlateWidgetStyleAsset* Style;
 
-	/** The scaling factor for the button border */
-	UPROPERTY(EditDefaultsOnly, Category=Appearance, AdvancedDisplay)
-	FVector2D DesiredSizeScale;
-
-	/** The scaling factor for the button content */
-	UPROPERTY(EditDefaultsOnly, Category=Appearance, AdvancedDisplay)
-	FVector2D ContentScale;
+	/** The button style used at runtime by the slate button */
+	UPROPERTY(VisibleAnywhere, Instanced, Category=Appearance)
+	TSubobjectPtr<UButtonWidgetStyle> ButtonStyle;
 	
 	/** The color multiplier for the button content */
 	UPROPERTY(EditDefaultsOnly, Category=Appearance )
@@ -47,24 +45,12 @@ public:
 	UPROPERTY(EditDefaultsOnly, Category="Interaction", AdvancedDisplay)
 	bool IsFocusable;
 
-	/** The sound made when the user 'clicks' the button */
-	UPROPERTY(EditDefaultsOnly, Category=Sound)
-	FSlateSound PressedSound;
-
-	/** The sound made when the user hovers over the button */
-	UPROPERTY(EditDefaultsOnly, Category=Sound)
-	FSlateSound HoveredSound;
-
 public:
 
 	/** Called when the button is clicked */
 	UPROPERTY(BlueprintAssignable)
 	FOnButtonClickedEvent OnClicked;
 
-	/** Sets the look and feel of a button by giving it a new style from an existing style asset. */
-	UFUNCTION(BlueprintCallable, Category="Button|Appearance")
-	void SetStyle(USlateWidgetStyleAsset* InStyle);
-
 	/** Sets the look and feel of a button from a new button style struct. */
 	UFUNCTION(BlueprintCallable, Category="Button|Appearance")
 	void SetButtonStyle(FButtonStyle InButtonStyle);
@@ -102,6 +88,8 @@ public:
 	virtual const FText GetToolboxCategory() override;
 #endif
 
+	static FName StyleName;
+
 protected:
 
 	// UPanelWidget
@@ -119,9 +107,6 @@ protected:
 	virtual TSharedRef<SWidget> RebuildWidget() override;
 	// End of UWidget interface
 
-	/** Get the button style of the current style asset */
-	const FButtonStyle* GetStyle() const;
-
 protected:
 	/** Cached pointer to the underlying slate button owned by this UWidget */
 	TSharedPtr<SButton> MyButton;
diff --git a/Engine/Source/Runtime/UMG/Public/Components/ContentWidget.h b/Engine/Source/Runtime/UMG/Public/Components/ContentWidget.h
--- a/Engine/Source/Runtime/UMG/Public/Components/ContentWidget.h
+++ b/Engine/Source/Runtime/UMG/Public/Components/ContentWidget.h
@@ -9,12 +9,14 @@ class UMG_API UContentWidget : public UPanelWidget
 {
 	GENERATED_UCLASS_BODY()
 
+public:
+	/**  */
+	UFUNCTION(BlueprintCallable, Category="Widget|Panel")
+	UPanelSlot* GetContentSlot() const;
+
 protected:
 
 	// UPanelWidget
 	virtual UClass* GetSlotClass() const override;
 	// End UPanelWidget
-
-protected:
-	UPanelSlot* GetContentSlot() const;
 };
diff --git a/Engine/Source/Runtime/UMG/Public/Components/Image.h b/Engine/Source/Runtime/UMG/Public/Components/Image.h
--- a/Engine/Source/Runtime/UMG/Public/Components/Image.h
+++ b/Engine/Source/Runtime/UMG/Public/Components/Image.h
@@ -17,12 +17,16 @@ class UMG_API UImage : public UWidget
 public:
 
 	/** Image to draw */
-	UPROPERTY(EditDefaultsOnly, Category=Appearance, meta=( DisplayThumbnail = "true" ))
-	USlateBrushAsset* Image;
+	UPROPERTY()
+	USlateBrushAsset* Image_DEPRECATED;
+
+	/** Image to draw */
+	UPROPERTY(EditDefaultsOnly, Category=Appearance)
+	FSlateBrush Brush;
 
 	/** A bindable delegate for the Image. */
 	UPROPERTY()
-	FGetSlateBrushAsset ImageDelegate;
+	FGetSlateBrush BrushDelegate;
 
 	/** Color and opacity */
 	UPROPERTY(EditDefaultsOnly, Category=Appearance)
@@ -47,17 +51,15 @@ public:
 
 	/**  */
 	UFUNCTION(BlueprintCallable, Category="Appearance")
-	void SetImage(USlateBrushAsset* InImage);
+	void SetBrushFromAsset(USlateBrushAsset* Asset);
 
 	/**  */
 	UFUNCTION(BlueprintCallable, Category="Appearance")
-	void SetImageFromBrush(FSlateBrush Brush);
-
-	UFUNCTION(BlueprintCallable, Category="Appearance")
-	void SetImageFromTexture(UTexture2D* Texture);
+	void SetBrushFromTexture(UTexture2D* Texture);
 
+	/**  */
 	UFUNCTION(BlueprintCallable, Category="Appearance")
-	void SetImageFromMaterial(UMaterialInterface* Material);
+	void SetBrushFromMaterial(UMaterialInterface* Material);
 
 	/**  */
 	UFUNCTION(BlueprintCallable, Category="Appearance")
@@ -67,7 +69,13 @@ public:
 	virtual void SynchronizeProperties() override;
 	// End of UWidget interface
 
+	// UVisual interface
 	virtual void ReleaseNativeWidget() override;
+	// End of UVisual interface
+
+	// Begin UObject interface
+	virtual void PostLoad() override;
+	// End of UObject interface
 
 #if WITH_EDITOR
 	// UWidget interface
@@ -81,14 +89,11 @@ protected:
 	virtual TSharedRef<SWidget> RebuildWidget() override;
 	// End of UWidget interface
 
-	const FSlateBrush* GetImageBrush() const;
-
-	const FSlateBrush* ConvertImage(TAttribute<USlateBrushAsset*> InImageAsset) const;
+	/** Translates the bound brush data and assigns it to the cached brush used by this widget. */
+	const FSlateBrush* ConvertImage(TAttribute<FSlateBrush> InImageAsset) const;
 
 	FReply HandleMouseButtonDown(const FGeometry& Geometry, const FPointerEvent& MouseEvent);
 
 protected:
 	TSharedPtr<SImage> MyImage;
-
-	TOptional<FSlateBrush> DynamicBrush;
 };
diff --git a/Engine/Source/Runtime/UMG/Public/Components/Widget.h b/Engine/Source/Runtime/UMG/Public/Components/Widget.h
--- a/Engine/Source/Runtime/UMG/Public/Components/Widget.h
+++ b/Engine/Source/Runtime/UMG/Public/Components/Widget.h
@@ -3,6 +3,7 @@
 #pragma once
 
 #include "SlateWrapperTypes.h"
+#include "WidgetTransform.h"
 
 #include "Widget.generated.h"
 
@@ -41,6 +42,7 @@ public:
 	DECLARE_DYNAMIC_DELEGATE_RetVal(ESlateVisibility::Type, FGetSlateVisibility);
 	DECLARE_DYNAMIC_DELEGATE_RetVal(EMouseCursor::Type, FGetMouseCursor);
 	DECLARE_DYNAMIC_DELEGATE_RetVal(USlateBrushAsset*, FGetSlateBrushAsset);
+	DECLARE_DYNAMIC_DELEGATE_RetVal(FSlateBrush, FGetSlateBrush);
 	DECLARE_DYNAMIC_DELEGATE_RetVal(ESlateCheckBoxState::Type, FGetCheckBoxState);
 	DECLARE_DYNAMIC_DELEGATE_RetVal(UWidget*, FGetContent);
 
@@ -104,6 +106,38 @@ public:
 	UPROPERTY()
 	FGetMouseCursor CursorDelegate;
 
+	/**  */
+	UPROPERTY(EditDefaultsOnly, Category=Transform)
+	FWidgetTransform RenderTransform;
+
+	/**  */
+	UPROPERTY(EditDefaultsOnly, Category=Transform)
+	FVector2D RenderTransformPivot;
+
+	/** */
+	UFUNCTION(BlueprintCallable, Category="Widget|Transform")
+	void SetRenderTransform(FWidgetTransform InTransform);
+
+	/** */
+	UFUNCTION(BlueprintCallable, Category="Widget|Transform")
+	void SetRenderScale(FVector2D Scale);
+
+	/** */
+	UFUNCTION(BlueprintCallable, Category="Widget|Transform")
+	void SetRenderShear(FVector2D Shear);
+
+	/** */
+	UFUNCTION(BlueprintCallable, Category="Widget|Transform")
+	void SetRenderAngle(float Angle);
+	
+	/** */
+	UFUNCTION(BlueprintCallable, Category="Widget|Transform")
+	void SetRenderTranslation(FVector2D Translation);
+
+	/** */
+	UFUNCTION(BlueprintCallable, Category="Widget|Transform")
+	void SetRenderTransformPivot(FVector2D Pivot);
+
 	/** Gets the current enabled status of the widget */
 	UFUNCTION(BlueprintCallable, Category="Widget")
 	bool GetIsEnabled() const;
@@ -210,9 +244,6 @@ public:
 
 	/** Gets the editor icon */
 	virtual const FSlateBrush* GetEditorIcon();
-
-	/** Gets a widget representing the tiny preview of the toolbox */
-	virtual TSharedRef<SWidget> GetToolboxPreviewWidget() const;
 	
 	/** Allows general fixups and connections only used at editor time. */
 	virtual void ConnectEditorData() { }
@@ -253,6 +284,8 @@ protected:
 	
 	TSharedRef<SWidget> BuildDesignTimeWidget(TSharedRef<SWidget> WrapWidget);
 
+	void UpdateRenderTransform();
+
 protected:
 	//TODO UMG Consider moving conversion functions into another class.
 	// Conversion functions
@@ -266,6 +299,11 @@ protected:
 		return InFloat.Get();
 	}
 
+	FSlateColor ConvertLinearColorToSlateColor(TAttribute<FLinearColor> InLinearColor) const
+	{
+		return FSlateColor(InLinearColor.Get());
+	}
+
 protected:
 	/** The underlying SWidget. */
 	TWeakPtr<SWidget> MyWidget;
diff --git a/Engine/Source/Runtime/UMG/Public/Slate/Anchors.h b/Engine/Source/Runtime/UMG/Public/Slate/Anchors.h
--- a/Engine/Source/Runtime/UMG/Public/Slate/Anchors.h
+++ b/Engine/Source/Runtime/UMG/Public/Slate/Anchors.h
@@ -7,17 +7,17 @@
 /**
  * Describes how a widget is anchored.
  */
-USTRUCT()
+USTRUCT(BlueprintType)
 struct FAnchors
 {
 	GENERATED_USTRUCT_BODY()
 
 	/** Holds the minimum anchors, left + top. */
-	UPROPERTY(EditAnywhere, Category=Appearance)
+	UPROPERTY(EditAnywhere, BlueprintReadWrite, Category=Appearance)
 	FVector2D Minimum;
 
 	/** Holds the maximum anchors, right + bottom. */
-	UPROPERTY(EditAnywhere, Category=Appearance)
+	UPROPERTY(EditAnywhere, BlueprintReadWrite, Category=Appearance)
 	FVector2D Maximum;
 
 public:
diff --git a/Engine/Source/Runtime/UMG/Public/Slate/WidgetTransform.h b/Engine/Source/Runtime/UMG/Public/Slate/WidgetTransform.h
new file mode 100644
--- /dev/null
+++ b/Engine/Source/Runtime/UMG/Public/Slate/WidgetTransform.h
@@ -0,0 +1,38 @@
+// Copyright 1998-2014 Epic Games, Inc. All Rights Reserved.
+
+#pragma once
+
+#include "WidgetTransform.generated.h"
+
+/**
+ * Describes the standard transformation of a widget
+ */
+USTRUCT(BlueprintType)
+struct FWidgetTransform
+{
+	GENERATED_USTRUCT_BODY()
+
+	/** The scale to apply to the widget */
+	UPROPERTY(EditAnywhere, BlueprintReadWrite, Category=Transform, meta=( UIMin = "-5", UIMax = "5", Delta = "0.05" ))
+	FVector2D Scale;
+
+	/** The amount to shear the widget in slate units */
+	UPROPERTY(EditAnywhere, BlueprintReadWrite, Category=Transform, meta=( UIMin = "-89", ClampMin = "-89", UIMax = "89", ClampMax = "89", Delta = "1" ))
+	FVector2D Shear;
+	
+	/** The angle in degrees to rotate */
+	UPROPERTY(EditAnywhere, BlueprintReadWrite, Category=Transform, meta=( UIMin = "-180", UIMax = "180", Delta = "1" ))
+	float Angle;
+	
+	/** The amount to translate the widget in slate units */
+	UPROPERTY(EditAnywhere, BlueprintReadWrite, Category=Transform, meta=( Delta = "1" ))
+	FVector2D Translation;
+
+	FWidgetTransform()
+		: Scale(1, 1)
+		, Shear(0, 0)
+		, Angle(0)
+		, Translation(0, 0)
+	{
+	}
+};
