UMG - Adding some fixes for non-unity.

diff --git a/Engine/Source/Editor/UMGEditor/Private/Designer/SDesignerView.cpp b/Engine/Source/Editor/UMGEditor/Private/Designer/SDesignerView.cpp
--- a/Engine/Source/Editor/UMGEditor/Private/Designer/SDesignerView.cpp
+++ b/Engine/Source/Editor/UMGEditor/Private/Designer/SDesignerView.cpp
@@ -11,6 +11,7 @@
 #include "SZoomPan.h"
 #include "SDisappearingBar.h"
 #include "SDesignerToolBar.h"
+#include "DesignerCommands.h"
 #include "STransformHandle.h"
 
 #define LOCTEXT_NAMESPACE "UMG"
diff --git a/Engine/Source/Runtime/UMG/Private/Components/BorderSlot.cpp b/Engine/Source/Runtime/UMG/Private/Components/BorderSlot.cpp
--- a/Engine/Source/Runtime/UMG/Private/Components/BorderSlot.cpp
+++ b/Engine/Source/Runtime/UMG/Private/Components/BorderSlot.cpp
@@ -1,6 +1,7 @@
 // Copyright 1998-2014 Epic Games, Inc. All Rights Reserved.
 
 #include "UMGPrivatePCH.h"
+#include "ObjectEditorUtils.h"
 
 /////////////////////////////////////////////////////
 // UBorderSlot
