Fixed some localization data for Blueprint menu names with Keys being the same in Namespaces.

Correctly categorized some UMG menu items in the new BP menu system.

diff --git a/Engine/Source/Editor/UMGEditor/Private/Nodes/K2Node_CreateDragDropOperation.cpp b/Engine/Source/Editor/UMGEditor/Private/Nodes/K2Node_CreateDragDropOperation.cpp
--- a/Engine/Source/Editor/UMGEditor/Private/Nodes/K2Node_CreateDragDropOperation.cpp
+++ b/Engine/Source/Editor/UMGEditor/Private/Nodes/K2Node_CreateDragDropOperation.cpp
@@ -14,17 +14,17 @@
 UK2Node_CreateDragDropOperation::UK2Node_CreateDragDropOperation(const class FPostConstructInitializeProperties& PCIP)
 	: Super(PCIP)
 {
-	NodeTooltip = LOCTEXT("NodeTooltip", "Creates a new drag drop operation");
+	NodeTooltip = LOCTEXT("DragDropNodeTooltip", "Creates a new drag drop operation");
 }
 
 FText UK2Node_CreateDragDropOperation::GetBaseNodeTitle() const
 {
-	return LOCTEXT("CreateWidget_BaseTitle", "Create Drag & Drop Operation");
+	return LOCTEXT("CreateDragDropWidget_BaseTitle", "Create Drag & Drop Operation");
 }
 
 FText UK2Node_CreateDragDropOperation::GetNodeTitleFormat() const
 {
-	return LOCTEXT("CreateWidget", "Create {ClassName}");
+	return LOCTEXT("CreateDragDropWidget", "Create {ClassName}");
 }
 
 UClass* UK2Node_CreateDragDropOperation::GetClassPinBaseClass() const
@@ -32,6 +32,11 @@ UClass* UK2Node_CreateDragDropOperation::GetClassPinBaseClass() const
 	return UDragDropOperation::StaticClass();
 }
 
+FText UK2Node_CreateDragDropOperation::GetMenuCategory() const
+{
+	return FEditorCategoryUtils::GetCommonCategory(FCommonEditorCategory::UserInterface);
+}
+
 void UK2Node_CreateDragDropOperation::GetMenuEntries(FGraphContextMenuBuilder& ContextMenuBuilder) const
 {
 	UK2Node_CreateDragDropOperation* TemplateNode = NewObject<UK2Node_CreateDragDropOperation>(GetTransientPackage(), GetClass());
diff --git a/Engine/Source/Editor/UMGEditor/Private/Nodes/K2Node_CreateDragDropOperation.h b/Engine/Source/Editor/UMGEditor/Private/Nodes/K2Node_CreateDragDropOperation.h
--- a/Engine/Source/Editor/UMGEditor/Private/Nodes/K2Node_CreateDragDropOperation.h
+++ b/Engine/Source/Editor/UMGEditor/Private/Nodes/K2Node_CreateDragDropOperation.h
@@ -17,6 +17,7 @@ class UMGEDITOR_API UK2Node_CreateDragDropOperation : public UK2Node_ConstructOb
 
 	// Begin UK2Node interface
 	void GetMenuEntries(FGraphContextMenuBuilder& ContextMenuBuilder) const;
+	virtual FText GetMenuCategory() const override;
 	// End UK2Node interface.
 
 protected:
diff --git a/Engine/Source/Editor/UMGEditor/Private/Nodes/K2Node_CreateWidget.cpp b/Engine/Source/Editor/UMGEditor/Private/Nodes/K2Node_CreateWidget.cpp
--- a/Engine/Source/Editor/UMGEditor/Private/Nodes/K2Node_CreateWidget.cpp
+++ b/Engine/Source/Editor/UMGEditor/Private/Nodes/K2Node_CreateWidget.cpp
@@ -55,6 +55,11 @@ UClass* UK2Node_CreateWidget::GetClassPinBaseClass() const
 	return UUserWidget::StaticClass();
 }
 
+FText UK2Node_CreateWidget::GetMenuCategory() const
+{
+	return FEditorCategoryUtils::GetCommonCategory(FCommonEditorCategory::UserInterface);
+}
+
 void UK2Node_CreateWidget::GetMenuEntries(FGraphContextMenuBuilder& ContextMenuBuilder) const
 {
 	UK2Node_CreateWidget* TemplateNode = NewObject<UK2Node_CreateWidget>(GetTransientPackage(), GetClass());
diff --git a/Engine/Source/Editor/UMGEditor/Private/Nodes/K2Node_CreateWidget.h b/Engine/Source/Editor/UMGEditor/Private/Nodes/K2Node_CreateWidget.h
--- a/Engine/Source/Editor/UMGEditor/Private/Nodes/K2Node_CreateWidget.h
+++ b/Engine/Source/Editor/UMGEditor/Private/Nodes/K2Node_CreateWidget.h
@@ -18,6 +18,7 @@ class UMGEDITOR_API UK2Node_CreateWidget : public UK2Node_ConstructObjectFromCla
 
 	// Begin UK2Node interface
 	void GetMenuEntries(FGraphContextMenuBuilder& ContextMenuBuilder) const;
+	virtual FText GetMenuCategory() const override;
 	// End UK2Node interface.
 
 	/** Get the owning player pin */
diff --git a/Engine/Source/Editor/UnrealEd/Private/EditorCategoryUtils.cpp b/Engine/Source/Editor/UnrealEd/Private/EditorCategoryUtils.cpp
--- a/Engine/Source/Editor/UnrealEd/Private/EditorCategoryUtils.cpp
+++ b/Engine/Source/Editor/UnrealEd/Private/EditorCategoryUtils.cpp
@@ -83,6 +83,7 @@ FEditorCategoryUtilsImpl::FCategoryInfoMap& FEditorCategoryUtilsImpl::GetCategor
 		RegisterCategoryKey("Utilities", LOCTEXT("UtilitiesCategory", "Utilities"));
 		RegisterCategoryKey("Delegates", LOCTEXT("DelegatesCategory", "Event Dispatchers"));
 		RegisterCategoryKey("Variables", LOCTEXT("VariablesCategory", "Variables"));
+		RegisterCategoryKey("UserInterface", LOCTEXT("UserInterfaceCategory", "User Interface"));
 
 		// Utilities sub categories
 		RegisterCategoryKey("FlowControl", BuildCategoryString(FCommonEditorCategory::Utilities, LOCTEXT("FlowControlCategory", "Flow Control")));
@@ -193,6 +194,7 @@ FText const& FEditorCategoryUtils::GetCommonCategory(const FCommonEditorCategory
 		CommonCategoryKeys.Add(FCommonEditorCategory::Utilities, "Utilities");
 		CommonCategoryKeys.Add(FCommonEditorCategory::Delegates, "Delegates");
 		CommonCategoryKeys.Add(FCommonEditorCategory::Variables, "Variables");
+		CommonCategoryKeys.Add(FCommonEditorCategory::UserInterface, "UserInterface");
 
 		CommonCategoryKeys.Add(FCommonEditorCategory::FlowControl, "FlowControl");
 		CommonCategoryKeys.Add(FCommonEditorCategory::Transformation, "Transformation");
diff --git a/Engine/Source/Editor/UnrealEd/Public/EditorCategoryUtils.h b/Engine/Source/Editor/UnrealEd/Public/EditorCategoryUtils.h
--- a/Engine/Source/Editor/UnrealEd/Public/EditorCategoryUtils.h
+++ b/Engine/Source/Editor/UnrealEd/Public/EditorCategoryUtils.h
@@ -33,6 +33,7 @@ struct FCommonEditorCategory
 		Transformation,
 		Utilities,
 		FlowControl,
+		UserInterface,
 
 		// Type library categories:
 		String,
