UMG - Fixing the tooltip on the designer mode button in the editor.  Now propery reports as the design mode, and not the defaults mode.

diff --git a/Engine/Source/Editor/UMGEditor/Private/WidgetBlueprintEditorToolbar.cpp b/Engine/Source/Editor/UMGEditor/Private/WidgetBlueprintEditorToolbar.cpp
--- a/Engine/Source/Editor/UMGEditor/Private/WidgetBlueprintEditorToolbar.cpp
+++ b/Engine/Source/Editor/UMGEditor/Private/WidgetBlueprintEditorToolbar.cpp
@@ -77,13 +77,13 @@ void FWidgetBlueprintEditorToolbar::FillWidgetBlueprintEditorModesToolbar(FToolB
 			.OnGetActiveMode(GetActiveMode)
 			.OnSetActiveMode(SetActiveMode)
 			.ToolTip(IDocumentation::Get()->CreateToolTip(
-				LOCTEXT("BlueprintDefaultsModeButtonTooltip", "Switch to Blueprint Defaults Mode"),
+				LOCTEXT("DesignerModeButtonTooltip", "Switch to Blueprint Designer Mode"),
 				NULL,
 				TEXT("Shared/Editors/BlueprintEditor"),
-				TEXT("DefaultsMode")))
+				TEXT("DesignerMode")))
 			.IconImage(FEditorStyle::GetBrush("FullBlueprintEditor.SwitchToBlueprintDefaultsMode"))
 			.SmallIconImage(FEditorStyle::GetBrush("FullBlueprintEditor.SwitchToBlueprintDefaultsMode.Small"))
-			.AddMetaData<FTagMetaData>(FTagMetaData(TEXT("DefaultsMode")))
+			.AddMetaData<FTagMetaData>(FTagMetaData(TEXT("DesignerMode")))
 		);
 
 		BlueprintEditorPtr->AddToolbarWidget(SNew(SBlueprintModeSeparator));
