UMG - Removing old comment code.

diff --git a/Engine/Source/Editor/UMGEditor/Private/BlueprintModes/WidgetDesignerApplicationMode.cpp b/Engine/Source/Editor/UMGEditor/Private/BlueprintModes/WidgetDesignerApplicationMode.cpp
--- a/Engine/Source/Editor/UMGEditor/Private/BlueprintModes/WidgetDesignerApplicationMode.cpp
+++ b/Engine/Source/Editor/UMGEditor/Private/BlueprintModes/WidgetDesignerApplicationMode.cpp
@@ -135,17 +135,9 @@ void FWidgetDesignerApplicationMode::RegisterTabFactories(TSharedPtr<FTabManager
 void FWidgetDesignerApplicationMode::PreDeactivateMode()
 {
 	//FWidgetBlueprintApplicationMode::PreDeactivateMode();
-
-	//GetBlueprintEditor()->GetInspector()->EnableComponentDetailsCustomization(false);
-
-	//static FName PropertyEditor("PropertyEditor");
-	//FPropertyEditorModule& PropertyModule = FModuleManager::GetModuleChecked<FPropertyEditorModule>(PropertyEditor);
-	//PropertyModule.UnregisterCustomPropertyTypeLayout(TEXT("CanvasPanelSlot"));
 }
 
 void FWidgetDesignerApplicationMode::PostActivateMode()
 {
 	//FWidgetBlueprintApplicationMode::PostActivateMode();
-
-	//GetBlueprintEditor()->GetInspector()->EnableComponentDetailsCustomization(false);
 }
