UMG - Now changing the text on the offsets property to match what they represent.

diff --git a/Engine/Source/Editor/UMGEditor/Private/Customizations/CanvasSlotCustomization.cpp b/Engine/Source/Editor/UMGEditor/Private/Customizations/CanvasSlotCustomization.cpp
--- a/Engine/Source/Editor/UMGEditor/Private/Customizations/CanvasSlotCustomization.cpp
+++ b/Engine/Source/Editor/UMGEditor/Private/Customizations/CanvasSlotCustomization.cpp
@@ -10,15 +10,16 @@
 
 #define LOCTEXT_NAMESPACE "UMG"
 
+// FCanvasSlotCustomization
+////////////////////////////////////////////////////////////////////////////////
+
 void FCanvasSlotCustomization::CustomizeHeader(TSharedRef<IPropertyHandle> PropertyHandle, FDetailWidgetRow& HeaderRow, IPropertyTypeCustomizationUtils& CustomizationUtils)
 {
-	// Build the layout category
-	//IDetailCategoryBuilder& LayoutCategory = DetailLayout.EditCategory("Layout");
 }
 
 void FCanvasSlotCustomization::CustomizeChildren(TSharedRef<IPropertyHandle> PropertyHandle, IDetailChildrenBuilder& ChildBuilder, IPropertyTypeCustomizationUtils& CustomizationUtils)
 {
-	CustomizeAnchors(PropertyHandle, ChildBuilder, CustomizationUtils);
+	CustomizeLayoutData(PropertyHandle, ChildBuilder, CustomizationUtils);
 
 	FillOutChildren(PropertyHandle, ChildBuilder, CustomizationUtils);
 }
@@ -48,23 +49,101 @@ void FCanvasSlotCustomization::FillOutChildren(TSharedRef<IPropertyHandle> Prope
 	}
 }
 
-void FCanvasSlotCustomization::CustomizeAnchors(TSharedRef<IPropertyHandle> PropertyHandle, IDetailChildrenBuilder& ChildBuilder, IPropertyTypeCustomizationUtils& CustomizationUtils)
+void FCanvasSlotCustomization::CustomizeLayoutData(TSharedRef<IPropertyHandle> PropertyHandle, IDetailChildrenBuilder& ChildBuilder, IPropertyTypeCustomizationUtils& CustomizationUtils)
+{
+	TSharedPtr<IPropertyHandle> LayoutData = PropertyHandle->GetChildHandle(GET_MEMBER_NAME_CHECKED(UCanvasPanelSlot, LayoutData));
+	if ( LayoutData.IsValid() )
+	{
+		LayoutData->MarkHiddenByCustomization();
+
+		CustomizeAnchors(LayoutData, ChildBuilder, CustomizationUtils);
+		CustomizeOffsets(LayoutData, ChildBuilder, CustomizationUtils);
+
+		TSharedPtr<IPropertyHandle> AlignmentHandle = LayoutData->GetChildHandle(GET_MEMBER_NAME_CHECKED(FAnchorData, Alignment));
+		AlignmentHandle->MarkHiddenByCustomization();
+		ChildBuilder.AddChildProperty(AlignmentHandle.ToSharedRef());
+	}
+}
+
+void FCanvasSlotCustomization::CustomizeOffsets(TSharedPtr<IPropertyHandle> PropertyHandle, IDetailChildrenBuilder& ChildBuilder, IPropertyTypeCustomizationUtils& CustomizationUtils)
+{
+	TSharedPtr<IPropertyHandle> OffsetsHandle = PropertyHandle->GetChildHandle(GET_MEMBER_NAME_CHECKED(FAnchorData, Offsets));
+	TSharedPtr<IPropertyHandle> AnchorsHandle = PropertyHandle->GetChildHandle(GET_MEMBER_NAME_CHECKED(FAnchorData, Anchors));
+
+	OffsetsHandle->MarkHiddenByCustomization();
+
+	TSharedPtr<IPropertyHandle> LeftHandle = OffsetsHandle->GetChildHandle(GET_MEMBER_NAME_CHECKED(FMargin, Left));
+	TSharedPtr<IPropertyHandle> TopHandle = OffsetsHandle->GetChildHandle(GET_MEMBER_NAME_CHECKED(FMargin, Top));
+	TSharedPtr<IPropertyHandle> RightHandle = OffsetsHandle->GetChildHandle(GET_MEMBER_NAME_CHECKED(FMargin, Right));
+	TSharedPtr<IPropertyHandle> BottomHandle = OffsetsHandle->GetChildHandle(GET_MEMBER_NAME_CHECKED(FMargin, Bottom));
+
+	IDetailPropertyRow& LeftRow = ChildBuilder.AddChildProperty(LeftHandle.ToSharedRef());
+	IDetailPropertyRow& TopRow = ChildBuilder.AddChildProperty(TopHandle.ToSharedRef());
+	IDetailPropertyRow& RightRow = ChildBuilder.AddChildProperty(RightHandle.ToSharedRef());
+	IDetailPropertyRow& BottomRow = ChildBuilder.AddChildProperty(BottomHandle.ToSharedRef());
+
+	TAttribute<FText> LeftLabel = TAttribute<FText>::Create(TAttribute<FText>::FGetter::CreateStatic(&FCanvasSlotCustomization::GetOffsetLabel, PropertyHandle, Orient_Horizontal, LOCTEXT("PositionX", "Position X"), LOCTEXT("OffsetLeft", "Offset Left")));
+	TAttribute<FText> TopLabel = TAttribute<FText>::Create(TAttribute<FText>::FGetter::CreateStatic(&FCanvasSlotCustomization::GetOffsetLabel, PropertyHandle, Orient_Vertical, LOCTEXT("PositionY", "Position Y"), LOCTEXT("OffsetTop", "Offset Top")));
+	TAttribute<FText> RightLabel = TAttribute<FText>::Create(TAttribute<FText>::FGetter::CreateStatic(&FCanvasSlotCustomization::GetOffsetLabel, PropertyHandle, Orient_Horizontal, LOCTEXT("SizeX", "Size X"), LOCTEXT("OffsetRight", "Offset Right")));
+	TAttribute<FText> BottomLabel = TAttribute<FText>::Create(TAttribute<FText>::FGetter::CreateStatic(&FCanvasSlotCustomization::GetOffsetLabel, PropertyHandle, Orient_Vertical, LOCTEXT("SizeY", "Size Y"), LOCTEXT("OffsetBottom", "Offset Bottom")));
+
+	CreateEditorWithDynamicLabel(LeftRow, LeftLabel);
+	CreateEditorWithDynamicLabel(TopRow, TopLabel);
+	CreateEditorWithDynamicLabel(RightRow, RightLabel);
+	CreateEditorWithDynamicLabel(BottomRow, BottomLabel);
+}
+
+void FCanvasSlotCustomization::CreateEditorWithDynamicLabel(IDetailPropertyRow& PropertyRow, TAttribute<FText> TextAttribute)
 {
-	TSharedPtr<IPropertyHandle> Slot = PropertyHandle->GetChildHandle(GET_MEMBER_NAME_CHECKED(FAnchorData, Anchors));
+	TSharedPtr<SWidget> NameWidget;
+	TSharedPtr<SWidget> ValueWidget;
+	FDetailWidgetRow Row;
+	PropertyRow.GetDefaultWidgets(NameWidget, ValueWidget, Row);
+
+	PropertyRow.CustomWidget(/*bShowChildren*/ true)
+	.NameContent()
+	[
+		SNew(STextBlock)
+		.Font(IDetailLayoutBuilder::GetDetailFont())
+		.Text(TextAttribute)
+	]
+	.ValueContent()
+	[
+		ValueWidget.ToSharedRef()
+	];
+}
 
-	if ( !Slot.IsValid() )
+FText FCanvasSlotCustomization::GetOffsetLabel(TSharedPtr<IPropertyHandle> PropertyHandle, EOrientation Orientation, FText NonStretchingLabel, FText StretchingLabel)
+{
+	TArray<UObject*> Objects;
+	PropertyHandle->GetOuterObjects(Objects);
+	
+	if ( Objects.Num() > 1 )
 	{
-		return;
+		return StretchingLabel;
 	}
 
-	Slot->MarkHiddenByCustomization();
+	TArray<void*> RawData;
+	PropertyHandle->AccessRawData(RawData);
+
+	FAnchorData* AnchorData = reinterpret_cast<FAnchorData*>( RawData[0] );
+
+	const bool bStretching =
+		( Orientation == Orient_Horizontal && AnchorData->Anchors.IsStretchedHorizontal() ) ||
+		( Orientation == Orient_Vertical && AnchorData->Anchors.IsStretchedVertical() );
+
+	return bStretching ? StretchingLabel : NonStretchingLabel;
+}
+
+void FCanvasSlotCustomization::CustomizeAnchors(TSharedPtr<IPropertyHandle> PropertyHandle, IDetailChildrenBuilder& ChildBuilder, IPropertyTypeCustomizationUtils& CustomizationUtils)
+{
+	TSharedPtr<IPropertyHandle> AnchorsHandle = PropertyHandle->GetChildHandle(GET_MEMBER_NAME_CHECKED(FAnchorData, Anchors));
 
-	TSharedRef<IPropertyHandle> AnchorsHandle = Slot.ToSharedRef();
+	AnchorsHandle->MarkHiddenByCustomization();
 
-	IDetailPropertyRow& AnchorsPropertyRow = ChildBuilder.AddChildProperty(AnchorsHandle);
+	IDetailPropertyRow& AnchorsPropertyRow = ChildBuilder.AddChildProperty(AnchorsHandle.ToSharedRef());
 
-	const bool bShowChildren = true;
-	AnchorsPropertyRow.CustomWidget(bShowChildren)
+	AnchorsPropertyRow.CustomWidget(/*bShowChildren*/ true)
 		.NameContent()
 		[
 			SNew(STextBlock)
@@ -256,7 +335,7 @@ void FCanvasSlotCustomization::CustomizeAnchors(TSharedRef<IPropertyHandle> Prop
 		];
 }
 
-FReply FCanvasSlotCustomization::OnAnchorClicked(TSharedRef<IPropertyHandle> AnchorsHandle, FAnchors Anchors)
+FReply FCanvasSlotCustomization::OnAnchorClicked(TSharedPtr<IPropertyHandle> AnchorsHandle, FAnchors Anchors)
 {
 	//TODO UMG Group under single undo transaction.
 
diff --git a/Engine/Source/Editor/UMGEditor/Private/Customizations/CanvasSlotCustomization.h b/Engine/Source/Editor/UMGEditor/Private/Customizations/CanvasSlotCustomization.h
--- a/Engine/Source/Editor/UMGEditor/Private/Customizations/CanvasSlotCustomization.h
+++ b/Engine/Source/Editor/UMGEditor/Private/Customizations/CanvasSlotCustomization.h
@@ -22,10 +22,17 @@ public:
 	
 private:
 
-	void CustomizeAnchors(TSharedRef<IPropertyHandle> PropertyHandle, class IDetailChildrenBuilder& ChildBuilder, IPropertyTypeCustomizationUtils& CustomizationUtils);
+	void CustomizeLayoutData(TSharedRef<IPropertyHandle> PropertyHandle, IDetailChildrenBuilder& ChildBuilder, IPropertyTypeCustomizationUtils& CustomizationUtils);
+	void CustomizeAnchors(TSharedPtr<IPropertyHandle> PropertyHandle, class IDetailChildrenBuilder& ChildBuilder, IPropertyTypeCustomizationUtils& CustomizationUtils);
+	void CustomizeOffsets(TSharedPtr<IPropertyHandle> PropertyHandle, class IDetailChildrenBuilder& ChildBuilder, IPropertyTypeCustomizationUtils& CustomizationUtils);
+	
 	void FillOutChildren(TSharedRef<IPropertyHandle> PropertyHandle, class IDetailChildrenBuilder& ChildBuilder, IPropertyTypeCustomizationUtils& CustomizationUtils);
+
+	void CreateEditorWithDynamicLabel(IDetailPropertyRow& PropertyRow, TAttribute<FText> TextAttribute);
+
+	static FText GetOffsetLabel(TSharedPtr<IPropertyHandle> AnchorStructureHandle, EOrientation Orientation, FText NonStretchingLabel, FText StretchingLabel);
 	
-	FReply OnAnchorClicked(TSharedRef<IPropertyHandle> AnchorsHandle, FAnchors Anchors);
+	FReply OnAnchorClicked(TSharedPtr<IPropertyHandle> AnchorsHandle, FAnchors Anchors);
 
 private:
 	UWidgetBlueprint* Blueprint;
