UMG - Now hiding experimental widgets in the palette.  Moving the developement status querying logic into a common area.

diff --git a/Engine/Source/Editor/UMGEditor/Private/Palette/SPaletteView.cpp b/Engine/Source/Editor/UMGEditor/Private/Palette/SPaletteView.cpp
--- a/Engine/Source/Editor/UMGEditor/Private/Palette/SPaletteView.cpp
+++ b/Engine/Source/Editor/UMGEditor/Private/Palette/SPaletteView.cpp
@@ -265,12 +265,15 @@ UWidgetBlueprint* SPaletteView::GetBlueprint() const
 
 void SPaletteView::BuildWidgetList()
 {
+	// Clear the current list of view models and categories
 	WidgetViewModels.Reset();
 	WidgetTemplateCategories.Reset();
 
+	// Generate a list of templates
 	BuildClassWidgetList();
 	BuildSpecialWidgetList();
 
+	// For each entry in the category create a view model for the widget template
 	for ( auto& Entry : WidgetTemplateCategories )
 	{
 		TSharedPtr<FWidgetHeaderViewModel> Header = MakeShareable(new FWidgetHeaderViewModel());
@@ -289,21 +292,32 @@ void SPaletteView::BuildWidgetList()
 		WidgetViewModels.Add(Header);
 	}
 
+	// Sort the view models by name
 	WidgetViewModels.Sort([] (TSharedPtr<FWidgetViewModel> L, TSharedPtr<FWidgetViewModel> R) { return R->GetName().CompareTo(L->GetName()) > 0; });
 }
 
 void SPaletteView::BuildClassWidgetList()
 {
+	static const FName DevelopmentStatusKey(TEXT("DevelopmentStatus"));
+
 	for ( TObjectIterator<UClass> ClassIt; ClassIt; ++ClassIt )
 	{
 		UClass* WidgetClass = *ClassIt;
 		if ( WidgetClass->IsChildOf(UWidget::StaticClass()) )
 		{
+			// We don't want to include experimental widget classes, so filter those out.
+			bool bIsExperimental, bIsEarlyAccess;
+			FObjectEditorUtils::GetClassDevelopmentStatus(WidgetClass, bIsExperimental, bIsEarlyAccess);
+
+			// Don't include skeleton classes of widgets.
 			const bool bIsSkeletonClass = WidgetClass->HasAnyFlags(RF_Transient) && WidgetClass->HasAnyClassFlags(CLASS_CompiledFromBlueprint);
 			const bool bIsValidClass =
-				!( WidgetClass->HasAnyClassFlags(CLASS_Abstract | CLASS_Deprecated | CLASS_NewerVersionExists) || bIsSkeletonClass );
+				!( WidgetClass->HasAnyClassFlags(CLASS_Abstract | CLASS_Deprecated | CLASS_NewerVersionExists) || bIsSkeletonClass || bIsExperimental );
 			const bool bIsSameClass = WidgetClass->GetFName() == GetBlueprint()->GeneratedClass->GetFName();
 
+			//TODO UMG does not prevent deep nested circular references
+
+			// Only use non-(abstract/deprecated/old) classes and don't include the current widget blueprint class.
 			if ( bIsValidClass && !bIsSameClass )
 			{
 				TSharedPtr<FWidgetTemplateClass> Template = MakeShareable(new FWidgetTemplateClass(WidgetClass));
diff --git a/Engine/Source/Editor/UMGEditor/Private/UMGEditorPrivatePCH.h b/Engine/Source/Editor/UMGEditor/Private/UMGEditorPrivatePCH.h
--- a/Engine/Source/Editor/UMGEditor/Private/UMGEditorPrivatePCH.h
+++ b/Engine/Source/Editor/UMGEditor/Private/UMGEditorPrivatePCH.h
@@ -15,6 +15,8 @@
 #include "ScopedTransaction.h"
 #include "TextFilter.h"
 
+#include "ObjectEditorUtils.h"
+
 #include "UMG.h"
 
 #include "UMGEditorClasses.h"
diff --git a/Engine/Source/Editor/UnrealEd/Private/AssetSelection.cpp b/Engine/Source/Editor/UnrealEd/Private/AssetSelection.cpp
--- a/Engine/Source/Editor/UnrealEd/Private/AssetSelection.cpp
+++ b/Engine/Source/Editor/UnrealEd/Private/AssetSelection.cpp
@@ -24,13 +24,11 @@
 #include "Editor/ActorPositioning.h"
 #include "Animation/SkeletalMeshActor.h"
 
+#include "ObjectEditorUtils.h"
+
 
 namespace AssetSelectionUtils
 {
-	const FName DevelopmentStatusKey(TEXT("DevelopmentStatus"));
-	const FString EarlyAccessValue(TEXT("EarlyAccess"));
-	const FString ExperimentalValue(TEXT("Experimental"));
-
 	bool IsClassPlaceable(const UClass* Class)
 	{
 		const bool bIsAddable =
@@ -205,23 +203,21 @@ namespace AssetSelectionUtils
 							}
 						}
 
-						// Check for experimental classes in the component hierarchy
-						FString DevelopmentStatus;
-						if (Component->GetClass()->GetStringMetaDataHierarchical(DevelopmentStatusKey, /*out*/ &DevelopmentStatus))
-						{
-							ActorInfo.bHaveExperimentalClass = ActorInfo.bHaveExperimentalClass || (DevelopmentStatus == ExperimentalValue);
-							ActorInfo.bHaveEarlyAccessClass = ActorInfo.bHaveEarlyAccessClass || (DevelopmentStatus == EarlyAccessValue);
-						}
+						// Check for experimental/early-access classes in the component hierarchy
+						bool bIsExperimental, bIsEarlyAccess;
+						FObjectEditorUtils::GetClassDevelopmentStatus(Component->GetClass(), bIsExperimental, bIsEarlyAccess);
+
+						ActorInfo.bHaveExperimentalClass |= bIsExperimental;
+						ActorInfo.bHaveEarlyAccessClass |= bIsEarlyAccess;
 					}
 
-					// Check for experimental classes in the actor hierarchy
+					// Check for experimental/early-access classes in the actor hierarchy
 					{
-						FString DevelopmentStatus;
-						if (CurrentClass->GetStringMetaDataHierarchical(DevelopmentStatusKey, /*out*/ &DevelopmentStatus))
-						{
-							ActorInfo.bHaveExperimentalClass = ActorInfo.bHaveExperimentalClass || (DevelopmentStatus == ExperimentalValue);
-							ActorInfo.bHaveEarlyAccessClass = ActorInfo.bHaveEarlyAccessClass || (DevelopmentStatus == EarlyAccessValue);
-						}
+						bool bIsExperimental, bIsEarlyAccess;
+						FObjectEditorUtils::GetClassDevelopmentStatus(CurrentClass, bIsExperimental, bIsEarlyAccess);
+
+						ActorInfo.bHaveExperimentalClass |= bIsExperimental;
+						ActorInfo.bHaveEarlyAccessClass |= bIsEarlyAccess;
 					}
 
 					if( CurrentActor->IsA( ALight::StaticClass() ) )
diff --git a/Engine/Source/Runtime/Engine/Private/ObjectEditorUtils.cpp b/Engine/Source/Runtime/Engine/Private/ObjectEditorUtils.cpp
--- a/Engine/Source/Runtime/Engine/Private/ObjectEditorUtils.cpp
+++ b/Engine/Source/Runtime/Engine/Private/ObjectEditorUtils.cpp
@@ -60,6 +60,22 @@ namespace FObjectEditorUtils
 		return bResult;
 	}
 
+	void GetClassDevelopmentStatus(UClass* Class, bool& bIsExperimental, bool& bIsEarlyAccess)
+	{
+		static const FName DevelopmentStatusKey(TEXT("DevelopmentStatus"));
+		static const FString EarlyAccessValue(TEXT("EarlyAccess"));
+		static const FString ExperimentalValue(TEXT("Experimental"));
+
+		bIsExperimental = bIsEarlyAccess = false;
+
+		FString DevelopmentStatus;
+		if ( Class->GetStringMetaDataHierarchical(DevelopmentStatusKey, /*out*/ &DevelopmentStatus) )
+		{
+			bIsExperimental = DevelopmentStatus == ExperimentalValue;
+			bIsEarlyAccess = DevelopmentStatus == EarlyAccessValue;
+		}
+	}
+
 	bool MigratePropertyValue(UObject* SourceObject, UProperty* SourceProperty, UObject* DestinationObject, UProperty* DestinationProperty)
 	{
 		FString SourceValue;
diff --git a/Engine/Source/Runtime/Engine/Public/ObjectEditorUtils.h b/Engine/Source/Runtime/Engine/Public/ObjectEditorUtils.h
--- a/Engine/Source/Runtime/Engine/Public/ObjectEditorUtils.h
+++ b/Engine/Source/Runtime/Engine/Public/ObjectEditorUtils.h
@@ -45,6 +45,15 @@ namespace FObjectEditorUtils
 	 */
 	ENGINE_API bool IsVariableCategoryHiddenFromClass( const UProperty* InVariable, const UClass* Class );
 
+	/**
+	 * Get the classes development status and return if it's either experimental or early access.
+	 * 
+	 * @param Class the class to inspect.
+	 * @param bIsExperimental [out] value indicating if the class is experimental.
+	 * @param bIsEarlyAccess [out] value indicating if the class is early access.
+	 */
+	ENGINE_API void GetClassDevelopmentStatus(UClass* Class, bool& bIsExperimental, bool& bIsEarlyAccess);
+
 	/**
 	 * Copy the value of a property from source object to a destination object.
 	 *
