UMG - Adding more checks to prevent crashes when the slot is not a canvas panel slot.

diff --git a/Engine/Source/Editor/UMGEditor/Private/Extensions/CanvasSlotExtension.cpp b/Engine/Source/Editor/UMGEditor/Private/Extensions/CanvasSlotExtension.cpp
--- a/Engine/Source/Editor/UMGEditor/Private/Extensions/CanvasSlotExtension.cpp
+++ b/Engine/Source/Editor/UMGEditor/Private/Extensions/CanvasSlotExtension.cpp
@@ -135,20 +135,22 @@ EVisibility FCanvasSlotExtension::GetAnchorVisibility(EAnchorWidget::Type Anchor
 		UWidget* PreviewWidget = Selection.GetPreview();
 		if ( PreviewWidget && PreviewWidget->Slot )
 		{
-			UCanvasPanelSlot* PreviewCanvasSlot = CastChecked<UCanvasPanelSlot>(PreviewWidget->Slot);
-			switch ( AnchorType )
+			if ( UCanvasPanelSlot* PreviewCanvasSlot = Cast<UCanvasPanelSlot>(PreviewWidget->Slot) )
 			{
-			case EAnchorWidget::Center:
-				return PreviewCanvasSlot->LayoutData.Anchors.Minimum == PreviewCanvasSlot->LayoutData.Anchors.Maximum ? EVisibility::Visible : EVisibility::Collapsed;
-			case EAnchorWidget::Left:
-			case EAnchorWidget::Right:
-				return PreviewCanvasSlot->LayoutData.Anchors.Minimum.Y == PreviewCanvasSlot->LayoutData.Anchors.Maximum.Y ? EVisibility::Visible : EVisibility::Collapsed;
-			case EAnchorWidget::Top:
-			case EAnchorWidget::Bottom:
-				return PreviewCanvasSlot->LayoutData.Anchors.Minimum.X == PreviewCanvasSlot->LayoutData.Anchors.Maximum.X ? EVisibility::Visible : EVisibility::Collapsed;
-			}
+				switch ( AnchorType )
+				{
+				case EAnchorWidget::Center:
+					return PreviewCanvasSlot->LayoutData.Anchors.Minimum == PreviewCanvasSlot->LayoutData.Anchors.Maximum ? EVisibility::Visible : EVisibility::Collapsed;
+				case EAnchorWidget::Left:
+				case EAnchorWidget::Right:
+					return PreviewCanvasSlot->LayoutData.Anchors.Minimum.Y == PreviewCanvasSlot->LayoutData.Anchors.Maximum.Y ? EVisibility::Visible : EVisibility::Collapsed;
+				case EAnchorWidget::Top:
+				case EAnchorWidget::Bottom:
+					return PreviewCanvasSlot->LayoutData.Anchors.Minimum.X == PreviewCanvasSlot->LayoutData.Anchors.Maximum.X ? EVisibility::Visible : EVisibility::Collapsed;
+				}
 
-			return EVisibility::Visible;
+				return EVisibility::Visible;
+			}
 		}
 	}
 
@@ -162,26 +164,27 @@ FVector2D FCanvasSlotExtension::GetAnchorAlignment(EAnchorWidget::Type AnchorTyp
 		UWidget* PreviewWidget = Selection.GetPreview();
 		if ( PreviewWidget && PreviewWidget->Slot )
 		{
-			UCanvasPanelSlot* PreviewCanvasSlot = CastChecked<UCanvasPanelSlot>(PreviewWidget->Slot);
-
-			FVector2D Minimum = PreviewCanvasSlot->LayoutData.Anchors.Minimum;
-			FVector2D Maximum = PreviewCanvasSlot->LayoutData.Anchors.Maximum;
-
-			switch ( AnchorType )
+			if ( UCanvasPanelSlot* PreviewCanvasSlot = Cast<UCanvasPanelSlot>(PreviewWidget->Slot) )
 			{
-			case EAnchorWidget::Center:
-			case EAnchorWidget::Left:
-			case EAnchorWidget::Top:
-			case EAnchorWidget::TopLeft:
-				return Minimum;
-			case EAnchorWidget::Right:
-			case EAnchorWidget::Bottom:
-			case EAnchorWidget::BottomRight:
-				return Maximum;
-			case EAnchorWidget::TopRight:
-				return  FVector2D(Maximum.X, Minimum.Y);
-			case EAnchorWidget::BottomLeft:
-				return  FVector2D(Minimum.X, Maximum.Y);
+				FVector2D Minimum = PreviewCanvasSlot->LayoutData.Anchors.Minimum;
+				FVector2D Maximum = PreviewCanvasSlot->LayoutData.Anchors.Maximum;
+
+				switch ( AnchorType )
+				{
+				case EAnchorWidget::Center:
+				case EAnchorWidget::Left:
+				case EAnchorWidget::Top:
+				case EAnchorWidget::TopLeft:
+					return Minimum;
+				case EAnchorWidget::Right:
+				case EAnchorWidget::Bottom:
+				case EAnchorWidget::BottomRight:
+					return Maximum;
+				case EAnchorWidget::TopRight:
+					return  FVector2D(Maximum.X, Minimum.Y);
+				case EAnchorWidget::BottomLeft:
+					return  FVector2D(Minimum.X, Maximum.Y);
+				}
 			}
 		}
 	}
