UMG - Dynamically changing the category name for slots to show the name of the class they are in, changing priority.  Fixing a bug with "Create Binding", wasn't working, it was deselecting the object in the details panel before it could complete the operation.

diff --git a/Engine/Source/Editor/UMGEditor/Private/Customizations/DetailCustomizations.cpp b/Engine/Source/Editor/UMGEditor/Private/Customizations/DetailCustomizations.cpp
--- a/Engine/Source/Editor/UMGEditor/Private/Customizations/DetailCustomizations.cpp
+++ b/Engine/Source/Editor/UMGEditor/Private/Customizations/DetailCustomizations.cpp
@@ -319,6 +319,30 @@ FReply FBlueprintWidgetCustomization::AddOrViewEventBinding(TSharedPtr<FEdGraphS
 
 void FBlueprintWidgetCustomization::CustomizeDetails( IDetailLayoutBuilder& DetailLayout )
 {
+	static const FName LayoutCategoryKey(TEXT("Layout"));
+
+	TArray< TWeakObjectPtr<UObject> > OutObjects;
+	DetailLayout.GetObjectsBeingCustomized(OutObjects);
+	
+	if ( OutObjects.Num() == 1 )
+	{
+		if ( UWidget* Widget = Cast<UWidget>(OutObjects[0].Get()) )
+		{
+			if ( Widget->Slot )
+			{
+				UClass* SlotClass = Widget->Slot->GetClass();
+				FString LayoutCatName = FString("Layout (") + SlotClass->GetDisplayNameText().ToString() + FString(")");
+
+				DetailLayout.EditCategory(LayoutCategoryKey, LayoutCatName, ECategoryPriority::TypeSpecific);
+			}
+			else
+			{
+				auto& Category = DetailLayout.EditCategory(LayoutCategoryKey);
+				//TODO UMG Hide Category
+			}
+		}
+	}
+
 	PerformBindingCustomization(DetailLayout);
 }
 
@@ -700,9 +724,6 @@ void FBlueprintWidgetCustomization::HandleCreateAndAddBinding(UWidget* Widget, T
 	FString FunctionName = Pre + WidgetName + Post;
 	UEdGraph* FunctionGraph = FBlueprintEditorUtils::CreateNewGraph(Blueprint, FBlueprintEditorUtils::FindUniqueKismetName(Blueprint, FunctionName), UEdGraph::StaticClass(), UEdGraphSchema_K2::StaticClass());
 
-	bool bUserCreated = true;
-	FBlueprintEditorUtils::AddFunctionGraph(Blueprint, FunctionGraph, bUserCreated, DelegateSignature);
-
 	// Only mark bindings as pure that need to be.
 	if ( bIsPure )
 	{
@@ -717,6 +738,9 @@ void FBlueprintWidgetCustomization::HandleCreateAndAddBinding(UWidget* Widget, T
 
 	HandleAddFunctionBinding(PropertyHandle, SelectedFunction);
 
+	bool bUserCreated = true;
+	FBlueprintEditorUtils::AddFunctionGraph(Blueprint, FunctionGraph, bUserCreated, DelegateSignature);
+
 	GotoFunction(FunctionGraph);
 }
 
