UMG - No longer using the fullscreen wrapper object when hosting the widget in the designer.

diff --git a/Engine/Source/Editor/UMGEditor/Private/Designer/SDesignerView.cpp b/Engine/Source/Editor/UMGEditor/Private/Designer/SDesignerView.cpp
--- a/Engine/Source/Editor/UMGEditor/Private/Designer/SDesignerView.cpp
+++ b/Engine/Source/Editor/UMGEditor/Private/Designer/SDesignerView.cpp
@@ -910,17 +910,11 @@ void SDesignerView::UpdatePreviewWidget()
 		PreviewWidget = LatestPreviewWidget;
 		if ( PreviewWidget )
 		{
-			const bool bModal = false;
-			const bool bShowCursor = false;
-			TSharedPtr<SWidget> OutUserWidget;
-			TSharedRef<SWidget> CurrentWidget = PreviewWidget->MakeViewportWidget(bModal, bShowCursor, OutUserWidget);
-			CurrentWidget->SlatePrepass();
+			TSharedRef<SWidget> NewPreviewSlateWidget = PreviewWidget->TakeWidget();
+			NewPreviewSlateWidget->SlatePrepass();
 
-			if ( CurrentWidget != PreviewSlateWidget.Pin() )
-			{
-				PreviewSlateWidget = CurrentWidget;
-				PreviewSurface->SetContent(CurrentWidget);
-			}
+			PreviewSlateWidget = NewPreviewSlateWidget;
+			PreviewSurface->SetContent(NewPreviewSlateWidget);
 
 			// Notify all selected widgets that they are selected, because there are new preview objects
 			// state may have been lost so this will recreate it if the widget does something special when
