UMG - Correctly registering the right customization with UMG.  Previous we were sharing customizations across UMG editors.

diff --git a/Engine/Source/Editor/UMGEditor/Private/Details/SWidgetDetailsView.cpp b/Engine/Source/Editor/UMGEditor/Private/Details/SWidgetDetailsView.cpp
--- a/Engine/Source/Editor/UMGEditor/Private/Details/SWidgetDetailsView.cpp
+++ b/Engine/Source/Editor/UMGEditor/Private/Details/SWidgetDetailsView.cpp
@@ -105,6 +105,15 @@ SWidgetDetailsView::~SWidgetDetailsView()
 		BlueprintEditor.Pin()->OnSelectedWidgetsChanging.RemoveAll(this);
 		BlueprintEditor.Pin()->OnSelectedWidgetsChanged.RemoveAll(this);
 	}
+
+	// Unregister the property type layouts
+	static FName PropertyEditor("PropertyEditor");
+	FPropertyEditorModule& PropertyModule = FModuleManager::GetModuleChecked<FPropertyEditorModule>(PropertyEditor);
+
+	PropertyModule.UnregisterCustomPropertyTypeLayout(TEXT("PanelSlot"), nullptr, PropertyView);
+	PropertyModule.UnregisterCustomPropertyTypeLayout(TEXT("EHorizontalAlignment"), nullptr, PropertyView);
+	PropertyModule.UnregisterCustomPropertyTypeLayout(TEXT("EVerticalAlignment"), nullptr, PropertyView);
+	PropertyModule.UnregisterCustomPropertyTypeLayout(TEXT("SlateChildSize"), nullptr, PropertyView);
 }
 
 void SWidgetDetailsView::RegisterCustomizations()
@@ -116,10 +125,10 @@ void SWidgetDetailsView::RegisterCustomizations()
 	static FName PropertyEditor("PropertyEditor");
 	FPropertyEditorModule& PropertyModule = FModuleManager::GetModuleChecked<FPropertyEditorModule>(PropertyEditor);
 
-	PropertyModule.RegisterCustomPropertyTypeLayout(TEXT("PanelSlot"), FOnGetPropertyTypeCustomizationInstance::CreateStatic(&FCanvasSlotCustomization::MakeInstance, BlueprintEditor.Pin()->GetBlueprintObj()));
-	PropertyModule.RegisterCustomPropertyTypeLayout(TEXT("EHorizontalAlignment"), FOnGetPropertyTypeCustomizationInstance::CreateStatic(&FHorizontalAlignmentCustomization::MakeInstance));
-	PropertyModule.RegisterCustomPropertyTypeLayout(TEXT("EVerticalAlignment"), FOnGetPropertyTypeCustomizationInstance::CreateStatic(&FVerticalAlignmentCustomization::MakeInstance));
-	PropertyModule.RegisterCustomPropertyTypeLayout(TEXT("SlateChildSize"), FOnGetPropertyTypeCustomizationInstance::CreateStatic(&FSlateChildSizeCustomization::MakeInstance));
+	PropertyModule.RegisterCustomPropertyTypeLayout(TEXT("PanelSlot"), FOnGetPropertyTypeCustomizationInstance::CreateStatic(&FCanvasSlotCustomization::MakeInstance, BlueprintEditor.Pin()->GetBlueprintObj()), nullptr, PropertyView);
+	PropertyModule.RegisterCustomPropertyTypeLayout(TEXT("EHorizontalAlignment"), FOnGetPropertyTypeCustomizationInstance::CreateStatic(&FHorizontalAlignmentCustomization::MakeInstance), nullptr, PropertyView);
+	PropertyModule.RegisterCustomPropertyTypeLayout(TEXT("EVerticalAlignment"), FOnGetPropertyTypeCustomizationInstance::CreateStatic(&FVerticalAlignmentCustomization::MakeInstance), nullptr, PropertyView);
+	PropertyModule.RegisterCustomPropertyTypeLayout(TEXT("SlateChildSize"), FOnGetPropertyTypeCustomizationInstance::CreateStatic(&FSlateChildSizeCustomization::MakeInstance), nullptr, PropertyView);
 }
 
 void SWidgetDetailsView::OnEditorSelectionChanging()
