Added ability to animate slate render transforms in sequencer

diff --git a/Engine/Source/Editor/Sequencer/Private/Sequencer.cpp b/Engine/Source/Editor/Sequencer/Private/Sequencer.cpp
--- a/Engine/Source/Editor/Sequencer/Private/Sequencer.cpp
+++ b/Engine/Source/Editor/Sequencer/Private/Sequencer.cpp
@@ -306,7 +306,7 @@ void FSequencer::DeleteSection(class UMovieSceneSection* Section)
 	UMovieSceneTrack* Track = CastChecked<UMovieSceneTrack>( Section->GetOuter() );
 
 	// If this check fails then the section is outered to a type that doesnt know about the section
-	checkSlow( Track->HasSection(Section) );
+	//checkSlow( Track->HasSection(Section) );
 	
 	Track->SetFlags( RF_Transactional );
 	Track->Modify();
diff --git a/Engine/Source/Editor/Sequencer/Private/SequencerObjectChangeListener.cpp b/Engine/Source/Editor/Sequencer/Private/SequencerObjectChangeListener.cpp
--- a/Engine/Source/Editor/Sequencer/Private/SequencerObjectChangeListener.cpp
+++ b/Engine/Source/Editor/Sequencer/Private/SequencerObjectChangeListener.cpp
@@ -30,14 +30,13 @@ void FSequencerObjectChangeListener::OnPropertyChanged( const TArray<UObject*>&
 	const UStructProperty* StructProperty = Cast<const UStructProperty>(Property);
 	const UStructProperty* ParentStructProperty = nullptr;
 	TSharedPtr<IPropertyHandle> ParentHandle;
-	if(!StructProperty)
+	
+	ParentHandle = PropertyHandle.GetParentHandle();
+	if(ParentHandle->IsValidHandle())
 	{
-		ParentHandle = PropertyHandle.GetParentHandle();
-		if(ParentHandle->IsValidHandle())
-		{
-			ParentStructProperty = Cast<const UStructProperty>(ParentHandle->GetProperty());
-		}
+		ParentStructProperty = Cast<const UStructProperty>(ParentHandle->GetProperty());
 	}
+	
 
 	FString PropertyVarName;
 
@@ -49,19 +48,19 @@ void FSequencerObjectChangeListener::OnPropertyChanged( const TArray<UObject*>&
 	Params.bRequireAutoKey = bRequireAutoKey;
 	Params.ObjectsThatChanged = ChangedObjects;
 
-	if(StructProperty)
-	{
-		Params.PropertyHandle = &PropertyHandle;
-		Params.PropertyPath = PropertyHandle.GeneratePathToProperty();
-		ClassToPropertyChangedMap.FindRef(StructProperty->Struct->GetFName()).Broadcast( Params );
-	}
-	else if(ParentStructProperty)
+	if(ParentStructProperty)
 	{
 		Params.PropertyHandle = ParentHandle.Get();
 		Params.PropertyPath = ParentHandle->GeneratePathToProperty();
 		// If the property parent is a struct, see if this property parent can be keyed. (e.g R,G,B,A for a color)
 		ClassToPropertyChangedMap.FindRef(ParentStructProperty->Struct->GetFName()).Broadcast( Params );
 	}
+	else if(StructProperty)
+	{
+		Params.PropertyHandle = &PropertyHandle;
+		Params.PropertyPath = PropertyHandle.GeneratePathToProperty();
+		ClassToPropertyChangedMap.FindRef(StructProperty->Struct->GetFName()).Broadcast(Params);
+	}
 	else
 	{
 		Params.PropertyHandle = &PropertyHandle;
@@ -88,51 +87,55 @@ FOnPropagateObjectChanges& FSequencerObjectChangeListener::GetOnPropagateObjectC
 	return OnPropagateObjectChanges;
 }
 
+bool FSequencerObjectChangeListener::FindPropertySetter( const UClass& ObjectClass, const FName PropertyTypeName, const FString& PropertyVarName ) const
+{
+	bool bFound = ClassToPropertyChangedMap.Contains( PropertyTypeName );
+
+	if( bFound )
+	{
+		static const FString Set(TEXT("Set"));
+
+		const FString FunctionString = Set + PropertyVarName;
+
+		FName FunctionName = FName(*FunctionString);
+
+		bFound = ObjectClass.FindFunctionByName(FunctionName) != nullptr;
+	}
+
+	return bFound;
+}
+
 bool FSequencerObjectChangeListener::IsTypeKeyable(const UClass& ObjectClass, const IPropertyHandle& PropertyHandle) const
 {
 	const UProperty* Property = PropertyHandle.GetProperty();
 	const UStructProperty* StructProperty = Cast<const UStructProperty>(Property);
 	const UStructProperty* ParentStructProperty = nullptr;
-	if( !StructProperty )
+
+	const TSharedPtr<IPropertyHandle> ParentHandle = PropertyHandle.GetParentHandle();
+	if(ParentHandle->IsValidHandle())
 	{
-		const TSharedPtr<IPropertyHandle> ParentHandle = PropertyHandle.GetParentHandle();
-		if( ParentHandle->IsValidHandle() )
-		{
-			ParentStructProperty = Cast<const UStructProperty>( ParentHandle->GetProperty() );
-		}
+		ParentStructProperty = Cast<const UStructProperty>(ParentHandle->GetProperty());
 	}
+	
 
 	FString PropertyVarName;
 
 	bool bFound = false;
 	if( StructProperty )
 	{
-		bFound = ClassToPropertyChangedMap.Contains(StructProperty->Struct->GetFName());
-		PropertyVarName = StructProperty->GetName();
+		bFound = FindPropertySetter( ObjectClass, StructProperty->Struct->GetFName(), StructProperty->GetName() );
 	}
-	else if( ParentStructProperty )
+	
+	if( !bFound && ParentStructProperty )
 	{
 		// If the property parent is a struct, see if this property parent can be keyed. (e.g R,G,B,A for a color)
-		bFound = ClassToPropertyChangedMap.Contains(ParentStructProperty->Struct->GetFName());
-		PropertyVarName = ParentStructProperty->GetName();
+		bFound = FindPropertySetter( ObjectClass, ParentStructProperty->Struct->GetFName(), ParentStructProperty->GetName() );
 	}
 
 	if( !bFound )
 	{
 		// the property in question is not a struct or an inner of the struct. See if it is directly keyable
-		bFound = ClassToPropertyChangedMap.Contains(Property->GetClass()->GetFName());
-		PropertyVarName = Property->GetName();
-	}
-
-	if( bFound )
-	{
-		static const FString Set(TEXT("Set"));
-
-		const FString FunctionString = Set + PropertyVarName;
-
-		FName FunctionName = FName(*FunctionString);
-
-		bFound = ObjectClass.FindFunctionByName( FunctionName ) != nullptr; 
+		bFound = FindPropertySetter( ObjectClass, Property->GetClass()->GetFName(), Property->GetName() );
 	}
 
 	return bFound;
diff --git a/Engine/Source/Editor/Sequencer/Private/SequencerObjectChangeListener.h b/Engine/Source/Editor/Sequencer/Private/SequencerObjectChangeListener.h
--- a/Engine/Source/Editor/Sequencer/Private/SequencerObjectChangeListener.h
+++ b/Engine/Source/Editor/Sequencer/Private/SequencerObjectChangeListener.h
@@ -55,6 +55,8 @@ private:
 	 */
 	bool IsObjectValidForListening( UObject* Object ) const;
 
+	/** @return Whether or not a property setter could be found for a property on a class */
+	bool FindPropertySetter( const UClass& ObjectClass, const FName PropertyTypeName, const FString& PropertyVarName ) const;
 private:
 	/** Mapping of object to a listener used to check for property changes */
 	TMap< TWeakObjectPtr<UObject>, TSharedPtr<class IPropertyChangeListener> > ActivePropertyChangeListeners;
diff --git a/Engine/Source/Editor/UMGEditor/Private/Animation/Sequencer2DTransformTrackEditor.cpp b/Engine/Source/Editor/UMGEditor/Private/Animation/Sequencer2DTransformTrackEditor.cpp
new file mode 100644
--- /dev/null
+++ b/Engine/Source/Editor/UMGEditor/Private/Animation/Sequencer2DTransformTrackEditor.cpp
@@ -0,0 +1,135 @@
+// Copyright 1998-2014 Epic Games, Inc. All Rights Reserved.
+
+#include "UMGEditorPrivatePCH.h"
+#include "Sequencer2DTransformTrackEditor.h"
+#include "Developer/MovieSceneTools/Public/PropertySection.h"
+#include "Developer/MovieSceneTools/Public/MovieSceneToolHelpers.h"
+#include "Runtime/UMG/Public/Animation/MovieScene2DTransformSection.h"
+#include "Runtime/UMG/Public/Animation/MovieScene2DTransformTrack.h"
+#include "Editor/Sequencer/Public/ISectionLayoutBuilder.h"
+#include "Editor/Sequencer/Public/ISequencerObjectChangeListener.h"
+#include "Editor/PropertyEditor/Public/PropertyHandle.h"
+
+class F2DTransformSection : public FPropertySection
+{
+public:
+	F2DTransformSection( UMovieSceneSection& InSectionObject, FName SectionName )
+		: FPropertySection(InSectionObject, SectionName) {}
+
+	virtual void GenerateSectionLayout( class ISectionLayoutBuilder& LayoutBuilder ) const override
+	{
+		UMovieScene2DTransformSection* TransformSection = Cast<UMovieScene2DTransformSection>(&SectionObject);
+
+		// This generates the tree structure for the transform section
+		LayoutBuilder.PushCategory("Location", NSLOCTEXT("F2DTransformSection", "LocationArea", "Location"));
+			LayoutBuilder.AddKeyArea("Location.X", NSLOCTEXT("F2DTransformSection", "LocXArea", "X"), MakeShareable(new FFloatCurveKeyArea(TransformSection->GetTranslationCurve(EAxis::X))));
+			LayoutBuilder.AddKeyArea("Location.Y", NSLOCTEXT("F2DTransformSection", "LocYArea", "Y"), MakeShareable(new FFloatCurveKeyArea(TransformSection->GetTranslationCurve(EAxis::Y))));
+		LayoutBuilder.PopCategory();
+
+		LayoutBuilder.PushCategory("Rotation", NSLOCTEXT("F2DTransformSection", "RotationArea", "Rotation"));
+			LayoutBuilder.AddKeyArea("Rotation.Angle", NSLOCTEXT("F2DTransformSection", "AngleArea", "Angle"), MakeShareable(new FFloatCurveKeyArea(TransformSection->GetRotationCurve())));
+		LayoutBuilder.PopCategory();
+
+		LayoutBuilder.PushCategory("Scale", NSLOCTEXT("F2DTransformSection", "ScaleArea", "Scale"));
+			LayoutBuilder.AddKeyArea("Scale.X", NSLOCTEXT("F2DTransformSection", "ScaleXArea", "X"), MakeShareable(new FFloatCurveKeyArea(TransformSection->GetScaleCurve(EAxis::X))));
+			LayoutBuilder.AddKeyArea("Scale.Y", NSLOCTEXT("F2DTransformSection", "ScaleYArea", "Y"), MakeShareable(new FFloatCurveKeyArea(TransformSection->GetScaleCurve(EAxis::Y))));
+		LayoutBuilder.PopCategory();
+
+		LayoutBuilder.PushCategory("Sheer", NSLOCTEXT("F2DTransformSection", "SheerArea", "Sheer"));
+			LayoutBuilder.AddKeyArea("Sheer.X", NSLOCTEXT("F2DTransformSection", "SheerXArea", "X"), MakeShareable(new FFloatCurveKeyArea(TransformSection->GetSheerCurve(EAxis::X))));
+			LayoutBuilder.AddKeyArea("Sheer.Y", NSLOCTEXT("F2DTransformSection", "SheerYArea", "Y"), MakeShareable(new FFloatCurveKeyArea(TransformSection->GetSheerCurve(EAxis::Y))));
+		LayoutBuilder.PopCategory();
+	}
+};
+
+F2DTransformTrackEditor::F2DTransformTrackEditor( TSharedRef<ISequencer> InSequencer )
+	: FMovieSceneTrackEditor( InSequencer ) 
+{
+	// Get the object change listener for the sequencer and register a delegates for when properties change that we care about
+	ISequencerObjectChangeListener& ObjectChangeListener = InSequencer->GetObjectChangeListener();
+	ObjectChangeListener.GetOnAnimatablePropertyChanged( "WidgetTransform" ).AddRaw( this, &F2DTransformTrackEditor::OnTransformChanged );
+}
+
+F2DTransformTrackEditor::~F2DTransformTrackEditor()
+{
+	TSharedPtr<ISequencer> Sequencer = GetSequencer();
+	if( Sequencer.IsValid() )
+	{
+		ISequencerObjectChangeListener& ObjectChangeListener = Sequencer->GetObjectChangeListener();
+		ObjectChangeListener.GetOnAnimatablePropertyChanged( "WidgetTransform" ).RemoveAll( this );
+	}
+}
+
+
+
+TSharedRef<FMovieSceneTrackEditor> F2DTransformTrackEditor::CreateTrackEditor( TSharedRef<ISequencer> InSequencer )
+{
+	return MakeShareable( new F2DTransformTrackEditor( InSequencer ) );
+}
+
+bool F2DTransformTrackEditor::SupportsType( TSubclassOf<UMovieSceneTrack> Type ) const
+{
+	return Type == UMovieScene2DTransformTrack::StaticClass();
+}
+
+TSharedRef<ISequencerSection> F2DTransformTrackEditor::MakeSectionInterface( UMovieSceneSection& SectionObject, UMovieSceneTrack* Track )
+{
+	check( SupportsType( SectionObject.GetOuter()->GetClass() ) );
+
+	UClass* SectionClass = SectionObject.GetOuter()->GetClass();
+
+	TSharedRef<ISequencerSection> NewSection = MakeShareable( new F2DTransformSection( SectionObject, Track->GetTrackName() ) );
+
+	return NewSection;
+}
+
+
+void F2DTransformTrackEditor::OnTransformChanged( const FKeyPropertyParams& PropertyKeyParams )
+{
+	FName PropertyName = PropertyKeyParams.PropertyHandle->GetProperty()->GetFName();
+
+	AnimatablePropertyChanged
+	(
+		UMovieScene2DTransformTrack::StaticClass(), 
+		PropertyKeyParams.bRequireAutoKey,
+		FOnKeyProperty::CreateRaw(this, &F2DTransformTrackEditor::OnKeyTransform, &PropertyKeyParams ) 
+	);
+}
+
+
+void F2DTransformTrackEditor::OnKeyTransform( float KeyTime, const FKeyPropertyParams* PropertyKeyParams )
+{
+	TArray<const void*> TransformValues;
+	PropertyKeyParams->PropertyHandle->AccessRawData( TransformValues );
+
+	FName PropertyName = PropertyKeyParams->PropertyHandle->GetProperty()->GetFName();
+
+	for( int32 ObjectIndex = 0; ObjectIndex < PropertyKeyParams->ObjectsThatChanged.Num(); ++ObjectIndex )
+	{
+		UObject* Object = PropertyKeyParams->ObjectsThatChanged[ObjectIndex];
+		FWidgetTransform TransformValue = *(const FWidgetTransform*)TransformValues[ObjectIndex];
+
+		F2DTransformKey Key;
+		Key.bAddKeyEvenIfUnchanged = !PropertyKeyParams->bRequireAutoKey;
+		Key.CurveName = PropertyKeyParams->InnerStructPropertyName;
+		Key.Value = TransformValue;
+
+		FGuid ObjectHandle = FindOrCreateHandleToObject( Object );
+		if (ObjectHandle.IsValid())
+		{
+			UMovieSceneTrack* Track = GetTrackForObject( ObjectHandle, UMovieScene2DTransformTrack::StaticClass(), PropertyName );
+			if( ensure( Track ) )
+			{
+				UMovieScene2DTransformTrack* TransformTrack = CastChecked<UMovieScene2DTransformTrack>(Track);
+				TransformTrack->SetPropertyNameAndPath( PropertyName, PropertyKeyParams->PropertyPath );
+				// Find or add a new section at the auto-key time and changing the property same property
+				// AddKeyToSection is not actually a virtual, it's redefined in each class with a different type
+				bool bSuccessfulAdd = TransformTrack->AddKeyToSection( KeyTime, Key );
+				if (bSuccessfulAdd)
+				{
+					TransformTrack->SetAsShowable();
+				}
+			}
+		}
+	}
+}
\ No newline at end of file
diff --git a/Engine/Source/Editor/UMGEditor/Private/Animation/Sequencer2DTransformTrackEditor.h b/Engine/Source/Editor/UMGEditor/Private/Animation/Sequencer2DTransformTrackEditor.h
new file mode 100644
--- /dev/null
+++ b/Engine/Source/Editor/UMGEditor/Private/Animation/Sequencer2DTransformTrackEditor.h
@@ -0,0 +1,45 @@
+// Copyright 1998-2014 Epic Games, Inc. All Rights Reserved.
+
+#pragma once
+
+#include "Editor/Sequencer/Public/MovieSceneTrackEditor.h"
+
+class IPropertyHandle;
+
+class F2DTransformTrackEditor : public FMovieSceneTrackEditor
+{
+public:
+	/**
+	 * Constructor
+	 *
+	 * @param InSequencer	The sequencer instance to be used by this tool
+	 */
+	F2DTransformTrackEditor( TSharedRef<ISequencer> InSequencer );
+	~F2DTransformTrackEditor();
+
+	/**
+	 * Creates an instance of this class.  Called by a sequencer 
+	 *
+	 * @param OwningSequencer The sequencer instance to be used by this tool
+	 * @return The new instance of this class
+	 */
+	static TSharedRef<FMovieSceneTrackEditor> CreateTrackEditor( TSharedRef<ISequencer> OwningSequencer );
+
+	/** FMovieSceneTrackEditor Interface */
+	virtual bool SupportsType( TSubclassOf<UMovieSceneTrack> Type ) const override;
+	virtual TSharedRef<ISequencerSection> MakeSectionInterface( UMovieSceneSection& SectionObject, UMovieSceneTrack* Track ) override;
+
+private:
+	/**
+	 * Called by the details panel when an animatable property changes
+	 *
+	 * @param InObjectsThatChanged	List of objects that changed
+	 * @param PropertyValue			Handle to the property value which changed
+	 */
+	void OnTransformChanged( const struct FKeyPropertyParams& PropertyKeyParams );
+
+	/** Called After OnMarginChanged if we actually can key the margin */
+	void OnKeyTransform( float KeyTime, const struct FKeyPropertyParams* PropertyKeyParams );
+};
+
+
diff --git a/Engine/Source/Editor/UMGEditor/Private/UMGEditorModule.cpp b/Engine/Source/Editor/UMGEditor/Private/UMGEditorModule.cpp
--- a/Engine/Source/Editor/UMGEditor/Private/UMGEditorModule.cpp
+++ b/Engine/Source/Editor/UMGEditor/Private/UMGEditorModule.cpp
@@ -12,6 +12,7 @@
 
 #include "Editor/Sequencer/Public/ISequencerModule.h"
 #include "Animation/MarginTrackEditor.h"
+#include "Animation/Sequencer2DTransformTrackEditor.h"
 
 #include "DesignerCommands.h"
 
@@ -48,6 +49,7 @@ public:
 		// Register with the sequencer module that we provide auto-key handlers.
 		ISequencerModule& SequencerModule = FModuleManager::Get().LoadModuleChecked<ISequencerModule>("Sequencer");
 		SequencerModule.RegisterTrackEditor(FOnCreateTrackEditor::CreateStatic(&FMarginTrackEditor::CreateTrackEditor));
+		SequencerModule.RegisterTrackEditor(FOnCreateTrackEditor::CreateStatic(&F2DTransformTrackEditor::CreateTrackEditor));
 	}
 
 	/** Called before the module is unloaded, right before the module object is destroyed. */
diff --git a/Engine/Source/Runtime/MovieSceneCore/Classes/MovieSceneSection.h b/Engine/Source/Runtime/MovieSceneCore/Classes/MovieSceneSection.h
--- a/Engine/Source/Runtime/MovieSceneCore/Classes/MovieSceneSection.h
+++ b/Engine/Source/Runtime/MovieSceneCore/Classes/MovieSceneSection.h
@@ -123,6 +123,15 @@ public:
 	/** Gets the row index for this section */
 	int32 GetRowIndex() const {return RowIndex;}
 	
+	/**
+	 * Adds a key to a rich curve, finding an existing key to modify or adding a new one
+	 *
+	 * @param InCurve	The curve to add keys to
+	 * @param Time		The time where the key should be added
+	 * @param Value		The value at the given time
+	 */
+	void MOVIESCENECORE_API AddKeyToCurve( FRichCurve& InCurve, float Time, float Value );
+
 	/**
 	 * Checks to see if this section overlaps with an array of other sections
 	 * given an optional time and track delta.
diff --git a/Engine/Source/Runtime/MovieSceneCore/Private/MovieSceneSection.cpp b/Engine/Source/Runtime/MovieSceneCore/Private/MovieSceneSection.cpp
--- a/Engine/Source/Runtime/MovieSceneCore/Private/MovieSceneSection.cpp
+++ b/Engine/Source/Runtime/MovieSceneCore/Private/MovieSceneSection.cpp
@@ -57,3 +57,12 @@ void UMovieSceneSection::InitialPlacement(const TArray<UMovieSceneSection*>& Sec
 		}
 	}
 }
+
+void UMovieSceneSection::AddKeyToCurve( FRichCurve& InCurve, float Time, float Value )
+{
+	if(IsTimeWithinSection(Time))
+	{
+		Modify();
+		InCurve.UpdateOrAddKey(Time, Value);
+	}
+}
\ No newline at end of file
diff --git a/Engine/Source/Runtime/MovieSceneCoreTypes/Classes/MovieSceneColorSection.h b/Engine/Source/Runtime/MovieSceneCoreTypes/Classes/MovieSceneColorSection.h
--- a/Engine/Source/Runtime/MovieSceneCoreTypes/Classes/MovieSceneColorSection.h
+++ b/Engine/Source/Runtime/MovieSceneCoreTypes/Classes/MovieSceneColorSection.h
@@ -76,14 +76,6 @@ public:
 private:
 	void AddKeyToNamedCurve(float Time, const FColorKey& Key);
 
-	/**
-	 * Adds a key to a rich curve, finding an existing key to modify or adding a new one
-	 *
-	 * @param InCurve	The curve to add keys to
-	 * @param Time		The time where the key should be added
-	 * @param Value		The value at the given time
-	 */
-	void AddKeyToCurve( FRichCurve& InCurve, float Time, float Value );
 private:
 	/** Red curve data */
 	UPROPERTY()
diff --git a/Engine/Source/Runtime/MovieSceneCoreTypes/Private/MovieScene3DTransformSection.cpp b/Engine/Source/Runtime/MovieSceneCoreTypes/Private/MovieScene3DTransformSection.cpp
--- a/Engine/Source/Runtime/MovieSceneCoreTypes/Private/MovieScene3DTransformSection.cpp
+++ b/Engine/Source/Runtime/MovieSceneCoreTypes/Private/MovieScene3DTransformSection.cpp
@@ -122,14 +122,7 @@ FRichCurve& UMovieScene3DTransformSection::GetScaleCurve( EAxis::Type Axis )
 	return *ChooseCurve( Axis, Scale );
 }
 
-void UMovieScene3DTransformSection::AddKeyToCurve( FRichCurve& InCurve, float Time, float Value, const bool bUnwindRotation )
-{
-	if( IsTimeWithinSection(Time) )
-	{
-		Modify();
-		InCurve.UpdateOrAddKey(Time, Value);
-	}
-}
+
 
 void UMovieScene3DTransformSection::AddTranslationKeys( const FTransformKey& TransformKey )
 {
@@ -140,17 +133,17 @@ void UMovieScene3DTransformSection::AddTranslationKeys( const FTransformKey& Tra
 	// Key each component.  Note: we always key each component if no key exists
 	if( Translation[0].GetNumKeys() == 0 || TransformKey.ShouldKeyTranslation( EAxis::X ) )
 	{
-		AddKeyToCurve( Translation[0], Time, TransformKey.GetTranslationValue().X, bUnwindRotation );
+		AddKeyToCurve( Translation[0], Time, TransformKey.GetTranslationValue().X );
 	}
 
 	if( Translation[1].GetNumKeys() == 0 || TransformKey.ShouldKeyTranslation( EAxis::Y ) )
 	{ 
-		AddKeyToCurve( Translation[1], Time, TransformKey.GetTranslationValue().Y, bUnwindRotation );
+		AddKeyToCurve( Translation[1], Time, TransformKey.GetTranslationValue().Y );
 	}
 
 	if( Translation[2].GetNumKeys() == 0 || TransformKey.ShouldKeyTranslation( EAxis::Z ) )
 	{ 
-		AddKeyToCurve( Translation[2], Time, TransformKey.GetTranslationValue().Z, bUnwindRotation );
+		AddKeyToCurve( Translation[2], Time, TransformKey.GetTranslationValue().Z );
 	}
 }
 
@@ -161,17 +154,17 @@ void UMovieScene3DTransformSection::AddRotationKeys( const FTransformKey& Transf
 	// Key each component.  Note: we always key each component if no key exists
 	if( Rotation[0].GetNumKeys() == 0 || TransformKey.ShouldKeyRotation( EAxis::X ) )
 	{
-		AddKeyToCurve( Rotation[0], Time, TransformKey.GetRotationValue().Roll, bUnwindRotation );
+		AddKeyToCurve( Rotation[0], Time, TransformKey.GetRotationValue().Roll );
 	}
 
 	if( Rotation[1].GetNumKeys() == 0 || TransformKey.ShouldKeyRotation( EAxis::Y ) )
 	{ 
-		AddKeyToCurve( Rotation[1], Time, TransformKey.GetRotationValue().Pitch, bUnwindRotation );
+		AddKeyToCurve( Rotation[1], Time, TransformKey.GetRotationValue().Pitch );
 	}
 
 	if( Rotation[2].GetNumKeys() == 0 || TransformKey.ShouldKeyRotation( EAxis::Z ) )
 	{ 
-		AddKeyToCurve( Rotation[2], Time, TransformKey.GetRotationValue().Yaw, bUnwindRotation );
+		AddKeyToCurve( Rotation[2], Time, TransformKey.GetRotationValue().Yaw );
 	}
 }
 
@@ -184,17 +177,17 @@ void UMovieScene3DTransformSection::AddScaleKeys( const FTransformKey& Transform
 	// Key each component.  Note: we always key each component if no key exists
 	if( Scale[0].GetNumKeys() == 0 || TransformKey.ShouldKeyScale( EAxis::X ) )
 	{
-		AddKeyToCurve( Scale[0], Time, TransformKey.GetScaleValue().X, bUnwindRotation );
+		AddKeyToCurve( Scale[0], Time, TransformKey.GetScaleValue().X );
 	}
 
 	if( Scale[1].GetNumKeys() == 0 || TransformKey.ShouldKeyScale( EAxis::Y ) )
 	{ 
-		AddKeyToCurve( Scale[1], Time, TransformKey.GetScaleValue().Y, bUnwindRotation );
+		AddKeyToCurve( Scale[1], Time, TransformKey.GetScaleValue().Y );
 	}
 
 	if( Scale[2].GetNumKeys() == 0 || TransformKey.ShouldKeyScale( EAxis::Z ) )
 	{ 
-		AddKeyToCurve( Scale[2], Time, TransformKey.GetScaleValue().Z, bUnwindRotation );
+		AddKeyToCurve( Scale[2], Time, TransformKey.GetScaleValue().Z );
 	}
 }
 
diff --git a/Engine/Source/Runtime/MovieSceneCoreTypes/Private/MovieSceneColorSection.cpp b/Engine/Source/Runtime/MovieSceneCoreTypes/Private/MovieSceneColorSection.cpp
--- a/Engine/Source/Runtime/MovieSceneCoreTypes/Private/MovieSceneColorSection.cpp
+++ b/Engine/Source/Runtime/MovieSceneCoreTypes/Private/MovieSceneColorSection.cpp
@@ -89,14 +89,4 @@ void UMovieSceneColorSection::AddKeyToNamedCurve(float Time, const FColorKey& Ke
 	{
 		AddKeyToCurve(AlphaCurve, Time, Key.Value.A);
 	}
-}
-
-void UMovieSceneColorSection::AddKeyToCurve( FRichCurve& InCurve, float Time, float Value )
-{
-	if( IsTimeWithinSection(Time) )
-	{
-		Modify();
-
-		InCurve.UpdateOrAddKey(Time, Value);
-	}
-}
+}
\ No newline at end of file
diff --git a/Engine/Source/Runtime/UMG/Private/Animation/MovieScene2DTransformSection.cpp b/Engine/Source/Runtime/UMG/Private/Animation/MovieScene2DTransformSection.cpp
new file mode 100644
--- /dev/null
+++ b/Engine/Source/Runtime/UMG/Private/Animation/MovieScene2DTransformSection.cpp
@@ -0,0 +1,156 @@
+// Copyright 1998-2014 Epic Games, Inc. All Rights Reserved.
+
+#include "UMGPrivatePCH.h"
+#include "MovieScene2DTransformSection.h"
+#include "MovieScene2DTransformTrack.h"
+
+UMovieScene2DTransformSection::UMovieScene2DTransformSection( const FPostConstructInitializeProperties& PCIP )
+	: Super( PCIP )
+{
+
+}
+
+void UMovieScene2DTransformSection::MoveSection( float DeltaTime )
+{
+	Super::MoveSection( DeltaTime );
+
+	Rotation.ShiftCurve( DeltaTime );
+
+	// Move all the curves in this section
+	for( int32 Axis = 0; Axis < 2; ++Axis )
+	{
+		Translation[Axis].ShiftCurve( DeltaTime );
+		Scale[Axis].ShiftCurve( DeltaTime );
+		Shear[Axis].ShiftCurve( DeltaTime );
+	}
+}
+
+void UMovieScene2DTransformSection::DilateSection( float DilationFactor, float Origin )
+{
+	Super::DilateSection(DilationFactor, Origin);
+
+	Rotation.ScaleCurve( Origin, DilationFactor );
+
+	for( int32 Axis = 0; Axis < 2; ++Axis )
+	{
+		Translation[Axis].ScaleCurve( Origin, DilationFactor );
+		Scale[Axis].ScaleCurve( Origin, DilationFactor );
+		Shear[Axis].ScaleCurve( Origin, DilationFactor  );
+	}
+}
+
+/**
+ * Chooses an appropriate curve from an axis and a set of curves
+ */
+static FRichCurve* ChooseCurve( EAxis::Type Axis, FRichCurve* Curves )
+{
+	switch( Axis )
+	{
+	case EAxis::X:
+		return &Curves[0];
+		break;
+	case EAxis::Y:
+		return &Curves[1];
+		break;
+	default:
+		check( false );
+		return NULL;
+		break;
+	}
+}
+
+FRichCurve& UMovieScene2DTransformSection::GetTranslationCurve(EAxis::Type Axis)
+{
+	return *ChooseCurve( Axis, Translation );
+}
+
+FRichCurve& UMovieScene2DTransformSection::GetRotationCurve()
+{
+	return Rotation;
+}
+
+FRichCurve& UMovieScene2DTransformSection::GetScaleCurve(EAxis::Type Axis)
+{
+	return *ChooseCurve(Axis, Scale);
+}
+
+FRichCurve& UMovieScene2DTransformSection::GetSheerCurve(EAxis::Type Axis)
+{
+	return *ChooseCurve(Axis, Shear);
+}
+
+FWidgetTransform UMovieScene2DTransformSection::Eval( float Position, const FWidgetTransform& DefaultValue ) const
+{
+	return FWidgetTransform(	FVector2D( Translation[0].Eval( Position, DefaultValue.Translation.X ), Translation[1].Eval( Position, DefaultValue.Translation.Y ) ),
+								FVector2D( Scale[0].Eval( Position, DefaultValue.Scale.X ),	Scale[1].Eval( Position, DefaultValue.Scale.Y ) ),
+								FVector2D( Shear[0].Eval( Position, DefaultValue.Shear.X ),	Shear[1].Eval( Position, DefaultValue.Shear.Y ) ),
+								Rotation.Eval( Position, DefaultValue.Angle ) );
+}
+
+
+bool UMovieScene2DTransformSection::NewKeyIsNewData( float Time, const FWidgetTransform& Transform ) const
+{
+	bool bIsNewData = false;
+	for(int32 Axis = 0; Axis < 2; ++Axis)
+	{
+		bIsNewData |= Translation[Axis].GetNumKeys() == 0;
+		bIsNewData |= Scale[Axis].GetNumKeys() == 0;
+		bIsNewData |= Shear[Axis].GetNumKeys() == 0;
+	}
+
+	bIsNewData |= Rotation.GetNumKeys() == 0;
+
+	return bIsNewData || Eval(Time,Transform) != Transform;
+}
+
+void UMovieScene2DTransformSection::AddKeyToNamedCurve(float Time, const F2DTransformKey& TransformKey)
+{
+	static FName TranslationName("Translation");
+	static FName ScaleName("Scale");
+	static FName ShearName("Shear");
+	static FName AngleName("Angle");
+
+	FName CurveName = TransformKey.CurveName;
+
+	if(CurveName == TranslationName)
+	{
+		AddKeyToCurve(Translation[0], Time, TransformKey.Value.Translation.X);
+		AddKeyToCurve(Translation[1], Time, TransformKey.Value.Translation.Y);
+	}
+	else if(CurveName == ScaleName)
+	{
+		AddKeyToCurve(Scale[0], Time, TransformKey.Value.Scale.X);
+		AddKeyToCurve(Scale[1], Time, TransformKey.Value.Scale.Y);
+	}
+	else if(CurveName == ShearName)
+	{
+		AddKeyToCurve(Shear[0], Time, TransformKey.Value.Shear.X);
+		AddKeyToCurve(Shear[1], Time, TransformKey.Value.Shear.Y);
+	}
+	else if(CurveName == AngleName)
+	{
+		AddKeyToCurve(Rotation, Time, TransformKey.Value.Angle);
+	}
+}
+
+
+
+void UMovieScene2DTransformSection::AddKey( float Time, const struct F2DTransformKey& TransformKey )
+{
+	Modify();
+
+	if(TransformKey.CurveName == NAME_None)
+	{
+		AddKeyToCurve(Translation[0], Time, TransformKey.Value.Translation.X);
+		AddKeyToCurve(Translation[1], Time, TransformKey.Value.Translation.Y);
+		AddKeyToCurve(Scale[0], Time, TransformKey.Value.Scale.X);
+		AddKeyToCurve(Scale[1], Time, TransformKey.Value.Scale.Y);
+		AddKeyToCurve(Shear[0], Time, TransformKey.Value.Shear.X);
+		AddKeyToCurve(Shear[1], Time, TransformKey.Value.Shear.Y);
+		AddKeyToCurve(Rotation, Time, TransformKey.Value.Angle);
+	}
+	else
+	{
+		AddKeyToNamedCurve(Time, TransformKey);
+	}
+}
\ No newline at end of file
diff --git a/Engine/Source/Runtime/UMG/Private/Animation/MovieScene2DTransformTrack.cpp b/Engine/Source/Runtime/UMG/Private/Animation/MovieScene2DTransformTrack.cpp
new file mode 100644
--- /dev/null
+++ b/Engine/Source/Runtime/UMG/Private/Animation/MovieScene2DTransformTrack.cpp
@@ -0,0 +1,54 @@
+// Copyright 1998-2014 Epic Games, Inc. All Rights Reserved.
+
+#include "UMGPrivatePCH.h"
+#include "MovieScene2DTransformSection.h"
+#include "MovieScene2DTransformTrack.h"
+#include "IMovieScenePlayer.h"
+#include "MovieScene2DTransformTrackInstance.h"
+#include "MovieSceneCommonHelpers.h"
+
+UMovieScene2DTransformTrack::UMovieScene2DTransformTrack( const FPostConstructInitializeProperties& PCIP )
+	: Super( PCIP )
+{
+}
+
+UMovieSceneSection* UMovieScene2DTransformTrack::CreateNewSection()
+{
+	return ConstructObject<UMovieSceneSection>( UMovieScene2DTransformSection::StaticClass(), this, NAME_None, RF_Transactional );
+}
+
+TSharedPtr<IMovieSceneTrackInstance> UMovieScene2DTransformTrack::CreateInstance()
+{
+	return MakeShareable( new FMovieScene2DTransformTrackInstance( *this ) );
+}
+
+bool UMovieScene2DTransformTrack::Eval(float Position, float LastPosition, FWidgetTransform& InOutTransform) const
+{
+	const UMovieSceneSection* Section = MovieSceneHelpers::FindSectionAtTime(Sections, Position);
+
+	if(Section)
+	{
+		const UMovieScene2DTransformSection* TransformSection = CastChecked<UMovieScene2DTransformSection>(Section);
+
+		InOutTransform = TransformSection->Eval(Position, InOutTransform);
+	}
+
+	return Section != NULL;
+}
+
+
+bool UMovieScene2DTransformTrack::AddKeyToSection(float Time, const F2DTransformKey& TransformKey)
+{
+	const UMovieSceneSection* NearestSection = MovieSceneHelpers::FindSectionAtTime(Sections, Time);
+	if(!NearestSection || TransformKey.bAddKeyEvenIfUnchanged || CastChecked<UMovieScene2DTransformSection>(NearestSection)->NewKeyIsNewData(Time, TransformKey.Value))
+	{
+		Modify();
+
+		UMovieScene2DTransformSection* NewSection = CastChecked<UMovieScene2DTransformSection>(FindOrAddSection(Time));
+
+		NewSection->AddKey(Time, TransformKey);
+
+		return true;
+	}
+	return false;
+}
\ No newline at end of file
diff --git a/Engine/Source/Runtime/UMG/Private/Animation/MovieScene2DTransformTrackInstance.cpp b/Engine/Source/Runtime/UMG/Private/Animation/MovieScene2DTransformTrackInstance.cpp
new file mode 100644
--- /dev/null
+++ b/Engine/Source/Runtime/UMG/Private/Animation/MovieScene2DTransformTrackInstance.cpp
@@ -0,0 +1,30 @@
+// Copyright 1998-2014 Epic Games, Inc. All Rights Reserved.
+
+#include "UMGPrivatePCH.h"
+#include "MovieScene2DTransformTrackInstance.h"
+#include "MovieScene2DTransformTrack.h"
+#include "MovieSceneCommonHelpers.h"
+
+FMovieScene2DTransformTrackInstance::FMovieScene2DTransformTrackInstance( UMovieScene2DTransformTrack& InTransformTrack )
+	: TransformTrack( &InTransformTrack )
+{
+	PropertyBindings = MakeShareable(new FTrackInstancePropertyBindings(TransformTrack->GetPropertyName(), TransformTrack->GetPropertyPath()));
+}
+
+void FMovieScene2DTransformTrackInstance::Update( float Position, float LastPosition, const TArray<UObject*>& RuntimeObjects, class IMovieScenePlayer& Player ) 
+{
+	for(UObject* Object : RuntimeObjects)
+	{
+		FWidgetTransform TransformValue = PropertyBindings->GetCurrentValue<FWidgetTransform>(Object);
+		if(TransformTrack->Eval(Position, LastPosition, TransformValue))
+		{
+			PropertyBindings->CallFunction(Object, &TransformValue);
+		}
+	}
+}
+
+void FMovieScene2DTransformTrackInstance::RefreshInstance(const TArray<UObject*>& RuntimeObjects, class IMovieScenePlayer& Player)
+{
+	PropertyBindings->UpdateBindings(RuntimeObjects);
+}
+
diff --git a/Engine/Source/Runtime/UMG/Private/Animation/MovieScene2DTransformTrackInstance.h b/Engine/Source/Runtime/UMG/Private/Animation/MovieScene2DTransformTrackInstance.h
new file mode 100644
--- /dev/null
+++ b/Engine/Source/Runtime/UMG/Private/Animation/MovieScene2DTransformTrackInstance.h
@@ -0,0 +1,23 @@
+// Copyright 1998-2014 Epic Games, Inc. All Rights Reserved.
+
+#pragma once
+
+#include "IMovieSceneTrackInstance.h"
+
+/**
+ * Instance of a UMovieScene2DTransformTrack
+ */
+class FMovieScene2DTransformTrackInstance : public IMovieSceneTrackInstance
+{
+public:
+	FMovieScene2DTransformTrackInstance( class UMovieScene2DTransformTrack& InTransformTrack );
+
+	/** IMovieSceneTrackInstance interface */
+	virtual void Update( float Position, float LastPosition, const TArray<UObject*>& RuntimeObjects, class IMovieScenePlayer& Player ) override;
+	virtual void RefreshInstance( const TArray<UObject*>& RuntimeObjects, class IMovieScenePlayer& Player ) override;
+private:
+	/** Track that is being instanced */
+	UMovieScene2DTransformTrack* TransformTrack;
+	/** Runtime property bindings */
+	TSharedPtr<class FTrackInstancePropertyBindings> PropertyBindings;
+};
diff --git a/Engine/Source/Runtime/UMG/Private/Animation/MovieSceneMarginSection.cpp b/Engine/Source/Runtime/UMG/Private/Animation/MovieSceneMarginSection.cpp
--- a/Engine/Source/Runtime/UMG/Private/Animation/MovieSceneMarginSection.cpp
+++ b/Engine/Source/Runtime/UMG/Private/Animation/MovieSceneMarginSection.cpp
@@ -91,10 +91,3 @@ void UMovieSceneMarginSection::AddKeyToNamedCurve( float Time, const FMarginKey&
 	}
 }
 
-void UMovieSceneMarginSection::AddKeyToCurve( FRichCurve& InCurve, float Time, float Value )
-{
-	if( IsTimeWithinSection(Time) )
-	{
-		InCurve.UpdateOrAddKey(Time, Value);
-	}
-}
diff --git a/Engine/Source/Runtime/UMG/Public/Animation/MovieScene2DTransformSection.h b/Engine/Source/Runtime/UMG/Public/Animation/MovieScene2DTransformSection.h
new file mode 100644
--- /dev/null
+++ b/Engine/Source/Runtime/UMG/Public/Animation/MovieScene2DTransformSection.h
@@ -0,0 +1,53 @@
+// Copyright 1998-2014 Epic Games, Inc. All Rights Reserved.
+
+#pragma once
+
+#include "MovieSceneSection.h"
+
+#include "MovieScene2DTransformSection.generated.h"
+
+/**
+ * A transform section
+ */
+UCLASS(MinimalAPI)
+class UMovieScene2DTransformSection : public UMovieSceneSection
+{
+	GENERATED_UCLASS_BODY()
+public:
+
+	/** MovieSceneSection interface */
+	virtual void MoveSection( float DeltaPosition ) override;
+	virtual void DilateSection( float DilationFactor, float Origin ) override;
+
+	UMG_API FRichCurve& GetTranslationCurve( EAxis::Type Axis );
+
+	UMG_API FRichCurve& GetRotationCurve();
+
+	UMG_API FRichCurve& GetScaleCurve( EAxis::Type Axis );
+
+	UMG_API FRichCurve& GetSheerCurve( EAxis::Type Axis );
+
+	FWidgetTransform Eval( float Position, const FWidgetTransform& DefaultValue ) const;
+
+	bool NewKeyIsNewData( float Time, const FWidgetTransform& Transform ) const;
+
+	void AddKey( float Time, const struct F2DTransformKey& TransformKey );
+private:
+	void AddKeyToNamedCurve(float Time, const F2DTransformKey& TransformKey);
+private:
+	/** Translation curves*/
+	UPROPERTY()
+	FRichCurve Translation[2];
+	
+	/** Rotation curve */
+	UPROPERTY()
+	FRichCurve Rotation;
+
+	/** Scale curves */
+	UPROPERTY()
+	FRichCurve Scale[2];
+
+	/** Shear curve */
+	UPROPERTY()
+	FRichCurve Shear[2];
+};
\ No newline at end of file
diff --git a/Engine/Source/Runtime/UMG/Public/Animation/MovieScene2DTransformTrack.h b/Engine/Source/Runtime/UMG/Public/Animation/MovieScene2DTransformTrack.h
new file mode 100644
--- /dev/null
+++ b/Engine/Source/Runtime/UMG/Public/Animation/MovieScene2DTransformTrack.h
@@ -0,0 +1,42 @@
+// Copyright 1998-2014 Epic Games, Inc. All Rights Reserved.
+
+#pragma once
+
+#include "MovieScenePropertyTrack.h"
+#include "MovieScene2DTransformTrack.generated.h"
+
+struct F2DTransformKey
+{
+	FWidgetTransform Value;
+	FName CurveName;
+	bool bAddKeyEvenIfUnchanged;
+};
+
+
+/**
+ * Handles manipulation of 2D transforms in a movie scene
+ */
+UCLASS( MinimalAPI )
+class UMovieScene2DTransformTrack : public UMovieScenePropertyTrack
+{
+	GENERATED_UCLASS_BODY()
+public:
+	/** UMovieSceneTrack interface */
+	virtual UMovieSceneSection* CreateNewSection() override;
+	virtual TSharedPtr<IMovieSceneTrackInstance> CreateInstance() override;
+	
+
+	/**
+	 * Evaluates the track at the playback position
+	 *
+	 * @param Position	The current playback position
+	 * @param LastPosition	The last playback position
+	 * @param InOutTransform 	The transform at the playback position
+	 * @return true if anything was evaluated. Note: if false is returned InOutTransform remains unchanged
+	 */
+	bool Eval( float Position, float LastPostion, FWidgetTransform& InOutTransform ) const;
+
+
+	UMG_API bool AddKeyToSection( float Time, const F2DTransformKey& TransformKey );
+
+};
\ No newline at end of file
diff --git a/Engine/Source/Runtime/UMG/Public/Animation/MovieSceneMarginSection.h b/Engine/Source/Runtime/UMG/Public/Animation/MovieSceneMarginSection.h
--- a/Engine/Source/Runtime/UMG/Public/Animation/MovieSceneMarginSection.h
+++ b/Engine/Source/Runtime/UMG/Public/Animation/MovieSceneMarginSection.h
@@ -76,14 +76,7 @@ public:
 	const FRichCurve& GetBottomCurve() const { return BottomCurve; }
 private:
 	void AddKeyToNamedCurve( float Time, const FMarginKey& MarginKey );
-	/**
-	 * Adds a key to a rich curve, finding an existing key to modify or adding a new one
-	 *
-	 * @param InCurve	The curve to add keys to
-	 * @param Time		The time where the key should be added
-	 * @param Value		The value at the given time
-	 */
-	void AddKeyToCurve( FRichCurve& InCurve, float Time, float Value );
+
 private:
 	/** Red curve data */
 	UPROPERTY()
diff --git a/Engine/Source/Runtime/UMG/Public/Animation/MovieSceneMarginTrack.h b/Engine/Source/Runtime/UMG/Public/Animation/MovieSceneMarginTrack.h
--- a/Engine/Source/Runtime/UMG/Public/Animation/MovieSceneMarginTrack.h
+++ b/Engine/Source/Runtime/UMG/Public/Animation/MovieSceneMarginTrack.h
@@ -32,7 +32,7 @@ public:
 	 * @param Value				The value of the key
 	 * @return True if the key was successfully added.
 	 */
-	virtual bool AddKeyToSection( float Time, const FMarginKey& MarginKey );
+	UMG_API bool AddKeyToSection( float Time, const FMarginKey& MarginKey );
 
 	/**
 	 * Evaluates the track at the playback position
diff --git a/Engine/Source/Runtime/UMG/Public/Components/Widget.h b/Engine/Source/Runtime/UMG/Public/Components/Widget.h
--- a/Engine/Source/Runtime/UMG/Public/Components/Widget.h
+++ b/Engine/Source/Runtime/UMG/Public/Components/Widget.h
@@ -107,11 +107,11 @@ public:
 	FGetMouseCursor CursorDelegate;
 
 	/**  */
-	UPROPERTY(EditDefaultsOnly, Category=Transform)
+	UPROPERTY(EditDefaultsOnly, Category=RenderTransform, meta=(DisplayName="Transform"))
 	FWidgetTransform RenderTransform;
 
 	/**  */
-	UPROPERTY(EditDefaultsOnly, Category=Transform)
+	UPROPERTY(EditDefaultsOnly, Category=RenderTransform, meta=(DisplayName="Pivot"))
 	FVector2D RenderTransformPivot;
 
 	/** */
diff --git a/Engine/Source/Runtime/UMG/Public/Slate/WidgetTransform.h b/Engine/Source/Runtime/UMG/Public/Slate/WidgetTransform.h
--- a/Engine/Source/Runtime/UMG/Public/Slate/WidgetTransform.h
+++ b/Engine/Source/Runtime/UMG/Public/Slate/WidgetTransform.h
@@ -12,6 +12,10 @@ struct FWidgetTransform
 {
 	GENERATED_USTRUCT_BODY()
 
+	/** The amount to translate the widget in slate units */
+	UPROPERTY(EditAnywhere, BlueprintReadWrite, Category=Transform, meta=(Delta = "1"))
+	FVector2D Translation;
+
 	/** The scale to apply to the widget */
 	UPROPERTY(EditAnywhere, BlueprintReadWrite, Category=Transform, meta=( UIMin = "-5", UIMax = "5", Delta = "0.05" ))
 	FVector2D Scale;
@@ -24,15 +28,30 @@ struct FWidgetTransform
 	UPROPERTY(EditAnywhere, BlueprintReadWrite, Category=Transform, meta=( UIMin = "-180", UIMax = "180", Delta = "1" ))
 	float Angle;
 	
-	/** The amount to translate the widget in slate units */
-	UPROPERTY(EditAnywhere, BlueprintReadWrite, Category=Transform, meta=( Delta = "1" ))
-	FVector2D Translation;
 
 	FWidgetTransform()
-		: Scale(1, 1)
+		: Translation(0, 0)
+		, Scale(1, 1)
 		, Shear(0, 0)
 		, Angle(0)
-		, Translation(0, 0)
 	{
 	}
+
+	FWidgetTransform( const FVector2D& InTranslation, const FVector2D& InScale, const FVector2D& InShear, float InAngle )
+		: Translation( InTranslation )
+		, Scale( InScale )
+		, Shear( InShear )
+		, Angle( InAngle )
+	{
+	}
+
+	bool operator==( const FWidgetTransform &Other ) const
+	{
+		return Scale == Other.Scale && Shear == Other.Shear && Angle == Other.Angle && Translation == Other.Translation;
+	}
+
+	bool operator!=( const FWidgetTransform& Other ) const
+	{
+		return !(*this == Other);
+	}
 };
