UMG - Smart naming now works as long as the name of the widget is non-generated, instead of it requiring to be a variable.

diff --git a/Engine/Source/Editor/UMGEditor/Private/Customizations/DetailCustomizations.cpp b/Engine/Source/Editor/UMGEditor/Private/Customizations/DetailCustomizations.cpp
--- a/Engine/Source/Editor/UMGEditor/Private/Customizations/DetailCustomizations.cpp
+++ b/Engine/Source/Editor/UMGEditor/Private/Customizations/DetailCustomizations.cpp
@@ -711,7 +711,7 @@ void FBlueprintWidgetCustomization::HandleCreateAndAddBinding(UWidget* Widget, T
 	FString Pre = bIsPure ? FString(TEXT("Get")) : FString(TEXT("On"));
 
 	FString WidgetName;
-	if ( Widget && Widget->bIsVariable )
+	if ( Widget && !Widget->IsGeneratedName() )
 	{
 		WidgetName = TEXT("_") + Widget->GetName() + TEXT("_");
 	}
