UMG - The hierarchy view is a lot more stable now, it can restore the proper expansion between compiles.

diff --git a/Engine/Source/Editor/UMGEditor/Private/Hierarchy/SHierarchyView.cpp b/Engine/Source/Editor/UMGEditor/Private/Hierarchy/SHierarchyView.cpp
--- a/Engine/Source/Editor/UMGEditor/Private/Hierarchy/SHierarchyView.cpp
+++ b/Engine/Source/Editor/UMGEditor/Private/Hierarchy/SHierarchyView.cpp
@@ -87,17 +87,23 @@ SHierarchyView::~SHierarchyView()
 
 void SHierarchyView::Tick(const FGeometry& AllottedGeometry, const double InCurrentTime, const float InDeltaTime)
 {
-	if ( bRebuildTreeRequested )
+	if ( bRebuildTreeRequested || bRefreshRequested )
 	{
-		bRebuildTreeRequested = false;
-		RebuildTreeView();
-	}
+		if ( bRebuildTreeRequested )
+		{
+			RebuildTreeView();
+		}
 
-	if ( bRefreshRequested )
-	{
-		bRefreshRequested = false;
+		SaveExpandedItems();
 
 		RefreshTree();
+
+		RestoreExpandedItems();
+
+		OnEditorSelectionChanged();
+
+		bRefreshRequested = false;
+		bRebuildTreeRequested = false;
 	}
 }
 
@@ -214,8 +220,6 @@ TSharedPtr<SWidget> SHierarchyView::WidgetHierarchy_OnContextMenuOpening()
 
 void SHierarchyView::WidgetHierarchy_OnGetChildren(TSharedPtr<FHierarchyModel> InParent, TArray< TSharedPtr<FHierarchyModel> >& OutChildren)
 {
-	VisibleItems.Add(InParent);
-
 	InParent->GatherChildren(OutChildren);
 }
 
@@ -236,13 +240,7 @@ void SHierarchyView::WidgetHierarchy_OnSelectionChanged(TSharedPtr<FHierarchyMod
 FReply SHierarchyView::HandleDeleteSelected()
 {
 	TSet<FWidgetReference> SelectedWidgets = BlueprintEditor.Pin()->GetSelectedWidgets();
-
-	// Remove the selected items from the filter cache
-	for (FWidgetReference& Item : SelectedWidgets)
-	{
-//		FilterHandler->RemoveCachedItem(Item.GetTemplate());
-	}
-
+	
 	FWidgetBlueprintEditorUtils::DeleteWidgets(GetBlueprint(), SelectedWidgets);
 
 	return FReply::Handled();
@@ -250,8 +248,6 @@ FReply SHierarchyView::HandleDeleteSelected()
 
 void SHierarchyView::RefreshTree()
 {
-	VisibleItems.Empty();
-
 	RootWidgets.Empty();
 	RootWidgets.Add( MakeShareable(new FHierarchyRoot(BlueprintEditor.Pin())) );
 
@@ -282,13 +278,50 @@ void SHierarchyView::OnObjectsReplaced(const TMap<UObject*, UObject*>& Replaceme
 {
 	if ( !bRebuildTreeRequested )
 	{
-		for ( TSharedPtr<FHierarchyModel> Widget : VisibleItems )
+		bRefreshRequested = true;
+		bRebuildTreeRequested = true;
+
+		// We save the expanded items immediately because they're potentially about to become invalid.
+		SaveExpandedItems();
+	}
+}
+
+void SHierarchyView::SaveExpandedItems()
+{
+	if ( ExpandedItems.Num() == 0 )
+	{
+		TSet < TSharedPtr<FHierarchyModel> > ExpandedModels;
+		WidgetTreeView->GetExpandedItems(ExpandedModels);
+
+		for ( TSharedPtr<FHierarchyModel>& Model : ExpandedModels )
 		{
-			//if ( ReplacementMap.Contains(Widget) )
-			{
-				bRefreshRequested = true;
-				bRebuildTreeRequested = true;
-			}
+			ExpandedItems.Add(Model->GetUniqueName());
+		}
+	}
+}
+
+void SHierarchyView::RestoreExpandedItems()
+{
+	for ( TSharedPtr<FHierarchyModel>& Model : RootWidgets )
+	{
+		RecursiveExpand(Model);
+	}
+
+	ExpandedItems.Empty();
+}
+
+void SHierarchyView::RecursiveExpand(TSharedPtr<FHierarchyModel>& Model)
+{
+	if ( ExpandedItems.Contains(Model->GetUniqueName()) )
+	{
+		WidgetTreeView->SetItemExpansion(Model, true);
+
+		TArray< TSharedPtr<FHierarchyModel> > Children;
+		Model->GatherChildren(Children);
+
+		for ( TSharedPtr<FHierarchyModel>& ChildModel : Children )
+		{
+			RecursiveExpand(ChildModel);
 		}
 	}
 }
diff --git a/Engine/Source/Editor/UMGEditor/Private/Hierarchy/SHierarchyView.h b/Engine/Source/Editor/UMGEditor/Private/Hierarchy/SHierarchyView.h
--- a/Engine/Source/Editor/UMGEditor/Private/Hierarchy/SHierarchyView.h
+++ b/Engine/Source/Editor/UMGEditor/Private/Hierarchy/SHierarchyView.h
@@ -79,6 +79,15 @@ private:
 	/** Called when a Blueprint is recompiled and live objects are swapped out for replacements */
 	void OnObjectsReplaced(const TMap<UObject*, UObject*>& ReplacementMap);
 
+	/** Saves the expansion of items into a list based on FName */
+	void SaveExpandedItems();
+
+	/** Restores the state of expanded items based on the saved expanded item state, then clears the expanded state cache. */
+	void RestoreExpandedItems();
+
+	/** Recursively expands the models based on the expansion set. */
+	void RecursiveExpand(TSharedPtr<FHierarchyModel>& Model);
+
 private:
 
 	/** Cached pointer to the blueprint editor that owns this tree. */
@@ -93,9 +102,6 @@ private:
 	/** The root widgets which are actually displayed by the TreeView which will be managed by the TreeFilterHandler. */
 	TArray< TSharedPtr<FHierarchyModel> > TreeRootWidgets;
 
-	/** Gets the visible items displayed in the tree view. */
-	TSet< TSharedPtr<FHierarchyModel> > VisibleItems;
-
 	/** The widget containing the treeview */
 	TSharedPtr<SBorder> TreeViewArea;
 
@@ -110,4 +116,7 @@ private:
 
 	/** Is the tree in such a changed state that the whole widget needs rebuilding? */
 	bool bRebuildTreeRequested;
+
+	/** Temporary expanded item state, used to restore expansion after tree rebuilds. */
+	TArray<FName> ExpandedItems;
 };
diff --git a/Engine/Source/Editor/UMGEditor/Private/Hierarchy/SHierarchyViewItem.cpp b/Engine/Source/Editor/UMGEditor/Private/Hierarchy/SHierarchyViewItem.cpp
--- a/Engine/Source/Editor/UMGEditor/Private/Hierarchy/SHierarchyViewItem.cpp
+++ b/Engine/Source/Editor/UMGEditor/Private/Hierarchy/SHierarchyViewItem.cpp
@@ -259,6 +259,12 @@ FHierarchyRoot::FHierarchyRoot(TSharedPtr<FWidgetBlueprintEditor> InBlueprintEdi
 {
 }
 
+FName FHierarchyRoot::GetUniqueName() const
+{
+	static const FName DesignerRootName(TEXT("WidgetDesignerRoot"));
+	return DesignerRootName;
+}
+
 FText FHierarchyRoot::GetText() const
 {
 	return LOCTEXT("Root", "[Root]");
@@ -319,6 +325,17 @@ FHierarchyWidget::FHierarchyWidget(FWidgetReference InItem, TSharedPtr<FWidgetBl
 {
 }
 
+FName FHierarchyWidget::GetUniqueName() const
+{
+	UWidget* WidgetTemplate = Item.GetTemplate();
+	if ( WidgetTemplate )
+	{
+		return WidgetTemplate->GetFName();
+	}
+
+	return NAME_None;
+}
+
 FText FHierarchyWidget::GetText() const
 {
 	UWidget* WidgetTemplate = Item.GetTemplate();
diff --git a/Engine/Source/Editor/UMGEditor/Private/Hierarchy/SHierarchyViewItem.h b/Engine/Source/Editor/UMGEditor/Private/Hierarchy/SHierarchyViewItem.h
--- a/Engine/Source/Editor/UMGEditor/Private/Hierarchy/SHierarchyViewItem.h
+++ b/Engine/Source/Editor/UMGEditor/Private/Hierarchy/SHierarchyViewItem.h
@@ -10,6 +10,9 @@ class FHierarchyModel : public TSharedFromThis < FHierarchyModel >
 public:
 	FHierarchyModel();
 
+	/** Gets the unique name of the item used to restore item expansion. */
+	virtual FName GetUniqueName() const = 0;
+
 	/* @returns the widget name to use for the tree item */
 	virtual FText GetText() const = 0;
 
@@ -50,6 +53,8 @@ public:
 	
 	virtual ~FHierarchyRoot() {}
 
+	virtual FName GetUniqueName() const override;
+
 	/* @returns the widget name to use for the tree item */
 	virtual FText GetText() const override;
 
@@ -77,6 +82,8 @@ public:
 	FHierarchyWidget(FWidgetReference InItem, TSharedPtr<FWidgetBlueprintEditor> InBlueprintEditor);
 
 	virtual ~FHierarchyWidget() {}
+
+	virtual FName GetUniqueName() const override;
 	
 	/* @returns the widget name to use for the tree item */
 	virtual FText GetText() const override;
