Fix crash creating the palette view

diff --git a/Engine/Source/Editor/UMGEditor/Private/Palette/SPaletteView.cpp b/Engine/Source/Editor/UMGEditor/Private/Palette/SPaletteView.cpp
--- a/Engine/Source/Editor/UMGEditor/Private/Palette/SPaletteView.cpp
+++ b/Engine/Source/Editor/UMGEditor/Private/Palette/SPaletteView.cpp
@@ -213,6 +213,9 @@ SPaletteView::~SPaletteView()
 		FilterHandler->SetIsEnabled(false);
 		FilterHandler->RefreshAndFilterTree();
 	}
+
+	GEditor->OnObjectsReplaced().RemoveAll( this );
+
 	SaveItemExpansion();
 }
 
