UMG - The Fill/Auto size customization now uses a uniform panel for better sizing of the buttons, the tooltips have also been improved.

diff --git a/Engine/Source/Editor/UMGEditor/Private/Customizations/SlateChildSizeCustomization.cpp b/Engine/Source/Editor/UMGEditor/Private/Customizations/SlateChildSizeCustomization.cpp
--- a/Engine/Source/Editor/UMGEditor/Private/Customizations/SlateChildSizeCustomization.cpp
+++ b/Engine/Source/Editor/UMGEditor/Private/Customizations/SlateChildSizeCustomization.cpp
@@ -38,31 +38,37 @@ void FSlateChildSizeCustomization::CustomizeHeader(TSharedRef<IPropertyHandle> P
 		.AutoWidth()
 		.Padding(OuterPadding)
 		[
-			SNew(SCheckBox)
-			.Style(FEditorStyle::Get(), "ToggleButtonCheckbox")
-			.ToolTipText(LOCTEXT("Auto_ToolTip", "Auto"))
-			.Padding(ContentPadding)
-			.OnCheckStateChanged(this, &FSlateChildSizeCustomization::HandleCheckStateChanged, RuleHandle, ESlateSizeRule::Automatic)
-			.IsChecked(this, &FSlateChildSizeCustomization::GetCheckState, RuleHandle, ESlateSizeRule::Automatic)
+			SNew(SUniformGridPanel)
+			.SlotPadding(OuterPadding)
+
+			+ SUniformGridPanel::Slot(0, 0)
 			[
-				SNew(STextBlock)
-				.Text(LOCTEXT("Auto", "Auto"))
+				SNew(SCheckBox)
+				.Style(FEditorStyle::Get(), "ToggleButtonCheckbox")
+				.ToolTipText(LOCTEXT("Auto_ToolTip", "Only requests as much room as it needs based on the widgets desired size."))
+				.Padding(ContentPadding)
+				.OnCheckStateChanged(this, &FSlateChildSizeCustomization::HandleCheckStateChanged, RuleHandle, ESlateSizeRule::Automatic)
+				.IsChecked(this, &FSlateChildSizeCustomization::GetCheckState, RuleHandle, ESlateSizeRule::Automatic)
+				.HAlign(HAlign_Center)
+				[
+					SNew(STextBlock)
+					.Text(LOCTEXT("Auto", "Auto"))
+				]
 			]
-		]
 
-		+ SHorizontalBox::Slot()
-		.AutoWidth()
-		.Padding(OuterPadding)
-		[
-			SNew(SCheckBox)
-			.Style(FEditorStyle::Get(), "ToggleButtonCheckbox")
-			.ToolTipText(LOCTEXT("Fill_ToolTip", "Fill"))
-			.Padding(ContentPadding)
-			.OnCheckStateChanged(this, &FSlateChildSizeCustomization::HandleCheckStateChanged, RuleHandle, ESlateSizeRule::Fill)
-			.IsChecked(this, &FSlateChildSizeCustomization::GetCheckState, RuleHandle, ESlateSizeRule::Fill)
+			+ SUniformGridPanel::Slot(1, 0)
 			[
-				SNew(STextBlock)
-				.Text(LOCTEXT("Fill", "Fill"))
+				SNew(SCheckBox)
+				.Style(FEditorStyle::Get(), "ToggleButtonCheckbox")
+				.ToolTipText(LOCTEXT("Fill_ToolTip", "Greedily attempts to fill all available room based on the percentage value 0..1"))
+				.Padding(ContentPadding)
+				.OnCheckStateChanged(this, &FSlateChildSizeCustomization::HandleCheckStateChanged, RuleHandle, ESlateSizeRule::Fill)
+				.IsChecked(this, &FSlateChildSizeCustomization::GetCheckState, RuleHandle, ESlateSizeRule::Fill)
+				.HAlign(HAlign_Center)
+				[
+					SNew(STextBlock)
+					.Text(LOCTEXT("Fill", "Fill"))
+				]
 			]
 		]
 
diff --git a/Engine/Source/Runtime/UMG/Public/Components/SlateWrapperTypes.h b/Engine/Source/Runtime/UMG/Public/Components/SlateWrapperTypes.h
--- a/Engine/Source/Runtime/UMG/Public/Components/SlateWrapperTypes.h
+++ b/Engine/Source/Runtime/UMG/Public/Components/SlateWrapperTypes.h
@@ -36,9 +36,9 @@ namespace ESlateSizeRule
 {
 	enum Type
 	{
-		/** The container will size to fit the needs of the child widgets */
+		/** Only requests as much room as it needs based on the widgets desired size. */
 		Automatic,
-		/** The container will fill the percentage of the container based on the Value 0..1 */
+		/** Greedily attempts to fill all available room based on the percentage value 0..1 */
 		Fill
 	};
 }
