UMG - Adding more zoom levels to the designer.

diff --git a/Engine/Source/Editor/UMGEditor/Private/Designer/SDesignSurface.cpp b/Engine/Source/Editor/UMGEditor/Private/Designer/SDesignSurface.cpp
--- a/Engine/Source/Editor/UMGEditor/Private/Designer/SDesignSurface.cpp
+++ b/Engine/Source/Editor/UMGEditor/Private/Designer/SDesignSurface.cpp
@@ -50,6 +50,9 @@ struct FFixedZoomLevelsContainer : public FZoomLevelsContainer
 			ZoomLevels.Add(FZoomLevelEntry(1.750f, LOCTEXT("ZoomLevel", "+5"), EGraphRenderingLOD::FullyZoomedIn));
 			ZoomLevels.Add(FZoomLevelEntry(1.875f, LOCTEXT("ZoomLevel", "+6"), EGraphRenderingLOD::FullyZoomedIn));
 			ZoomLevels.Add(FZoomLevelEntry(2.000f, LOCTEXT("ZoomLevel", "+7"), EGraphRenderingLOD::FullyZoomedIn));
+			ZoomLevels.Add(FZoomLevelEntry(2.250f, LOCTEXT("ZoomLevel", "+8"), EGraphRenderingLOD::FullyZoomedIn));
+			ZoomLevels.Add(FZoomLevelEntry(2.375f, LOCTEXT("ZoomLevel", "+9"), EGraphRenderingLOD::FullyZoomedIn));
+			ZoomLevels.Add(FZoomLevelEntry(2.500f, LOCTEXT("ZoomLevel", "+10"), EGraphRenderingLOD::FullyZoomedIn));
 		}
 	}
 
