UMG - Details panel improvements, the Is Variable checkbox is now a toggle button.  There is now a class link widget so that you can quickly jump to user widgets/source files.

diff --git a/Engine/Source/Editor/UMGEditor/Private/Details/SWidgetDetailsView.cpp b/Engine/Source/Editor/UMGEditor/Private/Details/SWidgetDetailsView.cpp
--- a/Engine/Source/Editor/UMGEditor/Private/Details/SWidgetDetailsView.cpp
+++ b/Engine/Source/Editor/UMGEditor/Private/Details/SWidgetDetailsView.cpp
@@ -8,6 +8,7 @@
 #include "Kismet2NameValidators.h"
 #include "ISequencer.h"
 #include "Animation/UMGDetailKeyframeHandler.h"
+#include "EditorClassUtils.h"
 
 #include "CanvasSlotCustomization.h"
 #include "HorizontalAlignmentCustomization.h"
@@ -76,13 +77,23 @@ void SWidgetDetailsView::Construct(const FArguments& InArgs, TSharedPtr<FWidgetB
 			.AutoWidth()
 			[
 				SNew(SCheckBox)
+				.Style(FEditorStyle::Get(), "ToggleButtonCheckbox")
 				.IsChecked(this, &SWidgetDetailsView::GetIsVariable)
 				.OnCheckStateChanged(this, &SWidgetDetailsView::HandleIsVariableChanged)
+				.Padding(FMargin(3,1,3,1))
 				[
 					SNew(STextBlock)
-					.Text(LOCTEXT("IsVariable", "Is Variable"))
+					.Font(FEditorStyle::GetFontStyle("BoldFont"))
+					.Text(LOCTEXT("IsVariable", "Variable"))
 				]
 			]
+
+			+ SHorizontalBox::Slot()
+			.AutoWidth()
+			.Padding(10,0,0,0)
+			[
+				SAssignNew(ClassLinkArea, SBox)
+			]
 		]
 
 		+ SVerticalBox::Slot()
@@ -159,6 +170,15 @@ void SWidgetDetailsView::OnEditorSelectionChanged()
 		}
 	}
 
+	if ( SelectedObjects.Num() == 1 )
+	{
+		ClassLinkArea->SetContent(FEditorClassUtils::GetSourceLink(SelectedObjects[0]->GetClass(), TWeakObjectPtr<UObject>()));
+	}
+	else
+	{
+		ClassLinkArea->SetContent(SNullWidget::NullWidget);
+	}
+
 	const bool bForceRefresh = false;
 	PropertyView->SetObjects(SelectedObjects, bForceRefresh);
 }
diff --git a/Engine/Source/Editor/UMGEditor/Private/Details/SWidgetDetailsView.h b/Engine/Source/Editor/UMGEditor/Private/Details/SWidgetDetailsView.h
--- a/Engine/Source/Editor/UMGEditor/Private/Details/SWidgetDetailsView.h
+++ b/Engine/Source/Editor/UMGEditor/Private/Details/SWidgetDetailsView.h
@@ -58,7 +58,11 @@ private:
 	/** The editor that owns this details view */
 	TWeakPtr<class FWidgetBlueprintEditor> BlueprintEditor;
 
+	/** The name text box that users can use to rename their widgets */
 	TSharedPtr<SEditableTextBox> NameTextBox;
+
+	/** The container widget for the class link users can click to open another asset */
+	TSharedPtr<SBox> ClassLinkArea;
 	
 	/** Property viewing widget */
 	TSharedPtr<class IDetailsView> PropertyView;
