UMG - Going back to the is variable check box in the details view.

diff --git a/Engine/Source/Editor/UMGEditor/Private/Details/SWidgetDetailsView.cpp b/Engine/Source/Editor/UMGEditor/Private/Details/SWidgetDetailsView.cpp
--- a/Engine/Source/Editor/UMGEditor/Private/Details/SWidgetDetailsView.cpp
+++ b/Engine/Source/Editor/UMGEditor/Private/Details/SWidgetDetailsView.cpp
@@ -77,20 +77,20 @@ void SWidgetDetailsView::Construct(const FArguments& InArgs, TSharedPtr<FWidgetB
 			.AutoWidth()
 			[
 				SNew(SCheckBox)
-				.Style(FEditorStyle::Get(), "ToggleButtonCheckbox")
+				//.Style(FEditorStyle::Get(), "ToggleButtonCheckbox")
 				.IsChecked(this, &SWidgetDetailsView::GetIsVariable)
 				.OnCheckStateChanged(this, &SWidgetDetailsView::HandleIsVariableChanged)
 				.Padding(FMargin(3,1,3,1))
 				[
 					SNew(STextBlock)
-					.Font(FEditorStyle::GetFontStyle("BoldFont"))
-					.Text(LOCTEXT("IsVariable", "Variable"))
+					//.Font(FEditorStyle::GetFontStyle("BoldFont"))
+					.Text(LOCTEXT("IsVariable", "Is Variable"))
 				]
 			]
 
 			+ SHorizontalBox::Slot()
 			.AutoWidth()
-			.Padding(10,0,0,0)
+			.Padding(15,0,0,0)
 			[
 				SAssignNew(ClassLinkArea, SBox)
 			]
