Fix crash adding a new animation and not changing the name

diff --git a/Engine/Source/Editor/UMGEditor/Private/TabFactory/AnimationTabSummoner.cpp b/Engine/Source/Editor/UMGEditor/Private/TabFactory/AnimationTabSummoner.cpp
--- a/Engine/Source/Editor/UMGEditor/Private/TabFactory/AnimationTabSummoner.cpp
+++ b/Engine/Source/Editor/UMGEditor/Private/TabFactory/AnimationTabSummoner.cpp
@@ -106,9 +106,10 @@ private:
 
 		FName CurrentName = WidgetAnimation->GetFName();
 
+		bool bNewAnimation = ListItem.Pin()->bNewAnimation;
+
 		if(!CurrentName.ToString().Equals(InText.ToString()) && !InText.IsEmpty())
 		{
-			bool bNewAnimation = ListItem.Pin()->bNewAnimation;
 			FText TransactionName = bNewAnimation ? LOCTEXT("NewAnimation", "New Animation") : LOCTEXT("RenameAnimation", "Rename Animation");
 
 			{
@@ -131,6 +132,13 @@ private:
 
 			FBlueprintEditorUtils::MarkBlueprintAsStructurallyModified(Blueprint);
 		}
+		else if( bNewAnimation )
+		{
+			const FScopedTransaction Transaction(LOCTEXT("NewAnimation", "New Animation"));
+			Blueprint->Modify();
+			Blueprint->Animations.Add(WidgetAnimation);
+			ListItem.Pin()->bNewAnimation = false;
+		}
 	}
 private:
 	TWeakPtr<FWidgetAnimationListItem> ListItem;
