UMG - The name field in the details panel now always shows the FName, instead of the Label.

diff --git a/Engine/Source/Editor/UMGEditor/Private/Details/SWidgetDetailsView.cpp b/Engine/Source/Editor/UMGEditor/Private/Details/SWidgetDetailsView.cpp
--- a/Engine/Source/Editor/UMGEditor/Private/Details/SWidgetDetailsView.cpp
+++ b/Engine/Source/Editor/UMGEditor/Private/Details/SWidgetDetailsView.cpp
@@ -234,14 +234,13 @@ FText SWidgetDetailsView::GetNameText() const
 		UWidget* Widget = Cast<UWidget>(SelectedObjects[0].Get());
 		if ( Widget )
 		{
-			return FText::FromString(Widget->GetLabel());
+			return FText::FromName(Widget->GetFName());
 		}
 	}
 	
 	return FText::GetEmpty();
 }
 
-
 void SWidgetDetailsView::HandleNameTextChanged(const FText& Text)
 {
 	FText OutErrorMessage;
