UMG - We no use the publically exposed slate brush customization in UMG to override it so that the preview is not displayed in UMG, but still displayed for brush assets.

diff --git a/Engine/Source/Editor/DetailCustomizations/Private/DetailCustomizations.cpp b/Engine/Source/Editor/DetailCustomizations/Private/DetailCustomizations.cpp
--- a/Engine/Source/Editor/DetailCustomizations/Private/DetailCustomizations.cpp
+++ b/Engine/Source/Editor/DetailCustomizations/Private/DetailCustomizations.cpp
@@ -151,7 +151,7 @@ void FDetailCustomizationsModule::RegisterPropertyTypeCustomizations()
 	RegisterCustomPropertyTypeLayout("DialogueWaveParameter", FOnGetPropertyTypeCustomizationInstance::CreateStatic(&FDialogueWaveParameterStructCustomization::MakeInstance));
 	RegisterCustomPropertyTypeLayout("BodyInstance", FOnGetPropertyTypeCustomizationInstance::CreateStatic(&FBodyInstanceCustomization::MakeInstance));
 	RegisterCustomPropertyTypeLayout("VehicleTransmissionData", FOnGetPropertyTypeCustomizationInstance::CreateStatic(&FVehicleTransmissionDataCustomization::MakeInstance));
-	RegisterCustomPropertyTypeLayout("SlateBrush", FOnGetPropertyTypeCustomizationInstance::CreateStatic(&FSlateBrushStructCustomization::MakeInstance));
+	RegisterCustomPropertyTypeLayout("SlateBrush", FOnGetPropertyTypeCustomizationInstance::CreateStatic(&FSlateBrushStructCustomization::MakeInstance, true));
 	RegisterCustomPropertyTypeLayout("SlateSound", FOnGetPropertyTypeCustomizationInstance::CreateStatic(&FSlateSoundStructCustomization::MakeInstance));
 	RegisterCustomPropertyTypeLayout("Guid", FOnGetPropertyTypeCustomizationInstance::CreateStatic(&FGuidStructCustomization::MakeInstance));
 	RegisterCustomPropertyTypeLayout("Key", FOnGetPropertyTypeCustomizationInstance::CreateStatic(&FKeyStructCustomization::MakeInstance));
diff --git a/Engine/Source/Editor/DetailCustomizations/Private/SlateBrushCustomization.cpp b/Engine/Source/Editor/DetailCustomizations/Private/SlateBrushCustomization.cpp
--- a/Engine/Source/Editor/DetailCustomizations/Private/SlateBrushCustomization.cpp
+++ b/Engine/Source/Editor/DetailCustomizations/Private/SlateBrushCustomization.cpp
@@ -1114,13 +1114,13 @@ private:
 // FSlateBrushStructCustomization
 ////////////////////////////////////////////////////////////////////////////////
 
-TSharedRef<IPropertyTypeCustomization> FSlateBrushStructCustomization::MakeInstance() 
+TSharedRef<IPropertyTypeCustomization> FSlateBrushStructCustomization::MakeInstance(bool bIncludePreview) 
 {
-	return MakeShareable( new FSlateBrushStructCustomization() );
+	return MakeShareable( new FSlateBrushStructCustomization(bIncludePreview) );
 }
 
-FSlateBrushStructCustomization::FSlateBrushStructCustomization()
-	: bSlimView(true)
+FSlateBrushStructCustomization::FSlateBrushStructCustomization(bool bIncludePreview)
+	: bIncludePreview(bIncludePreview)
 {
 }
 
@@ -1173,7 +1173,7 @@ void FSlateBrushStructCustomization::CustomizeChildren( TSharedRef<IPropertyHand
 	.Visibility( TAttribute<EVisibility>::Create( TAttribute<EVisibility>::FGetter::CreateSP( this, &FSlateBrushStructCustomization::GetMarginPropertyVisibility ) ) );
 
 	// Don't show the preview area when in slim view mode.
-	if ( !bSlimView )
+	if ( bIncludePreview )
 	{
 		// Create the Slate Brush Preview widget and add the Preview group
 		TArray<void*> RawData;
diff --git a/Engine/Source/Editor/DetailCustomizations/Public/Customizations/SlateBrushCustomization.h b/Engine/Source/Editor/DetailCustomizations/Public/Customizations/SlateBrushCustomization.h
--- a/Engine/Source/Editor/DetailCustomizations/Public/Customizations/SlateBrushCustomization.h
+++ b/Engine/Source/Editor/DetailCustomizations/Public/Customizations/SlateBrushCustomization.h
@@ -3,12 +3,12 @@
 #pragma once
 
 
-class FSlateBrushStructCustomization : public IPropertyTypeCustomization
+class DETAILCUSTOMIZATIONS_API FSlateBrushStructCustomization : public IPropertyTypeCustomization
 {
 public:
-	static TSharedRef<IPropertyTypeCustomization> MakeInstance();
+	static TSharedRef<IPropertyTypeCustomization> MakeInstance(bool bIncludePreview);
 
-	FSlateBrushStructCustomization();
+	FSlateBrushStructCustomization(bool bIncludePreview);
 
 	/** IPropertyTypeCustomization interface */
 	virtual void CustomizeHeader( TSharedRef<class IPropertyHandle> StructPropertyHandle, class FDetailWidgetRow& HeaderRow, IPropertyTypeCustomizationUtils& StructCustomizationUtils ) override;
@@ -32,7 +32,7 @@ private:
 	/** Error text to display if the resource object is not valid*/
 	TSharedPtr<SErrorText> ResourceErrorText;
 
-	/**  */
-	bool bSlimView;
+	/** Should we show the preview portion of the customization? */
+	bool bIncludePreview;
 };
 
diff --git a/Engine/Source/Editor/UMGEditor/Private/Details/SWidgetDetailsView.cpp b/Engine/Source/Editor/UMGEditor/Private/Details/SWidgetDetailsView.cpp
--- a/Engine/Source/Editor/UMGEditor/Private/Details/SWidgetDetailsView.cpp
+++ b/Engine/Source/Editor/UMGEditor/Private/Details/SWidgetDetailsView.cpp
@@ -10,6 +10,8 @@
 #include "Animation/UMGDetailKeyframeHandler.h"
 #include "EditorClassUtils.h"
 
+#include "Customizations/SlateBrushCustomization.h"
+
 #include "CanvasSlotCustomization.h"
 #include "HorizontalAlignmentCustomization.h"
 #include "VerticalAlignmentCustomization.h"
@@ -125,14 +127,13 @@ SWidgetDetailsView::~SWidgetDetailsView()
 	PropertyModule.UnregisterCustomPropertyTypeLayout(TEXT("EHorizontalAlignment"), nullptr, PropertyView);
 	PropertyModule.UnregisterCustomPropertyTypeLayout(TEXT("EVerticalAlignment"), nullptr, PropertyView);
 	PropertyModule.UnregisterCustomPropertyTypeLayout(TEXT("SlateChildSize"), nullptr, PropertyView);
+	PropertyModule.UnregisterCustomPropertyTypeLayout(TEXT("SlateBrush"), nullptr, PropertyView);
 }
 
 void SWidgetDetailsView::RegisterCustomizations()
 {
 	PropertyView->RegisterInstancedCustomPropertyLayout(UWidget::StaticClass(), FOnGetDetailCustomizationInstance::CreateStatic(&FBlueprintWidgetCustomization::MakeInstance, BlueprintEditor.Pin().ToSharedRef(), BlueprintEditor.Pin()->GetBlueprintObj()));
 
-	//PropertyView->RegisterInstancedCustomPropertyLayout(FSlateBrush::StaticStruct()->StaticClass(), FOnGetDetailCustomizationInstance::CreateStatic(&FSlateBrushStructCustomization::MakeInstance));
-
 	static FName PropertyEditor("PropertyEditor");
 	FPropertyEditorModule& PropertyModule = FModuleManager::GetModuleChecked<FPropertyEditorModule>(PropertyEditor);
 
@@ -140,6 +141,7 @@ void SWidgetDetailsView::RegisterCustomizations()
 	PropertyModule.RegisterCustomPropertyTypeLayout(TEXT("EHorizontalAlignment"), FOnGetPropertyTypeCustomizationInstance::CreateStatic(&FHorizontalAlignmentCustomization::MakeInstance), nullptr, PropertyView);
 	PropertyModule.RegisterCustomPropertyTypeLayout(TEXT("EVerticalAlignment"), FOnGetPropertyTypeCustomizationInstance::CreateStatic(&FVerticalAlignmentCustomization::MakeInstance), nullptr, PropertyView);
 	PropertyModule.RegisterCustomPropertyTypeLayout(TEXT("SlateChildSize"), FOnGetPropertyTypeCustomizationInstance::CreateStatic(&FSlateChildSizeCustomization::MakeInstance), nullptr, PropertyView);
+	PropertyModule.RegisterCustomPropertyTypeLayout(TEXT("SlateBrush"), FOnGetPropertyTypeCustomizationInstance::CreateStatic(&FSlateBrushStructCustomization::MakeInstance, false), nullptr, PropertyView);
 }
 
 void SWidgetDetailsView::OnEditorSelectionChanging()
diff --git a/Engine/Source/Editor/UMGEditor/UMGEditor.Build.cs b/Engine/Source/Editor/UMGEditor/UMGEditor.Build.cs
--- a/Engine/Source/Editor/UMGEditor/UMGEditor.Build.cs
+++ b/Engine/Source/Editor/UMGEditor/UMGEditor.Build.cs
@@ -48,6 +48,7 @@ public class UMGEditor : ModuleRules
 				"SlateCore",
 				"MovieSceneCore",
 				"Sequencer",
+				"DetailCustomizations"
 			}
 			);
 
