UMG - The widget designer now has some better logic for registering selection.  It no longer does it when panning, additionally you can now select parent widgets, and even if you click on a child, if that action becomes a drag operation it will drag the parent, instead of dragging the child.

diff --git a/Engine/Source/Editor/UMGEditor/Private/Designer/SDesignerView.cpp b/Engine/Source/Editor/UMGEditor/Private/Designer/SDesignerView.cpp
--- a/Engine/Source/Editor/UMGEditor/Private/Designer/SDesignerView.cpp
+++ b/Engine/Source/Editor/UMGEditor/Private/Designer/SDesignerView.cpp
@@ -675,6 +675,18 @@ FWidgetReference SDesignerView::GetWidgetAtCursor(const FGeometry& MyGeometry, c
 	return FWidgetReference();
 }
 
+void SDesignerView::ResolvePendingSelectedWidgets()
+{
+	if ( PendingSelectedWidget.IsValid() )
+	{
+		TSet<FWidgetReference> SelectedTemplates;
+		SelectedTemplates.Add(PendingSelectedWidget);
+		BlueprintEditor.Pin()->SelectWidgets(SelectedTemplates);
+
+		PendingSelectedWidget = FWidgetReference();
+	}
+}
+
 FReply SDesignerView::OnMouseButtonDown(const FGeometry& MyGeometry, const FPointerEvent& MouseEvent)
 {
 	SDesignSurface::OnMouseButtonDown(MyGeometry, MouseEvent);
@@ -686,15 +698,20 @@ FReply SDesignerView::OnMouseButtonDown(const FGeometry& MyGeometry, const FPoin
 
 	if ( NewSelectedWidget.IsValid() )
 	{
-		//@TODO UMG primary FBlueprintEditor needs to be inherited and selection control needs to be centralized.
-		// Set the template as selected in the details panel
-		TSet<FWidgetReference> SelectedTemplates;
-		SelectedTemplates.Add(NewSelectedWidget);
-		BlueprintEditor.Pin()->SelectWidgets(SelectedTemplates);
+		PendingSelectedWidget = NewSelectedWidget;
 
 		if ( MouseEvent.GetEffectingButton() == EKeys::LeftMouseButton )
 		{
+			// If the newly clicked item is a child of the active selection, add it to the pending set of selected 
+			// widgets, if they begin dragging we can just move the parent, but if it's not part of the parent set, we want to immediately
+			// begin dragging it.
+			if ( !SelectedWidget.IsValid() || !NewSelectedWidget.GetTemplate()->IsChildOf(SelectedWidget.GetTemplate()) )
+			{
+				ResolvePendingSelectedWidgets();
+			}
+
 			bMouseDown = true;
+			ScreenMouseDownLocation = MouseEvent.GetScreenSpacePosition();
 		}
 	}
 
@@ -706,6 +723,8 @@ FReply SDesignerView::OnMouseButtonUp(const FGeometry& MyGeometry, const FPointe
 {
 	if ( HasMouseCapture() && MouseEvent.GetEffectingButton() == EKeys::LeftMouseButton )
 	{
+		ResolvePendingSelectedWidgets();
+
 		bMouseDown = false;
 		bMovingExistingWidget = false;
 		DesignerMessage = EDesignerMessage::None;
@@ -714,6 +733,8 @@ FReply SDesignerView::OnMouseButtonUp(const FGeometry& MyGeometry, const FPointe
 	{
 		if ( !bIsPanning )
 		{
+			ResolvePendingSelectedWidgets();
+
 			ShowContextMenu(MyGeometry, MouseEvent);
 		}
 	}
@@ -987,6 +1008,14 @@ FReply SDesignerView::OnDragDetected(const FGeometry& MyGeometry, const FPointer
 {
 	if ( SelectedWidget.IsValid() )
 	{
+		// Clear any pending selected widgets, the user has already decided what widget they want.
+		PendingSelectedWidget = FWidgetReference();
+
+		//
+		FArrangedWidget ArrangedWidget(SNullWidget::NullWidget, FGeometry());
+		FDesignTimeUtils::GetArrangedWidget(SelectedWidget.GetPreview()->GetCachedWidget().ToSharedRef(), ArrangedWidget);
+		SelectedWidgetContextMenuLocation = ArrangedWidget.Geometry.AbsoluteToLocal(ScreenMouseDownLocation);
+
 		ClearExtensionWidgets();
 
 		return FReply::Handled().BeginDragDrop(FSelectedWidgetDragDropOp::New(BlueprintEditor.Pin(), SelectedWidget));
@@ -1145,7 +1174,7 @@ UWidget* SDesignerView::ProcessDropAndAddWidget(const FGeometry& MyGeometry, con
 				SlateWidget->SlatePrepass();
 				const FVector2D& WidgetDesiredSize = SlateWidget->GetDesiredSize();
 
-				static const FVector2D MinimumDefaultSize(20, 20);
+				static const FVector2D MinimumDefaultSize(100, 40);
 				FVector2D LocalSize = FVector2D(FMath::Max(WidgetDesiredSize.X, MinimumDefaultSize.X), FMath::Max(WidgetDesiredSize.Y, MinimumDefaultSize.Y));
 
 				Slot->SetDesiredPosition(LocalPosition);
diff --git a/Engine/Source/Editor/UMGEditor/Private/Designer/SDesignerView.h b/Engine/Source/Editor/UMGEditor/Private/Designer/SDesignerView.h
--- a/Engine/Source/Editor/UMGEditor/Private/Designer/SDesignerView.h
+++ b/Engine/Source/Editor/UMGEditor/Private/Designer/SDesignerView.h
@@ -76,6 +76,9 @@ private:
 
 	virtual FSlateRect ComputeAreaBounds() const override;
 
+	/** Adds any pending selected widgets to the selection set */
+	void ResolvePendingSelectedWidgets();
+
 	/** Updates the designer to display the latest preview widget */
 	void UpdatePreviewWidget();
 
@@ -181,9 +184,18 @@ private:
 	TArray< TWeakPtr<SWidget> > SelectedSlateWidgets;
 	TWeakPtr<SWidget> SelectedSlateWidget;
 
+	/**
+	 * Holds onto a temporary widget that the user may be getting ready to select, or may just 
+	 * be the widget that got hit on the initial mouse down before moving the parent.
+	 */
+	FWidgetReference PendingSelectedWidget;
+
 	/**  */
 	bool bMouseDown;
 
+	/**  */
+	FVector2D ScreenMouseDownLocation;
+
 	/** An existing widget is being moved in its current container, or in to a new container. */
 	bool bMovingExistingWidget;
 
