UMG - Tweaking the padding in the size customization.

diff --git a/Engine/Source/Editor/UMGEditor/Private/Customizations/SlateChildSizeCustomization.cpp b/Engine/Source/Editor/UMGEditor/Private/Customizations/SlateChildSizeCustomization.cpp
--- a/Engine/Source/Editor/UMGEditor/Private/Customizations/SlateChildSizeCustomization.cpp
+++ b/Engine/Source/Editor/UMGEditor/Private/Customizations/SlateChildSizeCustomization.cpp
@@ -17,8 +17,8 @@ void FSlateChildSizeCustomization::CustomizeHeader(TSharedRef<IPropertyHandle> P
 	TSharedPtr<IPropertyHandle> ValueHandle = PropertyHandle->GetChildHandle(GET_MEMBER_NAME_CHECKED(FSlateChildSize, Value));
 	TSharedPtr<IPropertyHandle> RuleHandle = PropertyHandle->GetChildHandle(GET_MEMBER_NAME_CHECKED(FSlateChildSize, SizeRule));
 
-	const FMargin OuterPadding(2);
-	const FMargin ContentPadding(2);
+	const FMargin OuterPadding(2, 0);
+	const FMargin ContentPadding(4, 2);
 
 	if ( !( ValueHandle.IsValid() || RuleHandle.IsValid() ) )
 	{
