#UE4

- Crash Fix. Was crashing every time I was deleting a UserWidget from the Designer panel.

diff --git a/Engine/Source/Editor/UMGEditor/Private/Details/SWidgetDetailsView.cpp b/Engine/Source/Editor/UMGEditor/Private/Details/SWidgetDetailsView.cpp
--- a/Engine/Source/Editor/UMGEditor/Private/Details/SWidgetDetailsView.cpp
+++ b/Engine/Source/Editor/UMGEditor/Private/Details/SWidgetDetailsView.cpp
@@ -172,7 +172,7 @@ void SWidgetDetailsView::OnEditorSelectionChanged()
 		}
 	}
 
-	if ( SelectedObjects.Num() == 1 )
+	if ( SelectedObjects.Num() == 1 && SelectedObjects[0].IsValid() )
 	{
 		ClassLinkArea->SetContent(FEditorClassUtils::GetSourceLink(SelectedObjects[0]->GetClass(), TWeakObjectPtr<UObject>()));
 	}
