Updated UMG to consistently use the term "Palette"

diff --git a/Engine/Source/Editor/UMGEditor/Private/Palette/SPaletteView.cpp b/Engine/Source/Editor/UMGEditor/Private/Palette/SPaletteView.cpp
--- a/Engine/Source/Editor/UMGEditor/Private/Palette/SPaletteView.cpp
+++ b/Engine/Source/Editor/UMGEditor/Private/Palette/SPaletteView.cpp
@@ -127,7 +127,7 @@ public:
 	virtual TSharedRef<ITableRow> BuildRow(const TSharedRef<STableViewBase>& OwnerTable) override
 	{
 		return SNew(STableRow< TSharedPtr<FWidgetViewModel> >, OwnerTable)
-		.Style( FEditorStyle::Get(), "UMGEditor.ToolboxHeader" )
+		.Style( FEditorStyle::Get(), "UMGEditor.PaletteHeader" )
 			.Padding(2.0f)
 			.ShowSelection(false)
 			[
@@ -182,7 +182,7 @@ void SPaletteView::Construct(const FArguments& InArgs, TSharedPtr<FBlueprintEdit
 		.AutoHeight()
 		[
 			SNew(SSearchBox)
-			.HintText(LOCTEXT("SearchTemplates", "Search Templates"))
+			.HintText(LOCTEXT("SearchTemplates", "Search Palette"))
 			.OnTextChanged(this, &SPaletteView::OnSearchChanged)
 		]
 
diff --git a/Engine/Source/Runtime/UMG/Private/Components/Border.cpp b/Engine/Source/Runtime/UMG/Private/Components/Border.cpp
--- a/Engine/Source/Runtime/UMG/Private/Components/Border.cpp
+++ b/Engine/Source/Runtime/UMG/Private/Components/Border.cpp
@@ -320,7 +320,7 @@ const FSlateBrush* UBorder::GetEditorIcon()
 	return FUMGStyle::Get().GetBrush("Widget.Border");
 }
 
-const FText UBorder::GetToolboxCategory()
+const FText UBorder::GetPaletteCategory()
 {
 	return LOCTEXT("Common", "Common");
 }
diff --git a/Engine/Source/Runtime/UMG/Private/Components/Button.cpp b/Engine/Source/Runtime/UMG/Private/Components/Button.cpp
--- a/Engine/Source/Runtime/UMG/Private/Components/Button.cpp
+++ b/Engine/Source/Runtime/UMG/Private/Components/Button.cpp
@@ -147,7 +147,7 @@ const FSlateBrush* UButton::GetEditorIcon()
 	return FUMGStyle::Get().GetBrush("Widget.Button");
 }
 
-const FText UButton::GetToolboxCategory()
+const FText UButton::GetPaletteCategory()
 {
 	return LOCTEXT("Common", "Common");
 }
diff --git a/Engine/Source/Runtime/UMG/Private/Components/CanvasPanel.cpp b/Engine/Source/Runtime/UMG/Private/Components/CanvasPanel.cpp
--- a/Engine/Source/Runtime/UMG/Private/Components/CanvasPanel.cpp
+++ b/Engine/Source/Runtime/UMG/Private/Components/CanvasPanel.cpp
@@ -110,7 +110,7 @@ const FSlateBrush* UCanvasPanel::GetEditorIcon()
 	return FUMGStyle::Get().GetBrush("Widget.Canvas");
 }
 
-const FText UCanvasPanel::GetToolboxCategory()
+const FText UCanvasPanel::GetPaletteCategory()
 {
 	return LOCTEXT("Panel", "Panel");
 }
diff --git a/Engine/Source/Runtime/UMG/Private/Components/CheckBox.cpp b/Engine/Source/Runtime/UMG/Private/Components/CheckBox.cpp
--- a/Engine/Source/Runtime/UMG/Private/Components/CheckBox.cpp
+++ b/Engine/Source/Runtime/UMG/Private/Components/CheckBox.cpp
@@ -214,7 +214,7 @@ const FSlateBrush* UCheckBox::GetEditorIcon()
 	return FUMGStyle::Get().GetBrush("Widget.CheckBox");
 }
 
-const FText UCheckBox::GetToolboxCategory()
+const FText UCheckBox::GetPaletteCategory()
 {
 	return LOCTEXT("Common", "Common");
 }
diff --git a/Engine/Source/Runtime/UMG/Private/Components/CircularThrobber.cpp b/Engine/Source/Runtime/UMG/Private/Components/CircularThrobber.cpp
--- a/Engine/Source/Runtime/UMG/Private/Components/CircularThrobber.cpp
+++ b/Engine/Source/Runtime/UMG/Private/Components/CircularThrobber.cpp
@@ -96,7 +96,7 @@ const FSlateBrush* UCircularThrobber::GetEditorIcon()
 	return FUMGStyle::Get().GetBrush("Widget.CircularThrobber");
 }
 
-const FText UCircularThrobber::GetToolboxCategory()
+const FText UCircularThrobber::GetPaletteCategory()
 {
 	return LOCTEXT("Primitive", "Primitive");
 }
diff --git a/Engine/Source/Runtime/UMG/Private/Components/ComboBox.cpp b/Engine/Source/Runtime/UMG/Private/Components/ComboBox.cpp
--- a/Engine/Source/Runtime/UMG/Private/Components/ComboBox.cpp
+++ b/Engine/Source/Runtime/UMG/Private/Components/ComboBox.cpp
@@ -55,7 +55,7 @@ const FSlateBrush* UComboBox::GetEditorIcon()
 	return FUMGStyle::Get().GetBrush("Widget.ComboBox");
 }
 
-const FText UComboBox::GetToolboxCategory()
+const FText UComboBox::GetPaletteCategory()
 {
 	return LOCTEXT("Misc", "Misc");
 }
diff --git a/Engine/Source/Runtime/UMG/Private/Components/ComboBoxString.cpp b/Engine/Source/Runtime/UMG/Private/Components/ComboBoxString.cpp
--- a/Engine/Source/Runtime/UMG/Private/Components/ComboBoxString.cpp
+++ b/Engine/Source/Runtime/UMG/Private/Components/ComboBoxString.cpp
@@ -141,7 +141,7 @@ const FSlateBrush* UComboBoxString::GetEditorIcon()
 	return FUMGStyle::Get().GetBrush("Widget.ComboBox");
 }
 
-const FText UComboBoxString::GetToolboxCategory()
+const FText UComboBoxString::GetPaletteCategory()
 {
 	return LOCTEXT("Input", "Input");
 }
diff --git a/Engine/Source/Runtime/UMG/Private/Components/EditableText.cpp b/Engine/Source/Runtime/UMG/Private/Components/EditableText.cpp
--- a/Engine/Source/Runtime/UMG/Private/Components/EditableText.cpp
+++ b/Engine/Source/Runtime/UMG/Private/Components/EditableText.cpp
@@ -183,7 +183,7 @@ const FSlateBrush* UEditableText::GetEditorIcon()
 	return FUMGStyle::Get().GetBrush("Widget.EditableText");
 }
 
-const FText UEditableText::GetToolboxCategory()
+const FText UEditableText::GetPaletteCategory()
 {
 	return LOCTEXT("Primitive", "Primitive");
 }
diff --git a/Engine/Source/Runtime/UMG/Private/Components/EditableTextBox.cpp b/Engine/Source/Runtime/UMG/Private/Components/EditableTextBox.cpp
--- a/Engine/Source/Runtime/UMG/Private/Components/EditableTextBox.cpp
+++ b/Engine/Source/Runtime/UMG/Private/Components/EditableTextBox.cpp
@@ -139,7 +139,7 @@ const FSlateBrush* UEditableTextBox::GetEditorIcon()
 	return FUMGStyle::Get().GetBrush("Widget.EditableTextBox");
 }
 
-const FText UEditableTextBox::GetToolboxCategory()
+const FText UEditableTextBox::GetPaletteCategory()
 {
 	return LOCTEXT("Common", "Common");
 }
diff --git a/Engine/Source/Runtime/UMG/Private/Components/GridPanel.cpp b/Engine/Source/Runtime/UMG/Private/Components/GridPanel.cpp
--- a/Engine/Source/Runtime/UMG/Private/Components/GridPanel.cpp
+++ b/Engine/Source/Runtime/UMG/Private/Components/GridPanel.cpp
@@ -88,7 +88,7 @@ const FSlateBrush* UGridPanel::GetEditorIcon()
 	return FUMGStyle::Get().GetBrush("Widget.Grid");
 }
 
-const FText UGridPanel::GetToolboxCategory()
+const FText UGridPanel::GetPaletteCategory()
 {
 	return LOCTEXT("Panel", "Panel");
 }
diff --git a/Engine/Source/Runtime/UMG/Private/Components/HorizontalBox.cpp b/Engine/Source/Runtime/UMG/Private/Components/HorizontalBox.cpp
--- a/Engine/Source/Runtime/UMG/Private/Components/HorizontalBox.cpp
+++ b/Engine/Source/Runtime/UMG/Private/Components/HorizontalBox.cpp
@@ -83,7 +83,7 @@ const FSlateBrush* UHorizontalBox::GetEditorIcon()
 	return FUMGStyle::Get().GetBrush("Widget.HorizontalBox");
 }
 
-const FText UHorizontalBox::GetToolboxCategory()
+const FText UHorizontalBox::GetPaletteCategory()
 {
 	return LOCTEXT("Panel", "Panel");
 }
diff --git a/Engine/Source/Runtime/UMG/Private/Components/Image.cpp b/Engine/Source/Runtime/UMG/Private/Components/Image.cpp
--- a/Engine/Source/Runtime/UMG/Private/Components/Image.cpp
+++ b/Engine/Source/Runtime/UMG/Private/Components/Image.cpp
@@ -155,7 +155,7 @@ const FSlateBrush* UImage::GetEditorIcon()
 	return FUMGStyle::Get().GetBrush("Widget.Image");
 }
 
-const FText UImage::GetToolboxCategory()
+const FText UImage::GetPaletteCategory()
 {
 	return LOCTEXT("Common", "Common");
 }
diff --git a/Engine/Source/Runtime/UMG/Private/Components/ListView.cpp b/Engine/Source/Runtime/UMG/Private/Components/ListView.cpp
--- a/Engine/Source/Runtime/UMG/Private/Components/ListView.cpp
+++ b/Engine/Source/Runtime/UMG/Private/Components/ListView.cpp
@@ -66,7 +66,7 @@ const FSlateBrush* UListView::GetEditorIcon()
 	return FUMGStyle::Get().GetBrush("Widget.ListView");
 }
 
-const FText UListView::GetToolboxCategory()
+const FText UListView::GetPaletteCategory()
 {
 	return LOCTEXT("Misc", "Misc");
 }
diff --git a/Engine/Source/Runtime/UMG/Private/Components/MenuAnchor.cpp b/Engine/Source/Runtime/UMG/Private/Components/MenuAnchor.cpp
--- a/Engine/Source/Runtime/UMG/Private/Components/MenuAnchor.cpp
+++ b/Engine/Source/Runtime/UMG/Private/Components/MenuAnchor.cpp
@@ -103,7 +103,7 @@ const FSlateBrush* UMenuAnchor::GetEditorIcon()
 	return FUMGStyle::Get().GetBrush("Widget.MenuAnchor");
 }
 
-const FText UMenuAnchor::GetToolboxCategory()
+const FText UMenuAnchor::GetPaletteCategory()
 {
 	return LOCTEXT("Advanced", "Advanced");
 }
diff --git a/Engine/Source/Runtime/UMG/Private/Components/MultiLineEditableText.cpp b/Engine/Source/Runtime/UMG/Private/Components/MultiLineEditableText.cpp
--- a/Engine/Source/Runtime/UMG/Private/Components/MultiLineEditableText.cpp
+++ b/Engine/Source/Runtime/UMG/Private/Components/MultiLineEditableText.cpp
@@ -101,7 +101,7 @@ const FSlateBrush* UMultiLineEditableText::GetEditorIcon()
 	return FUMGStyle::Get().GetBrush("Widget.MultiLineEditableText");
 }
 
-const FText UMultiLineEditableText::GetToolboxCategory()
+const FText UMultiLineEditableText::GetPaletteCategory()
 {
 	return LOCTEXT("Primitive", "Primitive");
 }
diff --git a/Engine/Source/Runtime/UMG/Private/Components/MultiLineEditableTextBox.cpp b/Engine/Source/Runtime/UMG/Private/Components/MultiLineEditableTextBox.cpp
--- a/Engine/Source/Runtime/UMG/Private/Components/MultiLineEditableTextBox.cpp
+++ b/Engine/Source/Runtime/UMG/Private/Components/MultiLineEditableTextBox.cpp
@@ -110,7 +110,7 @@ const FSlateBrush* UMultiLineEditableTextBox::GetEditorIcon()
 	return FUMGStyle::Get().GetBrush("Widget.MultiLineEditableTextBox");
 }
 
-const FText UMultiLineEditableTextBox::GetToolboxCategory()
+const FText UMultiLineEditableTextBox::GetPaletteCategory()
 {
 	return LOCTEXT("Common", "Common");
 }
diff --git a/Engine/Source/Runtime/UMG/Private/Components/NativeWidgetHost.cpp b/Engine/Source/Runtime/UMG/Private/Components/NativeWidgetHost.cpp
--- a/Engine/Source/Runtime/UMG/Private/Components/NativeWidgetHost.cpp
+++ b/Engine/Source/Runtime/UMG/Private/Components/NativeWidgetHost.cpp
@@ -52,7 +52,7 @@ const FSlateBrush* UNativeWidgetHost::GetEditorIcon()
 	return FUMGStyle::Get().GetBrush("Widget.NativeWidgetHost");
 }
 
-const FText UNativeWidgetHost::GetToolboxCategory()
+const FText UNativeWidgetHost::GetPaletteCategory()
 {
 	return LOCTEXT("Primitive", "Primitive");
 }
diff --git a/Engine/Source/Runtime/UMG/Private/Components/Overlay.cpp b/Engine/Source/Runtime/UMG/Private/Components/Overlay.cpp
--- a/Engine/Source/Runtime/UMG/Private/Components/Overlay.cpp
+++ b/Engine/Source/Runtime/UMG/Private/Components/Overlay.cpp
@@ -73,7 +73,7 @@ const FSlateBrush* UOverlay::GetEditorIcon()
 	return FUMGStyle::Get().GetBrush("Widget.Overlay");
 }
 
-const FText UOverlay::GetToolboxCategory()
+const FText UOverlay::GetPaletteCategory()
 {
 	return LOCTEXT("Panel", "Panel");
 }
diff --git a/Engine/Source/Runtime/UMG/Private/Components/ProgressBar.cpp b/Engine/Source/Runtime/UMG/Private/Components/ProgressBar.cpp
--- a/Engine/Source/Runtime/UMG/Private/Components/ProgressBar.cpp
+++ b/Engine/Source/Runtime/UMG/Private/Components/ProgressBar.cpp
@@ -108,7 +108,7 @@ const FSlateBrush* UProgressBar::GetEditorIcon()
 	return FUMGStyle::Get().GetBrush("Widget.ProgressBar");
 }
 
-const FText UProgressBar::GetToolboxCategory()
+const FText UProgressBar::GetPaletteCategory()
 {
 	return LOCTEXT("Common", "Common");
 }
diff --git a/Engine/Source/Runtime/UMG/Private/Components/ScaleBox.cpp b/Engine/Source/Runtime/UMG/Private/Components/ScaleBox.cpp
--- a/Engine/Source/Runtime/UMG/Private/Components/ScaleBox.cpp
+++ b/Engine/Source/Runtime/UMG/Private/Components/ScaleBox.cpp
@@ -73,7 +73,7 @@ const FSlateBrush* UScaleBox::GetEditorIcon()
 	return FUMGStyle::Get().GetBrush("Widget.ScaleBox");
 }
 
-const FText UScaleBox::GetToolboxCategory()
+const FText UScaleBox::GetPaletteCategory()
 {
 	return LOCTEXT("Panel", "Panel");
 }
diff --git a/Engine/Source/Runtime/UMG/Private/Components/ScrollBar.cpp b/Engine/Source/Runtime/UMG/Private/Components/ScrollBar.cpp
--- a/Engine/Source/Runtime/UMG/Private/Components/ScrollBar.cpp
+++ b/Engine/Source/Runtime/UMG/Private/Components/ScrollBar.cpp
@@ -57,7 +57,7 @@ const FSlateBrush* UScrollBar::GetEditorIcon()
 	return FUMGStyle::Get().GetBrush("Widget.ScrollBar");
 }
 
-const FText UScrollBar::GetToolboxCategory()
+const FText UScrollBar::GetPaletteCategory()
 {
 	return LOCTEXT("Primitive", "Primitive");
 }
diff --git a/Engine/Source/Runtime/UMG/Private/Components/ScrollBox.cpp b/Engine/Source/Runtime/UMG/Private/Components/ScrollBox.cpp
--- a/Engine/Source/Runtime/UMG/Private/Components/ScrollBox.cpp
+++ b/Engine/Source/Runtime/UMG/Private/Components/ScrollBox.cpp
@@ -127,7 +127,7 @@ const FSlateBrush* UScrollBox::GetEditorIcon()
 	return FUMGStyle::Get().GetBrush("Widget.ScrollBox");
 }
 
-const FText UScrollBox::GetToolboxCategory()
+const FText UScrollBox::GetPaletteCategory()
 {
 	return LOCTEXT("Panel", "Panel");
 }
diff --git a/Engine/Source/Runtime/UMG/Private/Components/SizeBox.cpp b/Engine/Source/Runtime/UMG/Private/Components/SizeBox.cpp
--- a/Engine/Source/Runtime/UMG/Private/Components/SizeBox.cpp
+++ b/Engine/Source/Runtime/UMG/Private/Components/SizeBox.cpp
@@ -236,7 +236,7 @@ const FSlateBrush* USizeBox::GetEditorIcon()
 	return FUMGStyle::Get().GetBrush("Widget.SizeBox");
 }
 
-const FText USizeBox::GetToolboxCategory()
+const FText USizeBox::GetPaletteCategory()
 {
 	return LOCTEXT("Common", "Common");
 }
diff --git a/Engine/Source/Runtime/UMG/Private/Components/Slider.cpp b/Engine/Source/Runtime/UMG/Private/Components/Slider.cpp
--- a/Engine/Source/Runtime/UMG/Private/Components/Slider.cpp
+++ b/Engine/Source/Runtime/UMG/Private/Components/Slider.cpp
@@ -82,7 +82,7 @@ const FSlateBrush* USlider::GetEditorIcon()
 	return FUMGStyle::Get().GetBrush("Widget.Slider");
 }
 
-const FText USlider::GetToolboxCategory()
+const FText USlider::GetPaletteCategory()
 {
 	return LOCTEXT("Common", "Common");
 }
diff --git a/Engine/Source/Runtime/UMG/Private/Components/Spacer.cpp b/Engine/Source/Runtime/UMG/Private/Components/Spacer.cpp
--- a/Engine/Source/Runtime/UMG/Private/Components/Spacer.cpp
+++ b/Engine/Source/Runtime/UMG/Private/Components/Spacer.cpp
@@ -48,7 +48,7 @@ const FSlateBrush* USpacer::GetEditorIcon()
 	return FUMGStyle::Get().GetBrush("Widget.Spacer");
 }
 
-const FText USpacer::GetToolboxCategory()
+const FText USpacer::GetPaletteCategory()
 {
 	return LOCTEXT("Primitive", "Primitive");
 }
diff --git a/Engine/Source/Runtime/UMG/Private/Components/SpinBox.cpp b/Engine/Source/Runtime/UMG/Private/Components/SpinBox.cpp
--- a/Engine/Source/Runtime/UMG/Private/Components/SpinBox.cpp
+++ b/Engine/Source/Runtime/UMG/Private/Components/SpinBox.cpp
@@ -274,7 +274,7 @@ const FSlateBrush* USpinBox::GetEditorIcon()
 	return FUMGStyle::Get().GetBrush("Widget.SpinBox");
 }
 
-const FText USpinBox::GetToolboxCategory()
+const FText USpinBox::GetPaletteCategory()
 {
 	return LOCTEXT("Common", "Common");
 }
diff --git a/Engine/Source/Runtime/UMG/Private/Components/TextBlock.cpp b/Engine/Source/Runtime/UMG/Private/Components/TextBlock.cpp
--- a/Engine/Source/Runtime/UMG/Private/Components/TextBlock.cpp
+++ b/Engine/Source/Runtime/UMG/Private/Components/TextBlock.cpp
@@ -158,7 +158,7 @@ const FSlateBrush* UTextBlock::GetEditorIcon()
 	return FUMGStyle::Get().GetBrush("Widget.TextBlock");
 }
 
-const FText UTextBlock::GetToolboxCategory()
+const FText UTextBlock::GetPaletteCategory()
 {
 	return LOCTEXT("Common", "Common");
 }
diff --git a/Engine/Source/Runtime/UMG/Private/Components/Throbber.cpp b/Engine/Source/Runtime/UMG/Private/Components/Throbber.cpp
--- a/Engine/Source/Runtime/UMG/Private/Components/Throbber.cpp
+++ b/Engine/Source/Runtime/UMG/Private/Components/Throbber.cpp
@@ -114,7 +114,7 @@ const FSlateBrush* UThrobber::GetEditorIcon()
 	return FUMGStyle::Get().GetBrush("Widget.Throbber");
 }
 
-const FText UThrobber::GetToolboxCategory()
+const FText UThrobber::GetPaletteCategory()
 {
 	return LOCTEXT("Primitive", "Primitive");
 }
diff --git a/Engine/Source/Runtime/UMG/Private/Components/TileView.cpp b/Engine/Source/Runtime/UMG/Private/Components/TileView.cpp
--- a/Engine/Source/Runtime/UMG/Private/Components/TileView.cpp
+++ b/Engine/Source/Runtime/UMG/Private/Components/TileView.cpp
@@ -83,7 +83,7 @@ const FSlateBrush* UTileView::GetEditorIcon()
 	return FUMGStyle::Get().GetBrush("Widget.TileView");
 }
 
-const FText UTileView::GetToolboxCategory()
+const FText UTileView::GetPaletteCategory()
 {
 	return LOCTEXT("Misc", "Misc");
 }
diff --git a/Engine/Source/Runtime/UMG/Private/Components/UniformGridPanel.cpp b/Engine/Source/Runtime/UMG/Private/Components/UniformGridPanel.cpp
--- a/Engine/Source/Runtime/UMG/Private/Components/UniformGridPanel.cpp
+++ b/Engine/Source/Runtime/UMG/Private/Components/UniformGridPanel.cpp
@@ -83,7 +83,7 @@ const FSlateBrush* UUniformGridPanel::GetEditorIcon()
 	return FUMGStyle::Get().GetBrush("Widget.UniformGrid");
 }
 
-const FText UUniformGridPanel::GetToolboxCategory()
+const FText UUniformGridPanel::GetPaletteCategory()
 {
 	return LOCTEXT("Panel", "Panel");
 }
diff --git a/Engine/Source/Runtime/UMG/Private/Components/VerticalBox.cpp b/Engine/Source/Runtime/UMG/Private/Components/VerticalBox.cpp
--- a/Engine/Source/Runtime/UMG/Private/Components/VerticalBox.cpp
+++ b/Engine/Source/Runtime/UMG/Private/Components/VerticalBox.cpp
@@ -78,7 +78,7 @@ const FSlateBrush* UVerticalBox::GetEditorIcon()
 	return FUMGStyle::Get().GetBrush("Widget.VerticalBox");
 }
 
-const FText UVerticalBox::GetToolboxCategory()
+const FText UVerticalBox::GetPaletteCategory()
 {
 	return LOCTEXT("Panel", "Panel");
 }
diff --git a/Engine/Source/Runtime/UMG/Private/Components/Viewport.cpp b/Engine/Source/Runtime/UMG/Private/Components/Viewport.cpp
--- a/Engine/Source/Runtime/UMG/Private/Components/Viewport.cpp
+++ b/Engine/Source/Runtime/UMG/Private/Components/Viewport.cpp
@@ -626,7 +626,7 @@ const FSlateBrush* UViewport::GetEditorIcon()
 	return FUMGStyle::Get().GetBrush("Widget.Viewport");
 }
 
-const FText UViewport::GetToolboxCategory()
+const FText UViewport::GetPaletteCategory()
 {
 	return LOCTEXT("Primitive", "Primitive");
 }
diff --git a/Engine/Source/Runtime/UMG/Private/Components/Widget.cpp b/Engine/Source/Runtime/UMG/Private/Components/Widget.cpp
--- a/Engine/Source/Runtime/UMG/Private/Components/Widget.cpp
+++ b/Engine/Source/Runtime/UMG/Private/Components/Widget.cpp
@@ -310,7 +310,7 @@ FString UWidget::GetLabel() const
 	}
 }
 
-const FText UWidget::GetToolboxCategory()
+const FText UWidget::GetPaletteCategory()
 {
 	return LOCTEXT("Uncategorized", "Uncategorized");
 }
diff --git a/Engine/Source/Runtime/UMG/Private/Components/WidgetSwitcher.cpp b/Engine/Source/Runtime/UMG/Private/Components/WidgetSwitcher.cpp
--- a/Engine/Source/Runtime/UMG/Private/Components/WidgetSwitcher.cpp
+++ b/Engine/Source/Runtime/UMG/Private/Components/WidgetSwitcher.cpp
@@ -122,7 +122,7 @@ const FSlateBrush* UWidgetSwitcher::GetEditorIcon()
 	return FUMGStyle::Get().GetBrush("Widget.WidgetSwitcher");
 }
 
-const FText UWidgetSwitcher::GetToolboxCategory()
+const FText UWidgetSwitcher::GetPaletteCategory()
 {
 	return LOCTEXT("Panel", "Panel");
 }
diff --git a/Engine/Source/Runtime/UMG/Private/UserWidget.cpp b/Engine/Source/Runtime/UMG/Private/UserWidget.cpp
--- a/Engine/Source/Runtime/UMG/Private/UserWidget.cpp
+++ b/Engine/Source/Runtime/UMG/Private/UserWidget.cpp
@@ -723,7 +723,7 @@ const FSlateBrush* UUserWidget::GetEditorIcon()
 	return FUMGStyle::Get().GetBrush("Widget.UserWidget");
 }
 
-const FText UUserWidget::GetToolboxCategory()
+const FText UUserWidget::GetPaletteCategory()
 {
 	return LOCTEXT("UserCreated", "User Created");
 }
diff --git a/Engine/Source/Runtime/UMG/Public/Blueprint/UserWidget.h b/Engine/Source/Runtime/UMG/Public/Blueprint/UserWidget.h
--- a/Engine/Source/Runtime/UMG/Public/Blueprint/UserWidget.h
+++ b/Engine/Source/Runtime/UMG/Public/Blueprint/UserWidget.h
@@ -275,7 +275,7 @@ public:
 #if WITH_EDITOR
 	// UWidget interface
 	virtual const FSlateBrush* GetEditorIcon() override;
-	virtual const FText GetToolboxCategory() override;
+	virtual const FText GetPaletteCategory() override;
 	// End UWidget interface
 #endif
 
diff --git a/Engine/Source/Runtime/UMG/Public/Components/Border.h b/Engine/Source/Runtime/UMG/Public/Components/Border.h
--- a/Engine/Source/Runtime/UMG/Public/Components/Border.h
+++ b/Engine/Source/Runtime/UMG/Public/Components/Border.h
@@ -124,7 +124,7 @@ public:
 	// End of UObject interface
 
 	virtual const FSlateBrush* GetEditorIcon() override;
-	virtual const FText GetToolboxCategory() override;
+	virtual const FText GetPaletteCategory() override;
 #endif
 
 protected:
diff --git a/Engine/Source/Runtime/UMG/Public/Components/Button.h b/Engine/Source/Runtime/UMG/Public/Components/Button.h
--- a/Engine/Source/Runtime/UMG/Public/Components/Button.h
+++ b/Engine/Source/Runtime/UMG/Public/Components/Button.h
@@ -82,7 +82,7 @@ public:
 
 #if WITH_EDITOR
 	virtual const FSlateBrush* GetEditorIcon() override;
-	virtual const FText GetToolboxCategory() override;
+	virtual const FText GetPaletteCategory() override;
 #endif
 
 protected:
diff --git a/Engine/Source/Runtime/UMG/Public/Components/CanvasPanel.h b/Engine/Source/Runtime/UMG/Public/Components/CanvasPanel.h
--- a/Engine/Source/Runtime/UMG/Public/Components/CanvasPanel.h
+++ b/Engine/Source/Runtime/UMG/Public/Components/CanvasPanel.h
@@ -25,7 +25,7 @@ public:
 #if WITH_EDITOR
 	// UWidget interface
 	virtual const FSlateBrush* GetEditorIcon() override;
-	virtual const FText GetToolboxCategory() override;
+	virtual const FText GetPaletteCategory() override;
 	// End UWidget interface
 
 	// UWidget interface
diff --git a/Engine/Source/Runtime/UMG/Public/Components/CheckBox.h b/Engine/Source/Runtime/UMG/Public/Components/CheckBox.h
--- a/Engine/Source/Runtime/UMG/Public/Components/CheckBox.h
+++ b/Engine/Source/Runtime/UMG/Public/Components/CheckBox.h
@@ -121,7 +121,7 @@ public:
 
 #if WITH_EDITOR
 	virtual const FSlateBrush* GetEditorIcon() override;
-	virtual const FText GetToolboxCategory() override;
+	virtual const FText GetPaletteCategory() override;
 #endif
 
 protected:
diff --git a/Engine/Source/Runtime/UMG/Public/Components/CircularThrobber.h b/Engine/Source/Runtime/UMG/Public/Components/CircularThrobber.h
--- a/Engine/Source/Runtime/UMG/Public/Components/CircularThrobber.h
+++ b/Engine/Source/Runtime/UMG/Public/Components/CircularThrobber.h
@@ -40,7 +40,7 @@ public:
 
 #if WITH_EDITOR
 	virtual const FSlateBrush* GetEditorIcon() override;
-	virtual const FText GetToolboxCategory() override;
+	virtual const FText GetPaletteCategory() override;
 #endif
 
 protected:
diff --git a/Engine/Source/Runtime/UMG/Public/Components/ComboBox.h b/Engine/Source/Runtime/UMG/Public/Components/ComboBox.h
--- a/Engine/Source/Runtime/UMG/Public/Components/ComboBox.h
+++ b/Engine/Source/Runtime/UMG/Public/Components/ComboBox.h
@@ -26,7 +26,7 @@ public:
 
 #if WITH_EDITOR
 	virtual const FSlateBrush* GetEditorIcon() override;
-	virtual const FText GetToolboxCategory() override;
+	virtual const FText GetPaletteCategory() override;
 #endif
 
 protected:
diff --git a/Engine/Source/Runtime/UMG/Public/Components/ComboBoxString.h b/Engine/Source/Runtime/UMG/Public/Components/ComboBoxString.h
--- a/Engine/Source/Runtime/UMG/Public/Components/ComboBoxString.h
+++ b/Engine/Source/Runtime/UMG/Public/Components/ComboBoxString.h
@@ -75,7 +75,7 @@ public:
 
 #if WITH_EDITOR
 	virtual const FSlateBrush* GetEditorIcon() override;
-	virtual const FText GetToolboxCategory() override;
+	virtual const FText GetPaletteCategory() override;
 #endif
 
 protected:
diff --git a/Engine/Source/Runtime/UMG/Public/Components/EditableText.h b/Engine/Source/Runtime/UMG/Public/Components/EditableText.h
--- a/Engine/Source/Runtime/UMG/Public/Components/EditableText.h
+++ b/Engine/Source/Runtime/UMG/Public/Components/EditableText.h
@@ -149,7 +149,7 @@ public:
 
 #if WITH_EDITOR
 	virtual const FSlateBrush* GetEditorIcon() override;
-	virtual const FText GetToolboxCategory() override;
+	virtual const FText GetPaletteCategory() override;
 #endif
 	// End of UWidget interface
 
diff --git a/Engine/Source/Runtime/UMG/Public/Components/EditableTextBox.h b/Engine/Source/Runtime/UMG/Public/Components/EditableTextBox.h
--- a/Engine/Source/Runtime/UMG/Public/Components/EditableTextBox.h
+++ b/Engine/Source/Runtime/UMG/Public/Components/EditableTextBox.h
@@ -124,7 +124,7 @@ public:
 
 #if WITH_EDITOR
 	virtual const FSlateBrush* GetEditorIcon() override;
-	virtual const FText GetToolboxCategory() override;
+	virtual const FText GetPaletteCategory() override;
 #endif
 
 protected:
diff --git a/Engine/Source/Runtime/UMG/Public/Components/GridPanel.h b/Engine/Source/Runtime/UMG/Public/Components/GridPanel.h
--- a/Engine/Source/Runtime/UMG/Public/Components/GridPanel.h
+++ b/Engine/Source/Runtime/UMG/Public/Components/GridPanel.h
@@ -30,7 +30,7 @@ public:
 
 #if WITH_EDITOR
 	virtual const FSlateBrush* GetEditorIcon() override;
-	virtual const FText GetToolboxCategory() override;
+	virtual const FText GetPaletteCategory() override;
 #endif
 
 protected:
diff --git a/Engine/Source/Runtime/UMG/Public/Components/HorizontalBox.h b/Engine/Source/Runtime/UMG/Public/Components/HorizontalBox.h
--- a/Engine/Source/Runtime/UMG/Public/Components/HorizontalBox.h
+++ b/Engine/Source/Runtime/UMG/Public/Components/HorizontalBox.h
@@ -20,7 +20,7 @@ class UMG_API UHorizontalBox : public UPanelWidget
 #if WITH_EDITOR
 	// UWidget interface
 	virtual const FSlateBrush* GetEditorIcon() override;
-	virtual const FText GetToolboxCategory() override;
+	virtual const FText GetPaletteCategory() override;
 	// End UWidget interface
 #endif
 
diff --git a/Engine/Source/Runtime/UMG/Public/Components/Image.h b/Engine/Source/Runtime/UMG/Public/Components/Image.h
--- a/Engine/Source/Runtime/UMG/Public/Components/Image.h
+++ b/Engine/Source/Runtime/UMG/Public/Components/Image.h
@@ -80,7 +80,7 @@ public:
 #if WITH_EDITOR
 	// UWidget interface
 	virtual const FSlateBrush* GetEditorIcon() override;
-	virtual const FText GetToolboxCategory() override;
+	virtual const FText GetPaletteCategory() override;
 	// End UWidget interface
 #endif
 
diff --git a/Engine/Source/Runtime/UMG/Public/Components/ListView.h b/Engine/Source/Runtime/UMG/Public/Components/ListView.h
--- a/Engine/Source/Runtime/UMG/Public/Components/ListView.h
+++ b/Engine/Source/Runtime/UMG/Public/Components/ListView.h
@@ -30,7 +30,7 @@ public:
 
 #if WITH_EDITOR
 	virtual const FSlateBrush* GetEditorIcon() override;
-	virtual const FText GetToolboxCategory() override;
+	virtual const FText GetPaletteCategory() override;
 #endif
 
 protected:
diff --git a/Engine/Source/Runtime/UMG/Public/Components/MenuAnchor.h b/Engine/Source/Runtime/UMG/Public/Components/MenuAnchor.h
--- a/Engine/Source/Runtime/UMG/Public/Components/MenuAnchor.h
+++ b/Engine/Source/Runtime/UMG/Public/Components/MenuAnchor.h
@@ -56,7 +56,7 @@ public:
 
 #if WITH_EDITOR
 	virtual const FSlateBrush* GetEditorIcon() override;
-	virtual const FText GetToolboxCategory() override;
+	virtual const FText GetPaletteCategory() override;
 #endif
 
 protected:
diff --git a/Engine/Source/Runtime/UMG/Public/Components/MultiLineEditableText.h b/Engine/Source/Runtime/UMG/Public/Components/MultiLineEditableText.h
--- a/Engine/Source/Runtime/UMG/Public/Components/MultiLineEditableText.h
+++ b/Engine/Source/Runtime/UMG/Public/Components/MultiLineEditableText.h
@@ -51,7 +51,7 @@ public:
 
 #if WITH_EDITOR
 	virtual const FSlateBrush* GetEditorIcon() override;
-	virtual const FText GetToolboxCategory() override;
+	virtual const FText GetPaletteCategory() override;
 #endif
 
 protected:
diff --git a/Engine/Source/Runtime/UMG/Public/Components/MultiLineEditableTextBox.h b/Engine/Source/Runtime/UMG/Public/Components/MultiLineEditableTextBox.h
--- a/Engine/Source/Runtime/UMG/Public/Components/MultiLineEditableTextBox.h
+++ b/Engine/Source/Runtime/UMG/Public/Components/MultiLineEditableTextBox.h
@@ -73,7 +73,7 @@ public:
 
 #if WITH_EDITOR
 	virtual const FSlateBrush* GetEditorIcon() override;
-	virtual const FText GetToolboxCategory() override;
+	virtual const FText GetPaletteCategory() override;
 #endif
 
 protected:
diff --git a/Engine/Source/Runtime/UMG/Public/Components/NativeWidgetHost.h b/Engine/Source/Runtime/UMG/Public/Components/NativeWidgetHost.h
--- a/Engine/Source/Runtime/UMG/Public/Components/NativeWidgetHost.h
+++ b/Engine/Source/Runtime/UMG/Public/Components/NativeWidgetHost.h
@@ -19,7 +19,7 @@ class UMG_API UNativeWidgetHost : public UWidget
 
 #if WITH_EDITOR
 	virtual const FSlateBrush* GetEditorIcon() override;
-	virtual const FText GetToolboxCategory() override;
+	virtual const FText GetPaletteCategory() override;
 #endif
 
 protected:
diff --git a/Engine/Source/Runtime/UMG/Public/Components/Overlay.h b/Engine/Source/Runtime/UMG/Public/Components/Overlay.h
--- a/Engine/Source/Runtime/UMG/Public/Components/Overlay.h
+++ b/Engine/Source/Runtime/UMG/Public/Components/Overlay.h
@@ -13,7 +13,7 @@ class UMG_API UOverlay : public UPanelWidget
 
 #if WITH_EDITOR
 	virtual const FSlateBrush* GetEditorIcon() override;
-	virtual const FText GetToolboxCategory() override;
+	virtual const FText GetPaletteCategory() override;
 #endif
 
 protected:
diff --git a/Engine/Source/Runtime/UMG/Public/Components/ProgressBar.h b/Engine/Source/Runtime/UMG/Public/Components/ProgressBar.h
--- a/Engine/Source/Runtime/UMG/Public/Components/ProgressBar.h
+++ b/Engine/Source/Runtime/UMG/Public/Components/ProgressBar.h
@@ -78,7 +78,7 @@ public:
 #if WITH_EDITOR
 	// UWidget interface
 	virtual const FSlateBrush* GetEditorIcon() override;
-	virtual const FText GetToolboxCategory() override;
+	virtual const FText GetPaletteCategory() override;
 	// End UWidget interface
 #endif
 
diff --git a/Engine/Source/Runtime/UMG/Public/Components/ScaleBox.h b/Engine/Source/Runtime/UMG/Public/Components/ScaleBox.h
--- a/Engine/Source/Runtime/UMG/Public/Components/ScaleBox.h
+++ b/Engine/Source/Runtime/UMG/Public/Components/ScaleBox.h
@@ -34,7 +34,7 @@ public:
 
 #if WITH_EDITOR
 	virtual const FSlateBrush* GetEditorIcon() override;
-	virtual const FText GetToolboxCategory() override;
+	virtual const FText GetPaletteCategory() override;
 #endif
 
 protected:
diff --git a/Engine/Source/Runtime/UMG/Public/Components/ScrollBar.h b/Engine/Source/Runtime/UMG/Public/Components/ScrollBar.h
--- a/Engine/Source/Runtime/UMG/Public/Components/ScrollBar.h
+++ b/Engine/Source/Runtime/UMG/Public/Components/ScrollBar.h
@@ -64,7 +64,7 @@ public:
 #if WITH_EDITOR
 	// UWidget interface
 	virtual const FSlateBrush* GetEditorIcon() override;
-	virtual const FText GetToolboxCategory() override;
+	virtual const FText GetPaletteCategory() override;
 	// End UWidget interface
 #endif
 
diff --git a/Engine/Source/Runtime/UMG/Public/Components/ScrollBox.h b/Engine/Source/Runtime/UMG/Public/Components/ScrollBox.h
--- a/Engine/Source/Runtime/UMG/Public/Components/ScrollBox.h
+++ b/Engine/Source/Runtime/UMG/Public/Components/ScrollBox.h
@@ -47,7 +47,7 @@ class UMG_API UScrollBox : public UPanelWidget
 #if WITH_EDITOR
 	// UWidget interface
 	virtual const FSlateBrush* GetEditorIcon() override;
-	virtual const FText GetToolboxCategory() override;
+	virtual const FText GetPaletteCategory() override;
 	// End UWidget interface
 #endif
 
diff --git a/Engine/Source/Runtime/UMG/Public/Components/SizeBox.h b/Engine/Source/Runtime/UMG/Public/Components/SizeBox.h
--- a/Engine/Source/Runtime/UMG/Public/Components/SizeBox.h
+++ b/Engine/Source/Runtime/UMG/Public/Components/SizeBox.h
@@ -71,7 +71,7 @@ public:
 
 #if WITH_EDITOR
 	virtual const FSlateBrush* GetEditorIcon() override;
-	virtual const FText GetToolboxCategory() override;
+	virtual const FText GetPaletteCategory() override;
 #endif
 
 protected:
diff --git a/Engine/Source/Runtime/UMG/Public/Components/Slider.h b/Engine/Source/Runtime/UMG/Public/Components/Slider.h
--- a/Engine/Source/Runtime/UMG/Public/Components/Slider.h
+++ b/Engine/Source/Runtime/UMG/Public/Components/Slider.h
@@ -74,7 +74,7 @@ public:
 
 #if WITH_EDITOR
 	virtual const FSlateBrush* GetEditorIcon() override;
-	virtual const FText GetToolboxCategory() override;
+	virtual const FText GetPaletteCategory() override;
 #endif
 
 protected:
diff --git a/Engine/Source/Runtime/UMG/Public/Components/Spacer.h b/Engine/Source/Runtime/UMG/Public/Components/Spacer.h
--- a/Engine/Source/Runtime/UMG/Public/Components/Spacer.h
+++ b/Engine/Source/Runtime/UMG/Public/Components/Spacer.h
@@ -28,7 +28,7 @@ public:
 
 #if WITH_EDITOR
 	virtual const FSlateBrush* GetEditorIcon() override;
-	virtual const FText GetToolboxCategory() override;
+	virtual const FText GetPaletteCategory() override;
 #endif
 
 protected:
diff --git a/Engine/Source/Runtime/UMG/Public/Components/SpinBox.h b/Engine/Source/Runtime/UMG/Public/Components/SpinBox.h
--- a/Engine/Source/Runtime/UMG/Public/Components/SpinBox.h
+++ b/Engine/Source/Runtime/UMG/Public/Components/SpinBox.h
@@ -134,7 +134,7 @@ public:
 
 #if WITH_EDITOR
 	virtual const FSlateBrush* GetEditorIcon() override;
-	virtual const FText GetToolboxCategory() override;
+	virtual const FText GetPaletteCategory() override;
 #endif
 	// End of UWidget interface
 
diff --git a/Engine/Source/Runtime/UMG/Public/Components/TextBlock.h b/Engine/Source/Runtime/UMG/Public/Components/TextBlock.h
--- a/Engine/Source/Runtime/UMG/Public/Components/TextBlock.h
+++ b/Engine/Source/Runtime/UMG/Public/Components/TextBlock.h
@@ -110,7 +110,7 @@ public:
 #if WITH_EDITOR
 	// UWidget interface
 	virtual const FSlateBrush* GetEditorIcon() override;
-	virtual const FText GetToolboxCategory() override;
+	virtual const FText GetPaletteCategory() override;
 	// End UWidget interface
 
 	virtual FString GetLabelMetadata() const override;
diff --git a/Engine/Source/Runtime/UMG/Public/Components/Throbber.h b/Engine/Source/Runtime/UMG/Public/Components/Throbber.h
--- a/Engine/Source/Runtime/UMG/Public/Components/Throbber.h
+++ b/Engine/Source/Runtime/UMG/Public/Components/Throbber.h
@@ -44,7 +44,7 @@ public:
 
 #if WITH_EDITOR
 	virtual const FSlateBrush* GetEditorIcon() override;
-	virtual const FText GetToolboxCategory() override;
+	virtual const FText GetPaletteCategory() override;
 #endif
 
 protected:
diff --git a/Engine/Source/Runtime/UMG/Public/Components/TileView.h b/Engine/Source/Runtime/UMG/Public/Components/TileView.h
--- a/Engine/Source/Runtime/UMG/Public/Components/TileView.h
+++ b/Engine/Source/Runtime/UMG/Public/Components/TileView.h
@@ -48,7 +48,7 @@ public:
 
 #if WITH_EDITOR
 	virtual const FSlateBrush* GetEditorIcon() override;
-	virtual const FText GetToolboxCategory() override;
+	virtual const FText GetPaletteCategory() override;
 #endif
 
 protected:
diff --git a/Engine/Source/Runtime/UMG/Public/Components/UniformGridPanel.h b/Engine/Source/Runtime/UMG/Public/Components/UniformGridPanel.h
--- a/Engine/Source/Runtime/UMG/Public/Components/UniformGridPanel.h
+++ b/Engine/Source/Runtime/UMG/Public/Components/UniformGridPanel.h
@@ -34,7 +34,7 @@ public:
 
 #if WITH_EDITOR
 	virtual const FSlateBrush* GetEditorIcon() override;
-	virtual const FText GetToolboxCategory() override;
+	virtual const FText GetPaletteCategory() override;
 #endif
 
 protected:
diff --git a/Engine/Source/Runtime/UMG/Public/Components/VerticalBox.h b/Engine/Source/Runtime/UMG/Public/Components/VerticalBox.h
--- a/Engine/Source/Runtime/UMG/Public/Components/VerticalBox.h
+++ b/Engine/Source/Runtime/UMG/Public/Components/VerticalBox.h
@@ -22,7 +22,7 @@ public:
 #if WITH_EDITOR
 	// UWidget interface
 	virtual const FSlateBrush* GetEditorIcon() override;
-	virtual const FText GetToolboxCategory() override;
+	virtual const FText GetPaletteCategory() override;
 	// End UWidget interface
 #endif
 
diff --git a/Engine/Source/Runtime/UMG/Public/Components/Viewport.h b/Engine/Source/Runtime/UMG/Public/Components/Viewport.h
--- a/Engine/Source/Runtime/UMG/Public/Components/Viewport.h
+++ b/Engine/Source/Runtime/UMG/Public/Components/Viewport.h
@@ -272,7 +272,7 @@ class UMG_API UViewport : public UContentWidget
 
 #if WITH_EDITOR
 	virtual const FSlateBrush* GetEditorIcon() override;
-	virtual const FText GetToolboxCategory() override;
+	virtual const FText GetPaletteCategory() override;
 #endif
 
 protected:
diff --git a/Engine/Source/Runtime/UMG/Public/Components/Widget.h b/Engine/Source/Runtime/UMG/Public/Components/Widget.h
--- a/Engine/Source/Runtime/UMG/Public/Components/Widget.h
+++ b/Engine/Source/Runtime/UMG/Public/Components/Widget.h
@@ -242,8 +242,8 @@ public:
 	/** Gets the label to display to the user for this widget. */
 	FString GetLabel() const;
 
-	/** Gets the toolbox category of the widget */
-	virtual const FText GetToolboxCategory();
+	/** Gets the palette category of the widget */
+	virtual const FText GetPaletteCategory();
 
 	/** Gets the editor icon */
 	virtual const FSlateBrush* GetEditorIcon();
diff --git a/Engine/Source/Runtime/UMG/Public/Components/WidgetSwitcher.h b/Engine/Source/Runtime/UMG/Public/Components/WidgetSwitcher.h
--- a/Engine/Source/Runtime/UMG/Public/Components/WidgetSwitcher.h
+++ b/Engine/Source/Runtime/UMG/Public/Components/WidgetSwitcher.h
@@ -42,7 +42,7 @@ public:
 
 #if WITH_EDITOR
 	virtual const FSlateBrush* GetEditorIcon() override;
-	virtual const FText GetToolboxCategory() override;
+	virtual const FText GetPaletteCategory() override;
 	virtual void OnDescendantSelected(UWidget* DescendantWidget) override;
 	virtual void OnDescendantDeselected(UWidget* DescendantWidget) override;
 #endif
