Added toolbar icon to toolbar tabs made by asset editor toolkit

diff --git a/Engine/Source/Editor/UMGEditor/Private/Templates/WidgetTemplateClass.cpp b/Engine/Source/Editor/UMGEditor/Private/Templates/WidgetTemplateClass.cpp
--- a/Engine/Source/Editor/UMGEditor/Private/Templates/WidgetTemplateClass.cpp
+++ b/Engine/Source/Editor/UMGEditor/Private/Templates/WidgetTemplateClass.cpp
@@ -16,7 +16,7 @@ FWidgetTemplateClass::FWidgetTemplateClass(TSubclassOf<UWidget> InWidgetClass)
 FText FWidgetTemplateClass::GetCategory() const
 {
 	auto DefaultWidget = WidgetClass->GetDefaultObject<UWidget>();
-	return DefaultWidget->GetToolboxCategory();
+	return DefaultWidget->GetPaletteCategory();
 }
 
 UWidget* FWidgetTemplateClass::Create(UWidgetTree* Tree)
diff --git a/Engine/Source/Editor/UnrealEd/Private/Toolkits/AssetEditorToolkit.cpp b/Engine/Source/Editor/UnrealEd/Private/Toolkits/AssetEditorToolkit.cpp
--- a/Engine/Source/Editor/UnrealEd/Private/Toolkits/AssetEditorToolkit.cpp
+++ b/Engine/Source/Editor/UnrealEd/Private/Toolkits/AssetEditorToolkit.cpp
@@ -594,6 +594,7 @@ TSharedRef<SDockTab> FAssetEditorToolkit::SpawnTab_Toolbar( const FSpawnTabArgs&
 
 	TSharedRef<SDockTab> DockTab = SNew(SDockTab)
 		.Label( NSLOCTEXT("AssetEditorToolkit", "Toolbar_TabTitle", "Toolbar") )
+		.Icon(FEditorStyle::GetBrush("LevelEditor.Tabs.Toolbar"))
 		.ShouldAutosize(true)
 		[
 			SAssignNew(ToolbarWidgetContent, SBorder)
