UMG - Adding a new customization that can customize ETextJustify, works just like th ones for halign/valign.

diff --git a/Engine/Source/Editor/UMGEditor/Private/Customizations/TextJustifyCustomization.cpp b/Engine/Source/Editor/UMGEditor/Private/Customizations/TextJustifyCustomization.cpp
new file mode 100644
--- /dev/null
+++ b/Engine/Source/Editor/UMGEditor/Private/Customizations/TextJustifyCustomization.cpp
@@ -0,0 +1,99 @@
+// Copyright 1998-2014 Epic Games, Inc. All Rights Reserved.
+
+#include "UMGEditorPrivatePCH.h"
+
+#include "PropertyEditing.h"
+#include "ObjectEditorUtils.h"
+#include "WidgetGraphSchema.h"
+#include "ScopedTransaction.h"
+#include "BlueprintEditorUtils.h"
+
+#include "TextJustifyCustomization.h"
+
+#define LOCTEXT_NAMESPACE "UMG"
+
+void FTextJustifyCustomization::CustomizeHeader(TSharedRef<IPropertyHandle> PropertyHandle, FDetailWidgetRow& HeaderRow, IPropertyTypeCustomizationUtils& CustomizationUtils)
+{
+	const FMargin OuterPadding(2);
+	const FMargin ContentPadding(2);
+
+	HeaderRow
+	.NameContent()
+	[
+		PropertyHandle->CreatePropertyNameWidget()
+	]
+	.ValueContent()
+	[
+		SNew(SHorizontalBox)
+
+		+ SHorizontalBox::Slot()
+		.AutoWidth()
+		.Padding(OuterPadding)
+		[
+			SNew( SCheckBox )
+			.Style(FEditorStyle::Get(), "ToggleButtonCheckbox")
+			.ToolTipText(LOCTEXT("HAlignLeft", "Align Text Left"))
+			.Padding(ContentPadding)
+			.OnCheckStateChanged(this, &FTextJustifyCustomization::HandleCheckStateChanged, PropertyHandle, ETextJustify::Left)
+			.IsChecked(this, &FTextJustifyCustomization::GetCheckState, PropertyHandle, ETextJustify::Left)
+			[
+				SNew(SImage)
+				.Image(FEditorStyle::GetBrush("HorizontalAlignment_Left"))
+			]
+		]
+
+		+ SHorizontalBox::Slot()
+		.AutoWidth()
+		.Padding(OuterPadding)
+		[
+			SNew(SCheckBox)
+			.Style(FEditorStyle::Get(), "ToggleButtonCheckbox")
+			.ToolTipText(LOCTEXT("HAlignCenter", "Align Text Center"))
+			.Padding(ContentPadding)
+			.OnCheckStateChanged(this, &FTextJustifyCustomization::HandleCheckStateChanged, PropertyHandle, ETextJustify::Center)
+			.IsChecked(this, &FTextJustifyCustomization::GetCheckState, PropertyHandle, ETextJustify::Center)
+			[
+				SNew(SImage)
+				.Image(FEditorStyle::GetBrush("HorizontalAlignment_Center"))
+			]
+		]
+
+		+ SHorizontalBox::Slot()
+		.AutoWidth()
+		.Padding(OuterPadding)
+		[
+			SNew(SCheckBox)
+			.Style(FEditorStyle::Get(), "ToggleButtonCheckbox")
+			.ToolTipText(LOCTEXT("HAlignRight", "Align Text Right"))
+			.Padding(ContentPadding)
+			.OnCheckStateChanged(this, &FTextJustifyCustomization::HandleCheckStateChanged, PropertyHandle, ETextJustify::Right)
+			.IsChecked(this, &FTextJustifyCustomization::GetCheckState, PropertyHandle, ETextJustify::Right)
+			[
+				SNew(SImage)
+				.Image(FEditorStyle::GetBrush("HorizontalAlignment_Right"))
+			]
+		]
+	];
+}
+
+void FTextJustifyCustomization::CustomizeChildren(TSharedRef<IPropertyHandle> PropertyHandle, IDetailChildrenBuilder& ChildBuilder, IPropertyTypeCustomizationUtils& CustomizationUtils)
+{
+}
+
+void FTextJustifyCustomization::HandleCheckStateChanged(ESlateCheckBoxState::Type InCheckboxState, TSharedRef<IPropertyHandle> PropertyHandle, ETextJustify::Type ToAlignment)
+{
+	PropertyHandle->SetValue((uint8)ToAlignment);
+}
+
+ESlateCheckBoxState::Type FTextJustifyCustomization::GetCheckState(TSharedRef<IPropertyHandle> PropertyHandle, ETextJustify::Type ForAlignment) const
+{
+	uint8 Value;
+	if ( PropertyHandle->GetValue(Value) )
+	{
+		return Value == ForAlignment ? ESlateCheckBoxState::Checked : ESlateCheckBoxState::Unchecked;
+	}
+
+	return ESlateCheckBoxState::Unchecked;
+}
+
+#undef LOCTEXT_NAMESPACE
diff --git a/Engine/Source/Editor/UMGEditor/Private/Customizations/TextJustifyCustomization.h b/Engine/Source/Editor/UMGEditor/Private/Customizations/TextJustifyCustomization.h
new file mode 100644
--- /dev/null
+++ b/Engine/Source/Editor/UMGEditor/Private/Customizations/TextJustifyCustomization.h
@@ -0,0 +1,25 @@
+// Copyright 1998-2014 Epic Games, Inc. All Rights Reserved.
+
+#pragma once
+
+class FTextJustifyCustomization : public IPropertyTypeCustomization
+{
+public:
+	/** Makes a new instance of this detail layout class for a specific detail view requesting it */
+	static TSharedRef<class IPropertyTypeCustomization> MakeInstance()
+	{
+		return MakeShareable(new FTextJustifyCustomization());
+	}
+
+	FTextJustifyCustomization()
+	{
+	}
+	
+	virtual void CustomizeHeader(TSharedRef<IPropertyHandle> PropertyHandle, class FDetailWidgetRow& HeaderRow, IPropertyTypeCustomizationUtils& CustomizationUtils) override;
+
+	virtual void CustomizeChildren(TSharedRef<IPropertyHandle> PropertyHandle, class IDetailChildrenBuilder& ChildBuilder, IPropertyTypeCustomizationUtils& CustomizationUtils) override;
+	
+	void HandleCheckStateChanged(ESlateCheckBoxState::Type InCheckboxState, TSharedRef<IPropertyHandle> PropertyHandle, ETextJustify::Type ToAlignment);
+
+	ESlateCheckBoxState::Type GetCheckState(TSharedRef<IPropertyHandle> PropertyHandle, ETextJustify::Type ForAlignment) const;
+};
diff --git a/Engine/Source/Editor/UMGEditor/Private/Details/SWidgetDetailsView.cpp b/Engine/Source/Editor/UMGEditor/Private/Details/SWidgetDetailsView.cpp
--- a/Engine/Source/Editor/UMGEditor/Private/Details/SWidgetDetailsView.cpp
+++ b/Engine/Source/Editor/UMGEditor/Private/Details/SWidgetDetailsView.cpp
@@ -16,6 +16,7 @@
 #include "HorizontalAlignmentCustomization.h"
 #include "VerticalAlignmentCustomization.h"
 #include "SlateChildSizeCustomization.h"
+#include "TextJustifyCustomization.h"
 
 #define LOCTEXT_NAMESPACE "UMG"
 
@@ -128,6 +129,7 @@ SWidgetDetailsView::~SWidgetDetailsView()
 	PropertyModule.UnregisterCustomPropertyTypeLayout(TEXT("EVerticalAlignment"), nullptr, PropertyView);
 	PropertyModule.UnregisterCustomPropertyTypeLayout(TEXT("SlateChildSize"), nullptr, PropertyView);
 	PropertyModule.UnregisterCustomPropertyTypeLayout(TEXT("SlateBrush"), nullptr, PropertyView);
+	PropertyModule.UnregisterCustomPropertyTypeLayout(TEXT("ETextJustify"), nullptr, PropertyView);
 }
 
 void SWidgetDetailsView::RegisterCustomizations()
@@ -142,6 +144,7 @@ void SWidgetDetailsView::RegisterCustomizations()
 	PropertyModule.RegisterCustomPropertyTypeLayout(TEXT("EVerticalAlignment"), FOnGetPropertyTypeCustomizationInstance::CreateStatic(&FVerticalAlignmentCustomization::MakeInstance), nullptr, PropertyView);
 	PropertyModule.RegisterCustomPropertyTypeLayout(TEXT("SlateChildSize"), FOnGetPropertyTypeCustomizationInstance::CreateStatic(&FSlateChildSizeCustomization::MakeInstance), nullptr, PropertyView);
 	PropertyModule.RegisterCustomPropertyTypeLayout(TEXT("SlateBrush"), FOnGetPropertyTypeCustomizationInstance::CreateStatic(&FSlateBrushStructCustomization::MakeInstance, false), nullptr, PropertyView);
+	PropertyModule.RegisterCustomPropertyTypeLayout(TEXT("ETextJustify"), FOnGetPropertyTypeCustomizationInstance::CreateStatic(&FTextJustifyCustomization::MakeInstance), nullptr, PropertyView);
 }
 
 void SWidgetDetailsView::OnEditorSelectionChanging()
