Fixed slot animation not working if the slot regenerates

diff --git a/Engine/Source/Editor/UMGEditor/Private/Animation/UMGSequencerObjectBindingManager.cpp b/Engine/Source/Editor/UMGEditor/Private/Animation/UMGSequencerObjectBindingManager.cpp
--- a/Engine/Source/Editor/UMGEditor/Private/Animation/UMGSequencerObjectBindingManager.cpp
+++ b/Engine/Source/Editor/UMGEditor/Private/Animation/UMGSequencerObjectBindingManager.cpp
@@ -39,16 +39,25 @@ void FUMGSequencerObjectBindingManager::BindPossessableObject( const FGuid& Poss
 	PreviewObjectToGuidMap.Add( &PossessedObject, PossessableGuid );
 
 	FWidgetAnimationBinding NewBinding;
-	NewBinding.WidgetName = PossessedObject.GetFName();
-	if( PossessedObject.IsA<UPanelSlot>() )
+	NewBinding.AnimationGuid = PossessableGuid;
+
+	UPanelSlot* Slot = Cast<UPanelSlot>( &PossessedObject );
+	if( Slot && Slot->Content )
 	{
-		// Save the name of the widget containing the slots.  This is the object to look up that contaisn the slot itself(the thing we are animating)
-		NewBinding.SlotWidgetName = PossessedObject.GetOuter()->GetFName();
+		// Save the name of the widget containing the slots.  This is the object to look up that contains the slot itself(the thing we are animating)
+		NewBinding.SlotWidgetName = Slot->GetFName();
+		NewBinding.WidgetName = Slot->Content->GetFName();
+
+
+		WidgetAnimation->AnimationBindings.Add(NewBinding);
 	}
+	else if( !Slot )
+	{
+		NewBinding.WidgetName = PossessedObject.GetFName();
 
-	NewBinding.AnimationGuid = PossessableGuid;
+		WidgetAnimation->AnimationBindings.Add(NewBinding);
+	}
 
-	WidgetAnimation->AnimationBindings.Add( NewBinding );
 }
 
 void FUMGSequencerObjectBindingManager::UnbindPossessableObjects( const FGuid& PossessableGuid )
@@ -96,27 +105,10 @@ void FUMGSequencerObjectBindingManager::InitPreviewObjects()
 
 		for(const FWidgetAnimationBinding& Binding : WidgetAnimation->AnimationBindings)
 		{
-			FName ObjectName = Binding.WidgetName;
-			FName SlotWidgetName = Binding.SlotWidgetName;
-
-			FName NameToSearchFor = SlotWidgetName != NAME_None ? SlotWidgetName : ObjectName;
-
-			UObject* FoundObject = FindObject<UObject>(WidgetTree, *NameToSearchFor.ToString());
-			if(FoundObject)
+			UObject* FoundObject = Binding.FindRuntimeObject( *WidgetTree );
+			if( FoundObject )
 			{
-				if(SlotWidgetName == NAME_None)
-				{
-					PreviewObjectToGuidMap.Add(FoundObject, Binding.AnimationGuid);
-				}
-				else
-				{
-					FoundObject = FindObject<UObject>(FoundObject, *ObjectName.ToString());
-					if(FoundObject)
-					{
-						PreviewObjectToGuidMap.Add(FoundObject, Binding.AnimationGuid);
-					}
-				}
-
+				PreviewObjectToGuidMap.Add(FoundObject, Binding.AnimationGuid);
 			}
 		}
 	}
diff --git a/Engine/Source/Runtime/UMG/Private/Animation/UMGSequencePlayer.cpp b/Engine/Source/Runtime/UMG/Private/Animation/UMGSequencePlayer.cpp
--- a/Engine/Source/Runtime/UMG/Private/Animation/UMGSequencePlayer.cpp
+++ b/Engine/Source/Runtime/UMG/Private/Animation/UMGSequencePlayer.cpp
@@ -34,28 +34,12 @@ void UUMGSequencePlayer::InitSequencePlayer( const UWidgetAnimation& InAnimation
 	// Bind to Runtime Objects
 	for (const FWidgetAnimationBinding& Binding : InAnimation.AnimationBindings)
 	{
-		FName ObjectName = Binding.WidgetName;
-		FName SlotWidgetName = Binding.SlotWidgetName;
+		UObject* FoundObject = Binding.FindRuntimeObject( *WidgetTree );
 
-		FName NameToSearchFor = SlotWidgetName != NAME_None ? SlotWidgetName : ObjectName;
-
-		UObject* FoundObject = FindObject<UObject>(WidgetTree, *NameToSearchFor.ToString());
-		if(FoundObject)
+		if( FoundObject )
 		{
-			if(SlotWidgetName == NAME_None)
-			{
-				TArray<UObject*>& Objects = GuidToRuntimeObjectMap.FindOrAdd(Binding.AnimationGuid);
-				Objects.Add(FoundObject);
-			}
-			else
-			{
-				FoundObject = FindObject<UObject>(FoundObject, *ObjectName.ToString());
-				if(FoundObject)
-				{
-					TArray<UObject*>& Objects = GuidToRuntimeObjectMap.FindOrAdd(Binding.AnimationGuid);
-					Objects.Add(FoundObject);
-				}
-			}
+			TArray<UObject*>& Objects = GuidToRuntimeObjectMap.FindOrAdd(Binding.AnimationGuid);
+			Objects.Add(FoundObject);
 		}
 	}
 
diff --git a/Engine/Source/Runtime/UMG/Private/Animation/WidgetAnimation.cpp b/Engine/Source/Runtime/UMG/Private/Animation/WidgetAnimation.cpp
--- a/Engine/Source/Runtime/UMG/Private/Animation/WidgetAnimation.cpp
+++ b/Engine/Source/Runtime/UMG/Private/Animation/WidgetAnimation.cpp
@@ -4,6 +4,25 @@
 #include "MovieScene.h"
 #include "WidgetAnimation.h"
 
+UObject* FWidgetAnimationBinding::FindRuntimeObject( UWidgetTree& WidgetTree ) const
+{
+	UObject* FoundObject = FindObject<UObject>(&WidgetTree, *WidgetName.ToString());
+	if(FoundObject)
+	{
+		if(SlotWidgetName != NAME_None)
+		{
+			// if we were animating the slot, look up the slot that contains the widget 
+			UWidget* WidgetObject = Cast<UWidget>(FoundObject);
+			if(WidgetObject->Slot)
+			{
+				FoundObject = WidgetObject->Slot;
+			}
+		}
+	}
+
+	return FoundObject;
+}
+
 UWidgetAnimation::UWidgetAnimation(const FPostConstructInitializeProperties& PCIP)
 	: Super(PCIP)
 {
diff --git a/Engine/Source/Runtime/UMG/Public/Animation/WidgetAnimation.h b/Engine/Source/Runtime/UMG/Public/Animation/WidgetAnimation.h
--- a/Engine/Source/Runtime/UMG/Public/Animation/WidgetAnimation.h
+++ b/Engine/Source/Runtime/UMG/Public/Animation/WidgetAnimation.h
@@ -23,6 +23,13 @@ struct FWidgetAnimationBinding
 	UPROPERTY()
 	FGuid AnimationGuid;
 
+public:
+	/**
+	 * Locates a runtime object to animate from the provided tree of widgets
+	 * @return the runtime object to animate or null if not found 
+	 */
+	UMG_API UObject* FindRuntimeObject(UWidgetTree& WidgetTree) const;
+
 	bool operator==(const FWidgetAnimationBinding& Other) const
 	{
 		return WidgetName == Other.WidgetName && SlotWidgetName == Other.SlotWidgetName && AnimationGuid == Other.AnimationGuid;
