Fixed recompiling widget blueprints sometimes crashing due to stale widgets

diff --git a/Engine/Source/Editor/UMGEditor/Private/Designer/SDesignerView.cpp b/Engine/Source/Editor/UMGEditor/Private/Designer/SDesignerView.cpp
--- a/Engine/Source/Editor/UMGEditor/Private/Designer/SDesignerView.cpp
+++ b/Engine/Source/Editor/UMGEditor/Private/Designer/SDesignerView.cpp
@@ -142,6 +142,8 @@ void SDesignerView::Construct(const FArguments& InArgs, TSharedPtr<FWidgetBluepr
 	UWidgetBlueprint* Blueprint = GetBlueprint();
 	Blueprint->OnChanged().AddSP(this, &SDesignerView::OnBlueprintChanged);
 
+	FWidgetBlueprintCompiler::OnWidgetBlueprintCompiled.AddSP( this, &SDesignerView::OnBlueprintCompiled );
+
 	BindCommands();
 
 	SDesignSurface::Construct(SDesignSurface::FArguments()
@@ -620,6 +622,15 @@ void SDesignerView::OnBlueprintChanged(UBlueprint* InBlueprint)
 	}
 }
 
+void SDesignerView::OnBlueprintCompiled(UBlueprint* InBlueprint)
+{
+	UBlueprint* BP = BlueprintEditor.Pin()->GetBlueprintObj();
+	if( BP == InBlueprint )
+	{
+		UpdatePreviewWidget();
+	}
+}
+
 void SDesignerView::OnObjectPropertyChanged(UObject* ObjectBeingModified, FPropertyChangedEvent& PropertyChangedEvent)
 {
 	if ( !ensure(ObjectBeingModified) )
diff --git a/Engine/Source/Editor/UMGEditor/Private/Designer/SDesignerView.h b/Engine/Source/Editor/UMGEditor/Private/Designer/SDesignerView.h
--- a/Engine/Source/Editor/UMGEditor/Private/Designer/SDesignerView.h
+++ b/Engine/Source/Editor/UMGEditor/Private/Designer/SDesignerView.h
@@ -101,6 +101,10 @@ private:
 
 	/** Called whenever the blueprint is structurally changed. */
 	void OnBlueprintChanged(UBlueprint* InBlueprint);
+
+        /** Called whenever the blueprint is recompiled */
+	void OnBlueprintCompiled(UBlueprint* InBlueprint);
+
 	void OnObjectPropertyChanged(UObject* ObjectBeingModified, FPropertyChangedEvent& PropertyChangedEvent);
 
 	void CacheSelectedWidgetGeometry();
diff --git a/Engine/Source/Editor/UMGEditor/Private/WidgetBlueprintCompiler.cpp b/Engine/Source/Editor/UMGEditor/Private/WidgetBlueprintCompiler.cpp
--- a/Engine/Source/Editor/UMGEditor/Private/WidgetBlueprintCompiler.cpp
+++ b/Engine/Source/Editor/UMGEditor/Private/WidgetBlueprintCompiler.cpp
@@ -9,6 +9,7 @@
 #include "MovieScene.h"
 #include "K2Node_FunctionEntry.h"
 
+FOnWidgetBlueprintCompiled FWidgetBlueprintCompiler::OnWidgetBlueprintCompiled = FOnWidgetBlueprintCompiled();
 
 ///-------------------------------------------------------------
 
@@ -273,6 +274,8 @@ void FWidgetBlueprintCompiler::Compile()
 {
 	Super::Compile();
 
+	OnWidgetBlueprintCompiled.Broadcast( WidgetBlueprint() );
+
 	WidgetToMemberVariableMap.Empty();
 }
 
diff --git a/Engine/Source/Editor/UMGEditor/Private/WidgetBlueprintCompiler.h b/Engine/Source/Editor/UMGEditor/Private/WidgetBlueprintCompiler.h
--- a/Engine/Source/Editor/UMGEditor/Private/WidgetBlueprintCompiler.h
+++ b/Engine/Source/Editor/UMGEditor/Private/WidgetBlueprintCompiler.h
@@ -7,12 +7,17 @@
 //////////////////////////////////////////////////////////////////////////
 // FWidgetBlueprintCompiler
 
+DECLARE_MULTICAST_DELEGATE_OneParam( FOnWidgetBlueprintCompiled, UBlueprint* )
+
 class FWidgetBlueprintCompiler : public FKismetCompilerContext
 {
 protected:
 	typedef FKismetCompilerContext Super;
 
 public:
+	/** Delegate called when a widget blueprint is compiled */
+	static FOnWidgetBlueprintCompiled OnWidgetBlueprintCompiled;
+
 	FWidgetBlueprintCompiler(UWidgetBlueprint* SourceSketch, FCompilerResultsLog& InMessageLog, const FKismetCompilerOptions& InCompilerOptions, TArray<UObject*>* InObjLoaded);
 	virtual ~FWidgetBlueprintCompiler();
 
