UMG  - Adding grid snapping to the translation layout mode.  Fixing some bugs with resolving pending selections when the root is selected.  Caching lots of layout information in the designer now for all widgets so that it can be referenced at any time.

diff --git a/Engine/Source/Editor/EditorStyle/Private/SlateEditorStyle.cpp b/Engine/Source/Editor/EditorStyle/Private/SlateEditorStyle.cpp
--- a/Engine/Source/Editor/EditorStyle/Private/SlateEditorStyle.cpp
+++ b/Engine/Source/Editor/EditorStyle/Private/SlateEditorStyle.cpp
@@ -5787,11 +5787,16 @@ void FSlateEditorStyle::FStyle::SetupAutomationStyles()
 
 void FSlateEditorStyle::FStyle::SetupUMGEditorStyles()
 {
+	const FLinearColor IconColor = FLinearColor::Black;
+
 	Set("WidgetDesigner.LayoutTransform", new IMAGE_BRUSH("Icons/UMG/Layout_TransformMode_16x", Icon16x16));
 	Set("WidgetDesigner.LayoutTransform.Small", new IMAGE_BRUSH("Icons/UMG/Layout_TransformMode_16x", Icon16x16));
 	Set("WidgetDesigner.RenderTransform", new IMAGE_BRUSH("Icons/UMG/Render_TransformMode_16x", Icon16x16));
 	Set("WidgetDesigner.RenderTransform.Small", new IMAGE_BRUSH("Icons/UMG/Render_TransformMode_16x", Icon16x16));
 
+	Set("WidgetDesigner.LocationGridSnap", new IMAGE_BRUSH("Old/LevelEditor/LocationGridSnap", Icon14x14, IconColor));
+	Set("WidgetDesigner.RotationGridSnap", new IMAGE_BRUSH("Old/LevelEditor/RotationGridSnap", Icon14x14, IconColor));
+
 	Set("WidgetDesigner.ZoomToFit", new IMAGE_BRUSH("Icons/UMG/Fit_16x", Icon16x16));
 	Set("WidgetDesigner.ZoomToFit.Small", new IMAGE_BRUSH("Icons/UMG/Fit_16x", Icon16x16));
 
@@ -5823,6 +5828,9 @@ void FSlateEditorStyle::FStyle::SetupUMGEditorStyles()
 	Set("UMGEditor.AnchorGizmo.BottomRight.Hovered", new IMAGE_BRUSH("Icons/UMG/AnchorGizmo/bottomright", FVector2D(24, 24), FLinearColor(0, 1, 0)));
 
 
+	Set("UMGEditor.AnchoredWidget", new BOX_BRUSH("Common/Button", FVector2D(32, 32), 8.0f / 32.0f));
+	Set("UMGEditor.AnchoredWidgetAlignment", new IMAGE_BRUSH("Icons/icon_tab_DeviceManager_16x", Icon8x8));
+
 	Set("UMGEditor.TopLeft", new IMAGE_BRUSH("Icons/UMG/umg_anchors_top_left", Icon64x64));
 	Set("UMGEditor.TopCenter", new IMAGE_BRUSH("Icons/UMG/umg_anchors_top_center", Icon64x64));
 	Set("UMGEditor.TopRight", new IMAGE_BRUSH("Icons/UMG/umg_anchors_top_right", Icon64x64));
diff --git a/Engine/Source/Editor/Sequencer/Private/SequencerObjectChangeListener.cpp b/Engine/Source/Editor/Sequencer/Private/SequencerObjectChangeListener.cpp
--- a/Engine/Source/Editor/Sequencer/Private/SequencerObjectChangeListener.cpp
+++ b/Engine/Source/Editor/Sequencer/Private/SequencerObjectChangeListener.cpp
@@ -168,7 +168,7 @@ void FSequencerObjectChangeListener::OnObjectPreEditChange( UObject* Object, con
 		FPropertyEditorModule& PropertyEditor = FModuleManager::Get().LoadModuleChecked<FPropertyEditorModule>( "PropertyEditor" );
 
 		// Sometimes due to edit inline new the object is not actually the object that contains the property
-		if( IsObjectValidForListening(Object) && PropertyChain.GetActiveMemberNode()->GetValue()->IsIn( Object->GetClass() ) )
+		if ( IsObjectValidForListening(Object) && Object->GetClass()->HasProperty(PropertyChain.GetActiveMemberNode()->GetValue()) )
 		{
 			TSharedPtr<IPropertyChangeListener> PropertyChangeListener = ActivePropertyChangeListeners.FindRef( Object );
 
diff --git a/Engine/Source/Editor/UMGEditor/Private/Designer/DesignerCommands.cpp b/Engine/Source/Editor/UMGEditor/Private/Designer/DesignerCommands.cpp
--- a/Engine/Source/Editor/UMGEditor/Private/Designer/DesignerCommands.cpp
+++ b/Engine/Source/Editor/UMGEditor/Private/Designer/DesignerCommands.cpp
@@ -7,4 +7,7 @@ void FDesignerCommands::RegisterCommands()
 {
 	UI_COMMAND( LayoutTransform, "Layout Transform Mode", "Adjust widget layout transform", EUserInterfaceActionType::ToggleButton, FInputGesture(EKeys::W) );
 	UI_COMMAND( RenderTransform, "Render Transform Mode", "Adjust widget render transform", EUserInterfaceActionType::ToggleButton, FInputGesture(EKeys::E) );
+
+	UI_COMMAND( LocationGridSnap, "Grid Snap", "Enables or disables snapping to the grid when dragging objects around", EUserInterfaceActionType::ToggleButton, FInputGesture() );
+	UI_COMMAND( RotationGridSnap, "Rotation Snap", "Enables or disables snapping objects to a rotation grid", EUserInterfaceActionType::ToggleButton, FInputGesture() );
 }
diff --git a/Engine/Source/Editor/UMGEditor/Private/Designer/DesignerCommands.h b/Engine/Source/Editor/UMGEditor/Private/Designer/DesignerCommands.h
--- a/Engine/Source/Editor/UMGEditor/Private/Designer/DesignerCommands.h
+++ b/Engine/Source/Editor/UMGEditor/Private/Designer/DesignerCommands.h
@@ -25,6 +25,12 @@ public:
 	/** Render Mode */
 	TSharedPtr< FUICommandInfo > RenderTransform;
 
+	/** Enables or disables snapping to the grid when dragging objects around */
+	TSharedPtr< FUICommandInfo > LocationGridSnap;
+
+	/** Enables or disables snapping to a grid when rotating objects */
+	TSharedPtr< FUICommandInfo > RotationGridSnap;
+
 public:
 	/** Registers our commands with the binding system */
 	virtual void RegisterCommands() override;
diff --git a/Engine/Source/Editor/UMGEditor/Private/Designer/SDesignerToolBar.cpp b/Engine/Source/Editor/UMGEditor/Private/Designer/SDesignerToolBar.cpp
--- a/Engine/Source/Editor/UMGEditor/Private/Designer/SDesignerToolBar.cpp
+++ b/Engine/Source/Editor/UMGEditor/Private/Designer/SDesignerToolBar.cpp
@@ -5,6 +5,8 @@
 #include "SDesignerToolBar.h"
 #include "DesignerCommands.h"
 
+#include "SViewportToolBarComboMenu.h"
+
 #define LOCTEXT_NAMESPACE "UMG"
 
 void SDesignerToolBar::Construct( const FArguments& InArgs )
@@ -39,8 +41,105 @@ TSharedRef< SWidget > SDesignerToolBar::MakeToolBar(const TSharedPtr< FExtender
 	}
 	ToolbarBuilder.EndBlockGroup();
 	ToolbarBuilder.EndSection();
+
+	ToolbarBuilder.BeginSection("LocationGridSnap");
+	{
+		// Grab the existing UICommand 
+		FUICommandInfo* Command = FDesignerCommands::Get().LocationGridSnap.Get();
+
+		static FName PositionSnapName = FName(TEXT("PositionSnap"));
+
+		// Setup a GridSnapSetting with the UICommand
+		ToolbarBuilder.AddWidget(SNew(SViewportToolBarComboMenu)
+			.Style(ToolBarStyle)
+			.BlockLocation(EMultiBlockLocation::Start)
+			.Cursor(EMouseCursor::Default)
+			.IsChecked(this, &SDesignerToolBar::IsLocationGridSnapChecked)
+			.OnCheckStateChanged(this, &SDesignerToolBar::HandleToggleLocationGridSnap)
+			.Label(this, &SDesignerToolBar::GetLocationGridLabel)
+			.OnGetMenuContent(this, &SDesignerToolBar::FillLocationGridSnapMenu)
+			.ToggleButtonToolTip(Command->GetDescription())
+			.MenuButtonToolTip(LOCTEXT("LocationGridSnap_ToolTip", "Set the Position Grid Snap value"))
+			.Icon(Command->GetIcon())
+			.ParentToolBar(SharedThis(this))
+			, PositionSnapName);
+	}
+	ToolbarBuilder.EndSection();
 	
 	return ToolbarBuilder.MakeWidget();
 }
 
+ESlateCheckBoxState::Type SDesignerToolBar::IsLocationGridSnapChecked() const
+{
+	return GetDefault<UWidgetDesignerSettings>()->GridSnapEnabled ? ESlateCheckBoxState::Checked : ESlateCheckBoxState::Unchecked;
+}
+
+void SDesignerToolBar::HandleToggleLocationGridSnap(ESlateCheckBoxState::Type InState)
+{
+	UWidgetDesignerSettings* ViewportSettings = GetMutableDefault<UWidgetDesignerSettings>();
+	ViewportSettings->GridSnapEnabled = !ViewportSettings->GridSnapEnabled;
+}
+
+FText SDesignerToolBar::GetLocationGridLabel() const
+{
+	const UWidgetDesignerSettings* ViewportSettings = GetDefault<UWidgetDesignerSettings>();
+	return FText::AsNumber(ViewportSettings->GridSnapSize);
+}
+
+TSharedRef<SWidget> SDesignerToolBar::FillLocationGridSnapMenu()
+{
+	const UWidgetDesignerSettings* ViewportSettings = GetDefault<UWidgetDesignerSettings>();
+
+	TArray<int32> GridSizes;
+	GridSizes.Add(1);
+	GridSizes.Add(2);
+	GridSizes.Add(3);
+	GridSizes.Add(4);
+	GridSizes.Add(5);
+	GridSizes.Add(10);
+	GridSizes.Add(15);
+	GridSizes.Add(25);
+
+	return BuildLocationGridCheckBoxList("Snap", LOCTEXT("LocationSnapText", "Snap Sizes"), GridSizes);
+}
+
+TSharedRef<SWidget> SDesignerToolBar::BuildLocationGridCheckBoxList(FName InExtentionHook, const FText& InHeading, const TArray<int32>& InGridSizes) const
+{
+	const UWidgetDesignerSettings* ViewportSettings = GetDefault<UWidgetDesignerSettings>();
+
+	const bool bShouldCloseWindowAfterMenuSelection = true;
+	FMenuBuilder LocationGridMenuBuilder(bShouldCloseWindowAfterMenuSelection, CommandList);
+
+	LocationGridMenuBuilder.BeginSection(InExtentionHook, InHeading);
+	for ( int32 CurGridSizeIndex = 0; CurGridSizeIndex < InGridSizes.Num(); ++CurGridSizeIndex )
+	{
+		const int32 CurGridSize = InGridSizes[CurGridSizeIndex];
+
+		LocationGridMenuBuilder.AddMenuEntry(
+			FText::AsNumber(CurGridSize),
+			FText::Format(LOCTEXT("LocationGridSize_ToolTip", "Sets grid size to {0}"), FText::AsNumber(CurGridSize)),
+			FSlateIcon(),
+			FUIAction(FExecuteAction::CreateStatic(&SDesignerToolBar::SetGridSize, CurGridSize),
+			FCanExecuteAction(),
+			FIsActionChecked::CreateStatic(&SDesignerToolBar::IsGridSizeChecked, CurGridSize)),
+			NAME_None,
+			EUserInterfaceActionType::RadioButton);
+	}
+	LocationGridMenuBuilder.EndSection();
+
+	return LocationGridMenuBuilder.MakeWidget();
+}
+
+void SDesignerToolBar::SetGridSize(int32 InGridSize)
+{
+	UWidgetDesignerSettings* ViewportSettings = GetMutableDefault<UWidgetDesignerSettings>();
+	ViewportSettings->GridSnapSize = InGridSize;
+}
+
+bool SDesignerToolBar::IsGridSizeChecked(int32 InGridSnapSize)
+{
+	const UWidgetDesignerSettings* ViewportSettings = GetDefault<UWidgetDesignerSettings>();
+	return ( ViewportSettings->GridSnapSize == InGridSnapSize );
+}
+
 #undef LOCTEXT_NAMESPACE
diff --git a/Engine/Source/Editor/UMGEditor/Private/Designer/SDesignerToolBar.h b/Engine/Source/Editor/UMGEditor/Private/Designer/SDesignerToolBar.h
--- a/Engine/Source/Editor/UMGEditor/Private/Designer/SDesignerToolBar.h
+++ b/Engine/Source/Editor/UMGEditor/Private/Designer/SDesignerToolBar.h
@@ -25,6 +25,15 @@ public:
 	TSharedRef< SWidget > MakeToolBar( const TSharedPtr< FExtender > InExtenders );	
 
 private:
+	// Begin Grid Snapping
+	ESlateCheckBoxState::Type IsLocationGridSnapChecked() const;
+	void HandleToggleLocationGridSnap(ESlateCheckBoxState::Type InState);
+	FText GetLocationGridLabel() const;
+	TSharedRef<SWidget> FillLocationGridSnapMenu();
+	TSharedRef<SWidget> BuildLocationGridCheckBoxList(FName InExtentionHook, const FText& InHeading, const TArray<int32>& InGridSizes) const;
+	static void SetGridSize(int32 InGridSize);
+	static bool IsGridSizeChecked(int32 InGridSnapSize);
+	// End Grid Snapping
 
 	/** Command list */
 	TSharedPtr<FUICommandList> CommandList;
diff --git a/Engine/Source/Editor/UMGEditor/Private/Designer/SDesignerView.cpp b/Engine/Source/Editor/UMGEditor/Private/Designer/SDesignerView.cpp
--- a/Engine/Source/Editor/UMGEditor/Private/Designer/SDesignerView.cpp
+++ b/Engine/Source/Editor/UMGEditor/Private/Designer/SDesignerView.cpp
@@ -334,7 +334,10 @@ void SDesignerView::BindCommands()
 
 void SDesignerView::SetTransformMode(ETransformMode::Type InTransformMode)
 {
-	TransformMode = InTransformMode;
+	if ( !InTransaction() )
+	{
+		TransformMode = InTransformMode;
+	}
 }
 
 bool SDesignerView::CanSetTransformMode(ETransformMode::Type InTransformMode) const
@@ -466,6 +469,45 @@ void SDesignerView::OnEditorSelectionChanged()
 	CreateExtensionWidgetsForSelection();
 }
 
+FGeometry SDesignerView::GetDesignerGeometry() const
+{
+	return CachedDesignerGeometry;
+}
+
+bool SDesignerView::GetWidgetParentGeometry(const FWidgetReference& Widget, FGeometry& Geometry) const
+{
+	if ( UWidget* PreviewWidget = Widget.GetPreview() )
+	{
+		if ( UPanelWidget* Parent = PreviewWidget->GetParent() )
+		{
+			FWidgetReference ParentReference = BlueprintEditor.Pin()->GetReferenceFromPreview(Parent);
+			return GetWidgetGeometry(ParentReference, Geometry);
+		}
+	}
+
+	Geometry = GetDesignerGeometry();
+	return true;
+}
+
+bool SDesignerView::GetWidgetGeometry(const FWidgetReference& Widget, FGeometry& Geometry) const
+{
+	if ( UWidget* PreviewWidget = Widget.GetPreview() )
+	{
+		TSharedPtr<SWidget> CachedPreviewWidget = PreviewWidget->GetCachedWidget();
+		if ( CachedPreviewWidget.IsValid() )
+		{
+			const FArrangedWidget* ArrangedWidget = CachedWidgetGeometry.Find(CachedPreviewWidget.ToSharedRef());
+			if ( ArrangedWidget )
+			{
+				Geometry = ArrangedWidget->Geometry;
+				return true;
+			}
+		}
+	}
+
+	return false;
+}
+
 void SDesignerView::ClearExtensionWidgets()
 {
 	ExtensionWidgetCanvas->ClearChildren();
@@ -713,10 +755,16 @@ FReply SDesignerView::OnMouseButtonDown(const FGeometry& MyGeometry, const FPoin
 
 		if ( MouseEvent.GetEffectingButton() == EKeys::LeftMouseButton )
 		{
+			const bool bResolvePendingSelectionImmediately =
+				!SelectedWidget.IsValid() ||
+				!NewSelectedWidget.GetTemplate()->IsChildOf(SelectedWidget.GetTemplate()) ||
+				SelectedWidget.GetTemplate()->GetParent() == nullptr;
+
 			// If the newly clicked item is a child of the active selection, add it to the pending set of selected 
-			// widgets, if they begin dragging we can just move the parent, but if it's not part of the parent set, we want to immediately
-			// begin dragging it.
-			if ( !SelectedWidget.IsValid() || !NewSelectedWidget.GetTemplate()->IsChildOf(SelectedWidget.GetTemplate()) )
+			// widgets, if they begin dragging we can just move the parent, but if it's not part of the parent set, 
+			// we want to immediately begin dragging it.  Also if the currently selected widget is the root widget, 
+			// we won't be moving him so just resolve immediately.
+			if ( bResolvePendingSelectionImmediately )
 			{
 				ResolvePendingSelectedWidgets();
 			}
@@ -739,6 +787,8 @@ FReply SDesignerView::OnMouseButtonUp(const FGeometry& MyGeometry, const FPointe
 		bMouseDown = false;
 		bMovingExistingWidget = false;
 		DesignerMessage = EDesignerMessage::None;
+
+		EndTransaction(false);
 	}
 	else if ( MouseEvent.GetEffectingButton() == EKeys::RightMouseButton )
 	{
@@ -772,12 +822,37 @@ FReply SDesignerView::OnMouseMove(const FGeometry& MyGeometry, const FPointerEve
 	{
 		if ( SelectedWidget.IsValid() && !bMovingExistingWidget )
 		{
-			const bool bIsRootWidget = SelectedWidget.GetTemplate()->GetParent() == NULL;
-			if ( !bIsRootWidget )
+			if ( TransformMode == ETransformMode::Layout )
 			{
-				bMovingExistingWidget = true;
-				//Drag selected widgets
-				return FReply::Handled().DetectDrag(AsShared(), EKeys::LeftMouseButton);
+				const bool bIsRootWidget = SelectedWidget.GetTemplate()->GetParent() == NULL;
+				if ( !bIsRootWidget )
+				{
+					bMovingExistingWidget = true;
+					//Drag selected widgets
+					return FReply::Handled().DetectDrag(AsShared(), EKeys::LeftMouseButton);
+				}
+			}
+			else
+			{
+				checkSlow(TransformMode == ETransformMode::Render);
+				checkSlow(bMovingExistingWidget == false);
+
+				BeginTransaction(LOCTEXT("MoveWidgetRT", "Move Widget (Render Transform)"));
+
+				if ( UWidget* PreviewWidget = SelectedWidget.GetPreview() )
+				{
+					FGeometry ParentGeometry;
+					if ( GetWidgetParentGeometry(SelectedWidget, ParentGeometry) )
+					{
+						const FSlateRenderTransform& AbsoluteToLocalTransform = Inverse(ParentGeometry.GetAccumulatedRenderTransform());
+
+						FWidgetTransform RenderTransform = PreviewWidget->RenderTransform;
+						RenderTransform.Translation += AbsoluteToLocalTransform.TransformVector(MouseEvent.GetCursorDelta());
+
+						FObjectEditorUtils::SetPropertyValue<UWidget, FWidgetTransform>(PreviewWidget, "RenderTransform", RenderTransform);
+						FObjectEditorUtils::SetPropertyValue<UWidget, FWidgetTransform>(SelectedWidget.GetTemplate(), "RenderTransform", RenderTransform);
+					}
+				}
 			}
 		}
 	}
@@ -838,6 +913,21 @@ void SDesignerView::ShowContextMenu(const FGeometry& MyGeometry, const FPointerE
 	}
 }
 
+void SDesignerView::PopulateWidgetGeometryCache(FArrangedWidget& Root)
+{
+	FArrangedChildren ArrangedChildren(EVisibility::All);
+	Root.Widget->ArrangeChildren(Root.Geometry, ArrangedChildren);
+
+	CachedWidgetGeometry.Add(Root.Widget, Root);
+
+	// A widget's children are implicitly Z-ordered from first to last
+	for ( int32 ChildIndex = ArrangedChildren.Num() - 1; ChildIndex >= 0; --ChildIndex )
+	{
+		FArrangedWidget& SomeChild = ArrangedChildren[ChildIndex];
+		PopulateWidgetGeometryCache(SomeChild);
+	}
+}
+
 void SDesignerView::CacheSelectedWidgetGeometry()
 {
 	if ( SelectedSlateWidget.IsValid() )
@@ -977,6 +1067,7 @@ void SDesignerView::UpdatePreviewWidget()
 
 void SDesignerView::Tick(const FGeometry& AllottedGeometry, const double InCurrentTime, const float InDeltaTime)
 {
+	CachedDesignerGeometry = AllottedGeometry;
 	HoverTime += InDeltaTime;
 
 	UpdatePreviewWidget();
@@ -1004,6 +1095,11 @@ void SDesignerView::Tick(const FGeometry& AllottedGeometry, const double InCurre
 		}
 	}
 
+	// Perform an arrange children pass to cache the geometry of all widgets so that we can query it later.
+	CachedWidgetGeometry.Reset();
+	FArrangedWidget WindowWidgetGeometry(PreviewHitTestRoot.ToSharedRef(), AllottedGeometry);
+	PopulateWidgetGeometryCache(WindowWidgetGeometry);
+
 	CacheSelectedWidgetGeometry();
 
 	// Tick all designer extensions in case they need to update widgets
@@ -1114,6 +1210,8 @@ UWidget* SDesignerView::ProcessDropAndAddWidget(const FGeometry& MyGeometry, con
 	FArrangedWidget ArrangedWidget(SNullWidget::NullWidget, FGeometry());
 	FWidgetReference WidgetUnderCursor = GetWidgetAtCursor(MyGeometry, DragDropEvent, ArrangedWidget);
 
+	FGeometry WidgetUnderCursorGeometry = ArrangedWidget.Geometry;
+
 	UWidget* Target = NULL;
 	if ( WidgetUnderCursor.IsValid() )
 	{
@@ -1177,7 +1275,7 @@ UWidget* SDesignerView::ProcessDropAndAddWidget(const FGeometry& MyGeometry, con
 			Widget->IsDesignTime(true);
 
 			// Determine local position inside the parent widget and add the widget to the slot.
-			FVector2D LocalPosition = ArrangedWidget.Geometry.AbsoluteToLocal(DragDropEvent.GetScreenSpacePosition());
+			FVector2D LocalPosition = WidgetUnderCursorGeometry.AbsoluteToLocal(DragDropEvent.GetScreenSpacePosition());
 			if ( UPanelSlot* Slot = Parent->AddChild(Widget) )
 			{
 				// HACK UMG - This seems like a bad idea to call TakeWidget
@@ -1188,6 +1286,13 @@ UWidget* SDesignerView::ProcessDropAndAddWidget(const FGeometry& MyGeometry, con
 				static const FVector2D MinimumDefaultSize(100, 40);
 				FVector2D LocalSize = FVector2D(FMath::Max(WidgetDesiredSize.X, MinimumDefaultSize.X), FMath::Max(WidgetDesiredSize.Y, MinimumDefaultSize.Y));
 
+				const UWidgetDesignerSettings* DesignerSettings = GetDefault<UWidgetDesignerSettings>();
+				if ( DesignerSettings->GridSnapEnabled )
+				{
+					LocalPosition.X = ( (int32)LocalPosition.X ) - (( (int32)LocalPosition.X ) % DesignerSettings->GridSnapSize);
+					LocalPosition.Y = ( (int32)LocalPosition.Y ) - (( (int32)LocalPosition.Y ) % DesignerSettings->GridSnapSize);
+				}
+
 				Slot->SetDesiredPosition(LocalPosition);
 				Slot->SetDesiredSize(LocalSize);
 
@@ -1238,17 +1343,10 @@ UWidget* SDesignerView::ProcessDropAndAddWidget(const FGeometry& MyGeometry, con
 		{
 			DesignerMessage = EDesignerMessage::MoveFromParent;
 
-			if ( UPanelWidget* LockedPreviewParent = Cast<UPanelWidget>(SelectedDragDropOp->ParentWidget.GetPreview()) )
+			WidgetUnderCursorGeometry = GetDesignerGeometry();
+			if ( GetWidgetGeometry(SelectedDragDropOp->ParentWidget, WidgetUnderCursorGeometry) )
 			{
-				TSharedPtr<SWidget> NewParentWidget = LockedPreviewParent->GetCachedWidget();
-
-				if ( NewParentWidget.IsValid() )
-				{
-					Target = bIsPreview ? SelectedDragDropOp->ParentWidget.GetPreview() : SelectedDragDropOp->ParentWidget.GetTemplate();
-
-					const bool bSuccess = FDesignTimeUtils::GetArrangedWidget(NewParentWidget.ToSharedRef(), ArrangedWidget);
-					check(bSuccess);
-				}
+				Target = bIsPreview ? SelectedDragDropOp->ParentWidget.GetPreview() : SelectedDragDropOp->ParentWidget.GetTemplate();
 			}
 		}
 
@@ -1281,9 +1379,18 @@ UWidget* SDesignerView::ProcessDropAndAddWidget(const FGeometry& MyGeometry, con
 				Widget->GetParent()->RemoveChild(Widget);
 			}
 
-			FVector2D LocalPosition = ArrangedWidget.Geometry.AbsoluteToLocal(DragDropEvent.GetScreenSpacePosition());
+			FVector2D LocalPosition = WidgetUnderCursorGeometry.AbsoluteToLocal(DragDropEvent.GetScreenSpacePosition());
 			if ( UPanelSlot* Slot = NewParent->AddChild(Widget) )
 			{
+				FVector2D NewPosition = LocalPosition - SelectedWidgetContextMenuLocation;
+
+				const UWidgetDesignerSettings* DesignerSettings = GetDefault<UWidgetDesignerSettings>();
+				if ( DesignerSettings->GridSnapEnabled )
+				{
+					NewPosition.X = ( (int32)NewPosition.X ) - ( ( (int32)NewPosition.X ) % DesignerSettings->GridSnapSize );
+					NewPosition.Y = ( (int32)NewPosition.Y ) - ( ( (int32)NewPosition.Y ) % DesignerSettings->GridSnapSize );
+				}
+
 				// HACK UMG: In order to correctly drop items into the canvas that have a non-zero anchor,
 				// we need to know the layout information after slate has performed a prepass.  So we have
 				// to rebase the layout and reinterpret the new position based on anchor point layout data.
@@ -1296,7 +1403,7 @@ UWidget* SDesignerView::ProcessDropAndAddWidget(const FGeometry& MyGeometry, con
 						FWidgetBlueprintEditorUtils::ImportPropertiesFromText(Slot, SelectedDragDropOp->ExportedSlotProperties);
 
 						CanvasSlot->SaveBaseLayout();
-						Slot->SetDesiredPosition(LocalPosition - SelectedWidgetContextMenuLocation);
+						Slot->SetDesiredPosition(NewPosition);
 						CanvasSlot->RebaseLayout();
 
 						FWidgetBlueprintEditorUtils::ExportPropertiesToText(Slot, SelectedDragDropOp->ExportedSlotProperties);
@@ -1304,13 +1411,12 @@ UWidget* SDesignerView::ProcessDropAndAddWidget(const FGeometry& MyGeometry, con
 					else
 					{
 						FWidgetBlueprintEditorUtils::ImportPropertiesFromText(Slot, SelectedDragDropOp->ExportedSlotProperties);
-						//Slot->SetDesiredPosition(LocalPosition - SelectedWidgetContextMenuLocation);
 					}
 				}
 				else
 				{
 					FWidgetBlueprintEditorUtils::ImportPropertiesFromText(Slot, SelectedDragDropOp->ExportedSlotProperties);
-					Slot->SetDesiredPosition(LocalPosition - SelectedWidgetContextMenuLocation);
+					Slot->SetDesiredPosition(NewPosition);
 				}
 
 				DropPreviewParent = NewParent;
@@ -1454,24 +1560,34 @@ TSharedRef<SWidget> SDesignerView::GetAspectMenu()
 
 void SDesignerView::BeginTransaction(const FText& SessionName)
 {
-	if ( ScopedTransaction == NULL )
+	if ( ScopedTransaction == nullptr )
 	{
 		ScopedTransaction = new FScopedTransaction(SessionName);
-	}
 
-	if ( SelectedWidget.IsValid() )
-	{
-		SelectedWidget.GetPreview()->Modify();
-		SelectedWidget.GetTemplate()->Modify();
+		if ( SelectedWidget.IsValid() )
+		{
+			SelectedWidget.GetPreview()->Modify();
+			SelectedWidget.GetTemplate()->Modify();
+		}
 	}
 }
 
-void SDesignerView::EndTransaction()
+bool SDesignerView::InTransaction() const
+{
+	return ScopedTransaction != nullptr;
+}
+
+void SDesignerView::EndTransaction(bool bCancel)
 {
-	if ( ScopedTransaction != NULL )
+	if ( ScopedTransaction != nullptr )
 	{
+		if ( bCancel )
+		{
+			ScopedTransaction->Cancel();
+		}
+
 		delete ScopedTransaction;
-		ScopedTransaction = NULL;
+		ScopedTransaction = nullptr;
 	}
 }
 
diff --git a/Engine/Source/Editor/UMGEditor/Private/Designer/SDesignerView.h b/Engine/Source/Editor/UMGEditor/Private/Designer/SDesignerView.h
--- a/Engine/Source/Editor/UMGEditor/Private/Designer/SDesignerView.h
+++ b/Engine/Source/Editor/UMGEditor/Private/Designer/SDesignerView.h
@@ -59,6 +59,9 @@ public:
 	virtual float GetPreviewScale() const override;
 	virtual FWidgetReference GetSelectedWidget() const override;
 	virtual ETransformMode::Type GetTransformMode() const override;
+	virtual FGeometry GetDesignerGeometry() const override;
+	virtual bool GetWidgetGeometry(const FWidgetReference& Widget, FGeometry& Geometry) const override;
+	virtual bool GetWidgetParentGeometry(const FWidgetReference& Widget, FGeometry& Geometry) const override;
 	// End of IUMGDesigner interface
 
 private:
@@ -107,6 +110,8 @@ private:
 
 	void OnObjectPropertyChanged(UObject* ObjectBeingModified, FPropertyChangedEvent& PropertyChangedEvent);
 
+	void PopulateWidgetGeometryCache(FArrangedWidget& Root);
+
 	void CacheSelectedWidgetGeometry();
 
 	/** @return Formatted text for the given resolution params */
@@ -122,7 +127,8 @@ private:
 	TSharedRef<SWidget> GetAspectMenu();
 
 	void BeginTransaction(const FText& SessionName);
-	void EndTransaction();
+	bool InTransaction() const;
+	void EndTransaction(bool bCancel);
 
 private:
 	FReply HandleZoomToFitClicked();
@@ -232,4 +238,10 @@ private:
 
 	/** */
 	ETransformMode::Type TransformMode;
+
+	/**  */
+	TMap<TSharedRef<SWidget>, FArrangedWidget> CachedWidgetGeometry;
+
+	/**  */
+	FGeometry CachedDesignerGeometry;
 };
diff --git a/Engine/Source/Editor/UMGEditor/Private/Designer/STransformHandle.cpp b/Engine/Source/Editor/UMGEditor/Private/Designer/STransformHandle.cpp
--- a/Engine/Source/Editor/UMGEditor/Private/Designer/STransformHandle.cpp
+++ b/Engine/Source/Editor/UMGEditor/Private/Designer/STransformHandle.cpp
@@ -93,6 +93,14 @@ FReply STransformHandle::OnMouseMove(const FGeometry& MyGeometry, const FPointer
 		ETransformMode::Type TransformMode = Designer->GetTransformMode();
 		if ( TransformMode == ETransformMode::Layout )
 		{
+			//FGeometry MyGeometry = Designer->GetWidgetGeometry(SelectedWidget);
+			//FGeometry ParentGeometry = Designer->GetWidgetParentGeometry(SelectedWidget);
+
+			//FVector2D LocalPosition = MyGeometry.Position;
+			//FVector2D LocalDragPosition = ParentGeometry.AbsoluteToLocal(MouseEvent.GetScreenSpacePosition());
+
+			//FVector2D NewSize = LocalDragPosition - LocalPosition;
+
 			FVector2D TranslateAmount = MouseEvent.GetCursorDelta() * ( 1.0f / Designer->GetPreviewScale() );
 
 			if ( Preview->Slot )
@@ -111,8 +119,17 @@ FReply STransformHandle::OnMouseMove(const FGeometry& MyGeometry, const FPointer
 
 			if ( Action == ETransformAction::Primary )
 			{
-				FVector2D ScaleDelta = MouseEvent.GetCursorDelta() * 0.05f;
-				RenderTransform.Scale += ScaleDelta;
+				FVector2D TranslateAmount = MouseEvent.GetCursorDelta() * ( 1.0f / Designer->GetPreviewScale() );
+
+				if ( Preview->Slot )
+				{
+					Preview->Slot->Resize(DragDirection, TranslateAmount);
+				}
+
+				if ( Template->Slot )
+				{
+					Template->Slot->Resize(DragDirection, TranslateAmount);
+				}
 			}
 			else
 			{
@@ -121,6 +138,7 @@ FReply STransformHandle::OnMouseMove(const FGeometry& MyGeometry, const FPointer
 			}
 
 			FObjectEditorUtils::SetPropertyValue<UWidget, FWidgetTransform>(Preview, "RenderTransform", RenderTransform);
+			FObjectEditorUtils::SetPropertyValue<UWidget, FWidgetTransform>(Template, "RenderTransform", RenderTransform);
 		}
 	}
 
diff --git a/Engine/Source/Editor/UMGEditor/Private/UMGEditorPrivatePCH.h b/Engine/Source/Editor/UMGEditor/Private/UMGEditorPrivatePCH.h
--- a/Engine/Source/Editor/UMGEditor/Private/UMGEditorPrivatePCH.h
+++ b/Engine/Source/Editor/UMGEditor/Private/UMGEditorPrivatePCH.h
@@ -19,6 +19,8 @@
 
 #include "UMG.h"
 
+#include "WidgetDesignerSettings.h"
+
 #include "UMGEditorClasses.h"
 #include "WidgetBlueprint.h"
 
diff --git a/Engine/Source/Editor/UMGEditor/Public/IUMGDesigner.h b/Engine/Source/Editor/UMGEditor/Public/IUMGDesigner.h
--- a/Engine/Source/Editor/UMGEditor/Public/IUMGDesigner.h
+++ b/Engine/Source/Editor/UMGEditor/Public/IUMGDesigner.h
@@ -31,4 +31,13 @@ public:
 
 	/** @return Get the transform mode currently in use in the designer. */
 	virtual ETransformMode::Type GetTransformMode() const = 0;
+
+	/** @return The Geometry representing the designer area, useful for when you need to convert mouse into designer space. */
+	virtual FGeometry GetDesignerGeometry() const = 0;
+
+	/**  */
+	virtual bool GetWidgetGeometry(const FWidgetReference& Widget, FGeometry& Geometry) const = 0;
+
+	/**  */
+	virtual bool GetWidgetParentGeometry(const FWidgetReference& Widget, FGeometry& Geometry) const = 0;
 };
diff --git a/Engine/Source/Editor/UMGEditor/Public/Settings/WidgetDesignerSettings.h b/Engine/Source/Editor/UMGEditor/Public/Settings/WidgetDesignerSettings.h
new file mode 100644
--- /dev/null
+++ b/Engine/Source/Editor/UMGEditor/Public/Settings/WidgetDesignerSettings.h
@@ -0,0 +1,22 @@
+// Copyright 1998-2014 Epic Games, Inc. All Rights Reserved.
+
+#pragma once
+
+#include "WidgetDesignerSettings.generated.h"
+
+/**
+ * Implements the settings for the Widget Blueprint Designer.
+ */
+UCLASS(config=EditorUserSettings)
+class UMGEDITOR_API UWidgetDesignerSettings
+	: public UObject
+{
+	GENERATED_UCLASS_BODY()
+
+	/** If enabled, actor positions will snap to the grid. */
+	UPROPERTY(EditAnywhere, config, Category=GridSnapping, meta=(DisplayName = "Enable Grid Snapping"))
+	uint32 GridSnapEnabled:1;
+
+	UPROPERTY(config)
+	int32 GridSnapSize;
+};
diff --git a/Engine/Source/Editor/UnrealEd/Private/STransformViewportToolbar.cpp b/Engine/Source/Editor/UnrealEd/Private/STransformViewportToolbar.cpp
--- a/Engine/Source/Editor/UnrealEd/Private/STransformViewportToolbar.cpp
+++ b/Engine/Source/Editor/UnrealEd/Private/STransformViewportToolbar.cpp
@@ -5,6 +5,8 @@
 #include "STransformViewportToolbar.h"
 #include "SLevelViewport.h"
 #include "EditorViewportCommands.h"
+#include "SViewportToolBarIconMenu.h"
+#include "SViewportToolBarComboMenu.h"
 
 #define LOCTEXT_NAMESPACE "TransformToolBar"
 
@@ -17,360 +19,6 @@ namespace TransformViewportToolbarDefs
 	const float ToggleImageScale = 16.0f;
 }
 
-/**
- * Custom widget to display a toggle/drop down menu. 
- *	Displayed as shown:
- * +--------+-------------+
- * | Toggle | Menu button |
- * +--------+-------------+
- */
-class SGridSnapSetting : public SCompoundWidget 
-{
-public:
-	SLATE_BEGIN_ARGS( SGridSnapSetting ) : _BlockLocation(EMultiBlockLocation::Start) {}
-	
-		/** We need to know about the toolbar we are in */
-		SLATE_ARGUMENT( TSharedPtr<class SViewportToolBar>, ParentToolBar );
-
-		/** Content for the drop down menu  */
-		SLATE_EVENT( FOnGetContent, OnGetMenuContent )
-
-		/** Called upon state change with the value of the next state */
-		SLATE_EVENT( FOnCheckStateChanged, OnCheckStateChanged )
-			
-		/** Sets the current checked state of the checkbox */
-		SLATE_ATTRIBUTE( ESlateCheckBoxState::Type, IsChecked )
-
-		/** Icon shown on the toggle button */
-		SLATE_ATTRIBUTE( FSlateIcon, Icon )
-
-		/** Label shown on the menu button */
-		SLATE_ATTRIBUTE( FText, Label )
-
-		/** Overall style */
-		SLATE_ATTRIBUTE( FName, Style )
-
-		/** ToolTip shown on the menu button */
-		SLATE_ATTRIBUTE( FText, MenuButtonToolTip )
-
-		/** ToolTip shown on the toggle button */
-		SLATE_ATTRIBUTE( FText, ToggleButtonToolTip )
-
-		/** The button location */
-		SLATE_ARGUMENT( EMultiBlockLocation::Type, BlockLocation )
-
-	SLATE_END_ARGS( )
-
-	/**
-	 * Construct this widget
-	 *
-	 * @param	InArgs	The declaration data for this widget
-	 */
-	void Construct( const FArguments& InArgs )
-	{
-		FName ButtonStyle = FEditorStyle::Join(InArgs._Style.Get(), ".Button");
-		FName CheckboxStyle = FEditorStyle::Join(InArgs._Style.Get(), ".ToggleButton");
-	
-		const FSlateIcon& Icon = InArgs._Icon.Get();
-
-		ParentToolBar = InArgs._ParentToolBar;
-	
-		TSharedPtr<SCheckBox> ToggleControl;
-		{
-			ToggleControl = SNew(SCheckBox)
-			.Cursor( EMouseCursor::Default )
-			.Padding(FMargin( 4.0f ))
-			.Style(FEditorStyle::Get(), EMultiBlockLocation::ToName(CheckboxStyle, InArgs._BlockLocation))
-			.OnCheckStateChanged(InArgs._OnCheckStateChanged)
-			.ToolTipText(InArgs._ToggleButtonToolTip)
-			.IsChecked(InArgs._IsChecked)
-			.Content()
-			[
-				SNew( SBox )
-				.WidthOverride( TransformViewportToolbarDefs::ToggleImageScale )
-				.HeightOverride( TransformViewportToolbarDefs::ToggleImageScale )
-				.HAlign(HAlign_Center)
-				.VAlign(VAlign_Center)
-				[
-					SNew(SImage)
-					.Image(Icon.GetIcon())
-				]
-			];
-		}
-
-		{
-			SAssignNew( MenuAnchor, SMenuAnchor )
-			.Placement( MenuPlacement_BelowAnchor )
-			[
-				SNew(SButton)
-				.ButtonStyle( FEditorStyle::Get(), EMultiBlockLocation::ToName(ButtonStyle,EMultiBlockLocation::End) )
-				.ContentPadding( FMargin( 5.0f, 0.0f ) )
-				.ToolTipText(InArgs._MenuButtonToolTip)
-				.OnClicked(this, &SGridSnapSetting::OnMenuClicked)
-				[
-					SNew(SVerticalBox)
-					+SVerticalBox::Slot()
-					.AutoHeight()
-					.HAlign(HAlign_Center)
-					.VAlign(VAlign_Top)
-					[
-						SNew(STextBlock)
-						.TextStyle( FEditorStyle::Get(), FEditorStyle::Join( InArgs._Style.Get(), ".Label" ) )
-						.Text(InArgs._Label)
-					]
-					+SVerticalBox::Slot()
-					.AutoHeight()
-					.VAlign(VAlign_Bottom)
-					[
-						SNew(SHorizontalBox)
-						+SHorizontalBox::Slot()
-						.FillWidth(1.0f)
-
-						+SHorizontalBox::Slot()
-						.AutoWidth()
-						[
-							SNew( SBox )
-							.WidthOverride( TransformViewportToolbarDefs::DownArrowSize )
-							.HeightOverride( TransformViewportToolbarDefs::DownArrowSize )
-							[
-								SNew(SImage)
-								.Image(FEditorStyle::GetBrush("ComboButton.Arrow"))
-								.ColorAndOpacity(FLinearColor::Black)
-							]
-						]
-
-						+SHorizontalBox::Slot()
-							.FillWidth(1.0f)
-					]
-				]
-			]
-			.OnGetMenuContent( InArgs._OnGetMenuContent );
-		}
-
-
-		ChildSlot
-		[
-			SNew(SHorizontalBox)
-
-			//Checkbox concept
-			+SHorizontalBox::Slot()
-			.AutoWidth()
-			[
-				ToggleControl->AsShared()
-			]
-
-			// Black Separator line
-			+SHorizontalBox::Slot()
-			.AutoWidth()
-			[
-				SNew(SBorder)
-				.Padding(FMargin(1.0f, 0.0f, 0.0f, 0.0f))
-				.BorderImage(FEditorStyle::GetDefaultBrush())
-				.BorderBackgroundColor(FLinearColor::Black)
-			]
-
-			// Menu dropdown concept
-			+SHorizontalBox::Slot()
-			.AutoWidth()
-			[
-				MenuAnchor->AsShared()
-			]
-		];
-	}
-	
-private:
-	/**
-	 * Called when the menu button is clicked.  Will toggle the visibility of the menu content                   
-	 */
-	FReply OnMenuClicked()
-	{
-		// If the menu button is clicked toggle the state of the menu anchor which will open or close the menu
-		MenuAnchor->SetIsOpen( !MenuAnchor->IsOpen() );
-		ParentToolBar.Pin()->SetOpenMenu( MenuAnchor );
-		return FReply::Handled();
-	}
-
-	/**
-	 * Called when the mouse enters a menu button.  If there was a menu previously opened
-	 * we open this menu automatically
-	 */
-	void OnMouseEnter( const FGeometry& MyGeometry, const FPointerEvent& MouseEvent )
-	{
-		// See if there is another menu on the same tool bar already open
-		TWeakPtr<SMenuAnchor> OpenedMenu = ParentToolBar.Pin()->GetOpenMenu();
-		if( OpenedMenu.IsValid() && OpenedMenu.Pin()->IsOpen() && OpenedMenu.Pin() != MenuAnchor )
-		{
-			// There is another menu open so we open this menu and close the other
-			ParentToolBar.Pin()->SetOpenMenu( MenuAnchor ); 
-			MenuAnchor->SetIsOpen( true );
-		}
-	}
-
-private:
-	/** Our menus anchor */
-	TSharedPtr<SMenuAnchor> MenuAnchor;
-
-	/** Parent tool bar for querying other open menus */
-	TWeakPtr<class SViewportToolBar> ParentToolBar;
-};
-
-/**
- * Custom widget to display an icon/drop down menu. 
- *	Displayed as shown:
- * +--------+-------------+
- * | Icon   | Menu button |
- * +--------+-------------+
- */
-class SCameraSpeedSetting : public SCompoundWidget 
-{
-public:
-	SLATE_BEGIN_ARGS( SCameraSpeedSetting ){}
-	
-		/** We need to know about the toolbar we are in */
-		SLATE_ARGUMENT( TSharedPtr<class SViewportToolBar>, ParentToolBar );
-
-		/** Content for the drop down menu  */
-		SLATE_EVENT( FOnGetContent, OnGetMenuContent )
-
-		/** Icon shown */
-		SLATE_ATTRIBUTE( FSlateIcon, Icon )
-
-		/** Label shown on the menu button */
-		SLATE_ATTRIBUTE( FText, Label )
-
-		/** Overall style */
-		SLATE_ATTRIBUTE( FName, Style )
-
-	SLATE_END_ARGS( )
-
-	/**
-	 * Construct this widget
-	 *
-	 * @param	InArgs	The declaration data for this widget
-	 */
-	void Construct( const FArguments& InArgs )
-	{
-		ParentToolBar = InArgs._ParentToolBar;
-
-		const FName ButtonStyle = FEditorStyle::Join(InArgs._Style.Get(), ".Button");
-
-		const FSlateIcon& Icon = InArgs._Icon.Get();
-
-		SAssignNew( MenuAnchor, SMenuAnchor )
-		.Placement( MenuPlacement_BelowAnchor )
-		.OnGetMenuContent( InArgs._OnGetMenuContent )
-		[
-			SNew(SButton)
-			.ButtonStyle( FEditorStyle::Get(), ButtonStyle )
-			.ContentPadding( FMargin( 5.0f, 0.0f ) )
-			.OnClicked(this, &SCameraSpeedSetting::OnMenuClicked)
-			[
-				SNew(SHorizontalBox)
-
-				// Icon
-				+SHorizontalBox::Slot()
-				.AutoWidth()
-				[
-					SNew(SBox)
-					.WidthOverride( TransformViewportToolbarDefs::ToggleImageScale )
-					.HeightOverride( TransformViewportToolbarDefs::ToggleImageScale )
-					.HAlign(HAlign_Center)
-					.VAlign(VAlign_Center)
-					[
-						SNew(SImage)
-						.Image(Icon.GetIcon())
-					]
-				]
-
-				// Spacer
-				+SHorizontalBox::Slot()
-				.AutoWidth()
-				.Padding( FMargin( 4.0f, 0.0f ) )
-
-				// Menu dropdown
-				+SHorizontalBox::Slot()
-				.AutoWidth()
-				[
-					SNew(SVerticalBox)
-					+SVerticalBox::Slot()
-					.AutoHeight()
-					.HAlign(HAlign_Center)
-					.VAlign(VAlign_Top)
-					[
-						SNew(STextBlock)
-						.TextStyle( FEditorStyle::Get(), FEditorStyle::Join( InArgs._Style.Get(), ".Label" ) )
-						.Text(InArgs._Label)
-					]
-					+SVerticalBox::Slot()
-					.AutoHeight()
-					.VAlign(VAlign_Bottom)
-					[
-						SNew(SHorizontalBox)
-						+SHorizontalBox::Slot()
-						.FillWidth(1.0f)
-
-						+SHorizontalBox::Slot()
-						.AutoWidth()
-						[
-							SNew( SBox )
-							.WidthOverride( TransformViewportToolbarDefs::DownArrowSize )
-							.HeightOverride( TransformViewportToolbarDefs::DownArrowSize )
-							[
-								SNew(SImage)
-								.Image(FEditorStyle::GetBrush("ComboButton.Arrow"))
-								.ColorAndOpacity(FLinearColor::Black)
-							]
-						]
-
-						+SHorizontalBox::Slot()
-						.FillWidth(1.0f)
-					]
-				]
-			]
-		];
-
-		ChildSlot
-		[
-			MenuAnchor.ToSharedRef()
-		];
-	}
-	
-private:
-	/**
-	 * Called when the menu button is clicked.  Will toggle the visibility of the menu content                   
-	 */
-	FReply OnMenuClicked()
-	{
-		// If the menu button is clicked toggle the state of the menu anchor which will open or close the menu
-		MenuAnchor->SetIsOpen( !MenuAnchor->IsOpen() );
-		ParentToolBar.Pin()->SetOpenMenu( MenuAnchor );
-		return FReply::Handled();
-	}
-
-	/**
-	 * Called when the mouse enters a menu button.  If there was a menu previously opened
-	 * we open this menu automatically
-	 */
-	void OnMouseEnter( const FGeometry& MyGeometry, const FPointerEvent& MouseEvent )
-	{
-		// See if there is another menu on the same tool bar already open
-		TWeakPtr<SMenuAnchor> OpenedMenu = ParentToolBar.Pin()->GetOpenMenu();
-		if( OpenedMenu.IsValid() && OpenedMenu.Pin()->IsOpen() && OpenedMenu.Pin() != MenuAnchor )
-		{
-			// There is another menu open so we open this menu and close the other
-			ParentToolBar.Pin()->SetOpenMenu( MenuAnchor ); 
-			MenuAnchor->SetIsOpen( true );
-		}
-	}
-
-private:
-	/** Our menus anchor */
-	TSharedPtr<SMenuAnchor> MenuAnchor;
-
-	/** Parent tool bar for querying other open menus */
-	TWeakPtr<class SViewportToolBar> ParentToolBar;
-};
-
 void STransformViewportToolBar::Construct( const FArguments& InArgs )
 {
 	Viewport = InArgs._Viewport;
@@ -585,7 +233,7 @@ TSharedRef< SWidget > STransformViewportToolBar::MakeTransformToolBar( const TSh
 		static FName PositionSnapName = FName(TEXT("PositionSnap"));
 
 		// Setup a GridSnapSetting with the UICommand
-		ToolbarBuilder.AddWidget(	SNew(SGridSnapSetting)
+		ToolbarBuilder.AddWidget(	SNew(SViewportToolBarComboMenu)
 									.Style(ToolBarStyle)
 									.BlockLocation(EMultiBlockLocation::Middle)
 									.Cursor( EMouseCursor::Default )
@@ -609,7 +257,7 @@ TSharedRef< SWidget > STransformViewportToolBar::MakeTransformToolBar( const TSh
 		static FName RotationSnapName = FName(TEXT("RotationSnap"));
 
 		// Setup a GridSnapSetting with the UICommand
-		ToolbarBuilder.AddWidget(	SNew(SGridSnapSetting)
+		ToolbarBuilder.AddWidget(	SNew(SViewportToolBarComboMenu)
 									.Cursor( EMouseCursor::Default )
 									.Style(ToolBarStyle)
 									.IsChecked(this, &STransformViewportToolBar::IsRotationGridSnapChecked)
@@ -632,7 +280,7 @@ TSharedRef< SWidget > STransformViewportToolBar::MakeTransformToolBar( const TSh
 		static FName ScaleSnapName = FName(TEXT("ScaleSnap"));
 
 		// Setup a GridSnapSetting with the UICommand
-		ToolbarBuilder.AddWidget(	SNew(SGridSnapSetting)
+		ToolbarBuilder.AddWidget(	SNew(SViewportToolBarComboMenu)
 									.Cursor( EMouseCursor::Default )
 									.Style(ToolBarStyle)
 									.IsChecked(this,&STransformViewportToolBar::IsScaleGridSnapChecked)
@@ -651,7 +299,7 @@ TSharedRef< SWidget > STransformViewportToolBar::MakeTransformToolBar( const TSh
 	ToolbarBuilder.BeginBlockGroup();
 	{
 		// Camera speed 
-		ToolbarBuilder.AddWidget(	SNew(SCameraSpeedSetting)
+		ToolbarBuilder.AddWidget(	SNew(SViewportToolBarIconMenu)
 									.Cursor(EMouseCursor::Default)
 									.Style(ToolBarStyle)
 									.Label(this, &STransformViewportToolBar::GetCameraSpeedLabel)
diff --git a/Engine/Source/Editor/UnrealEd/Private/SViewportToolBarComboMenu.cpp b/Engine/Source/Editor/UnrealEd/Private/SViewportToolBarComboMenu.cpp
new file mode 100644
--- /dev/null
+++ b/Engine/Source/Editor/UnrealEd/Private/SViewportToolBarComboMenu.cpp
@@ -0,0 +1,138 @@
+// Copyright 1998-2014 Epic Games, Inc. All Rights Reserved.
+
+#pragma once
+
+#include "UnrealEd.h"
+#include "SViewportToolBarIconMenu.h"
+
+void SViewportToolBarComboMenu::Construct( const FArguments& InArgs )
+{
+	FName ButtonStyle = FEditorStyle::Join(InArgs._Style.Get(), ".Button");
+	FName CheckboxStyle = FEditorStyle::Join(InArgs._Style.Get(), ".ToggleButton");
+
+	const FSlateIcon& Icon = InArgs._Icon.Get();
+
+	ParentToolBar = InArgs._ParentToolBar;
+
+	TSharedPtr<SCheckBox> ToggleControl;
+	{
+		ToggleControl = SNew(SCheckBox)
+		.Cursor( EMouseCursor::Default )
+		.Padding(FMargin( 4.0f ))
+		.Style(FEditorStyle::Get(), EMultiBlockLocation::ToName(CheckboxStyle, InArgs._BlockLocation))
+		.OnCheckStateChanged(InArgs._OnCheckStateChanged)
+		.ToolTipText(InArgs._ToggleButtonToolTip)
+		.IsChecked(InArgs._IsChecked)
+		.Content()
+		[
+			SNew( SBox )
+			.WidthOverride( 16 )
+			.HeightOverride( 16 )
+			.HAlign(HAlign_Center)
+			.VAlign(VAlign_Center)
+			[
+				SNew(SImage)
+				.Image(Icon.GetIcon())
+			]
+		];
+	}
+
+	{
+		SAssignNew( MenuAnchor, SMenuAnchor )
+		.Placement( MenuPlacement_BelowAnchor )
+		[
+			SNew(SButton)
+			.ButtonStyle( FEditorStyle::Get(), EMultiBlockLocation::ToName(ButtonStyle,EMultiBlockLocation::End) )
+			.ContentPadding( FMargin( 5.0f, 0.0f ) )
+			.ToolTipText(InArgs._MenuButtonToolTip)
+			.OnClicked(this, &SViewportToolBarComboMenu::OnMenuClicked)
+			[
+				SNew(SVerticalBox)
+				+SVerticalBox::Slot()
+				.AutoHeight()
+				.HAlign(HAlign_Center)
+				.VAlign(VAlign_Top)
+				[
+					SNew(STextBlock)
+					.TextStyle( FEditorStyle::Get(), FEditorStyle::Join( InArgs._Style.Get(), ".Label" ) )
+					.Text(InArgs._Label)
+				]
+				+SVerticalBox::Slot()
+				.AutoHeight()
+				.VAlign(VAlign_Bottom)
+				[
+					SNew(SHorizontalBox)
+					+SHorizontalBox::Slot()
+					.FillWidth(1.0f)
+
+					+SHorizontalBox::Slot()
+					.AutoWidth()
+					[
+						SNew( SBox )
+						.WidthOverride( 4 )
+						.HeightOverride( 4 )
+						[
+							SNew(SImage)
+							.Image(FEditorStyle::GetBrush("ComboButton.Arrow"))
+							.ColorAndOpacity(FLinearColor::Black)
+						]
+					]
+
+					+SHorizontalBox::Slot()
+						.FillWidth(1.0f)
+				]
+			]
+		]
+		.OnGetMenuContent( InArgs._OnGetMenuContent );
+	}
+
+
+	ChildSlot
+	[
+		SNew(SHorizontalBox)
+
+		//Checkbox concept
+		+SHorizontalBox::Slot()
+		.AutoWidth()
+		[
+			ToggleControl->AsShared()
+		]
+
+		// Black Separator line
+		+SHorizontalBox::Slot()
+		.AutoWidth()
+		[
+			SNew(SBorder)
+			.Padding(FMargin(1.0f, 0.0f, 0.0f, 0.0f))
+			.BorderImage(FEditorStyle::GetDefaultBrush())
+			.BorderBackgroundColor(FLinearColor::Black)
+		]
+
+		// Menu dropdown concept
+		+SHorizontalBox::Slot()
+		.AutoWidth()
+		[
+			MenuAnchor->AsShared()
+		]
+	];
+}
+
+FReply SViewportToolBarComboMenu::OnMenuClicked()
+{
+	// If the menu button is clicked toggle the state of the menu anchor which will open or close the menu
+	MenuAnchor->SetIsOpen( !MenuAnchor->IsOpen() );
+	ParentToolBar.Pin()->SetOpenMenu( MenuAnchor );
+	return FReply::Handled();
+}
+
+void SViewportToolBarComboMenu::OnMouseEnter( const FGeometry& MyGeometry, const FPointerEvent& MouseEvent )
+{
+	// See if there is another menu on the same tool bar already open
+	TWeakPtr<SMenuAnchor> OpenedMenu = ParentToolBar.Pin()->GetOpenMenu();
+	if( OpenedMenu.IsValid() && OpenedMenu.Pin()->IsOpen() && OpenedMenu.Pin() != MenuAnchor )
+	{
+		// There is another menu open so we open this menu and close the other
+		ParentToolBar.Pin()->SetOpenMenu( MenuAnchor ); 
+		MenuAnchor->SetIsOpen( true );
+	}
+}
diff --git a/Engine/Source/Editor/UnrealEd/Private/SViewportToolBarIconMenu.cpp b/Engine/Source/Editor/UnrealEd/Private/SViewportToolBarIconMenu.cpp
new file mode 100644
--- /dev/null
+++ b/Engine/Source/Editor/UnrealEd/Private/SViewportToolBarIconMenu.cpp
@@ -0,0 +1,111 @@
+// Copyright 1998-2014 Epic Games, Inc. All Rights Reserved.
+
+#include "UnrealEd.h"
+#include "SViewportToolBarIconMenu.h"
+
+void SViewportToolBarIconMenu::Construct( const FArguments& InArgs )
+{
+	ParentToolBar = InArgs._ParentToolBar;
+
+	const FName ButtonStyle = FEditorStyle::Join(InArgs._Style.Get(), ".Button");
+
+	const FSlateIcon& Icon = InArgs._Icon.Get();
+
+	SAssignNew( MenuAnchor, SMenuAnchor )
+	.Placement( MenuPlacement_BelowAnchor )
+	.OnGetMenuContent( InArgs._OnGetMenuContent )
+	[
+		SNew(SButton)
+		.ButtonStyle( FEditorStyle::Get(), ButtonStyle )
+		.ContentPadding( FMargin( 5.0f, 0.0f ) )
+		.OnClicked(this, &SViewportToolBarIconMenu::OnMenuClicked)
+		[
+			SNew(SHorizontalBox)
+
+			// Icon
+			+SHorizontalBox::Slot()
+			.AutoWidth()
+			[
+				SNew(SBox)
+				.WidthOverride( 16 )
+				.HeightOverride( 16 )
+				.HAlign(HAlign_Center)
+				.VAlign(VAlign_Center)
+				[
+					SNew(SImage)
+					.Image(Icon.GetIcon())
+				]
+			]
+
+			// Spacer
+			+SHorizontalBox::Slot()
+			.AutoWidth()
+			.Padding( FMargin( 4.0f, 0.0f ) )
+
+			// Menu dropdown
+			+SHorizontalBox::Slot()
+			.AutoWidth()
+			[
+				SNew(SVerticalBox)
+				+SVerticalBox::Slot()
+				.AutoHeight()
+				.HAlign(HAlign_Center)
+				.VAlign(VAlign_Top)
+				[
+					SNew(STextBlock)
+					.TextStyle( FEditorStyle::Get(), FEditorStyle::Join( InArgs._Style.Get(), ".Label" ) )
+					.Text(InArgs._Label)
+				]
+				+SVerticalBox::Slot()
+				.AutoHeight()
+				.VAlign(VAlign_Bottom)
+				[
+					SNew(SHorizontalBox)
+					+SHorizontalBox::Slot()
+					.FillWidth(1.0f)
+
+					+SHorizontalBox::Slot()
+					.AutoWidth()
+					[
+						SNew( SBox )
+						.WidthOverride( 4 )
+						.HeightOverride( 4 )
+						[
+							SNew(SImage)
+							.Image(FEditorStyle::GetBrush("ComboButton.Arrow"))
+							.ColorAndOpacity(FLinearColor::Black)
+						]
+					]
+
+					+SHorizontalBox::Slot()
+					.FillWidth(1.0f)
+				]
+			]
+		]
+	];
+
+	ChildSlot
+	[
+		MenuAnchor.ToSharedRef()
+	];
+}
+
+FReply SViewportToolBarIconMenu::OnMenuClicked()
+{
+	// If the menu button is clicked toggle the state of the menu anchor which will open or close the menu
+	MenuAnchor->SetIsOpen( !MenuAnchor->IsOpen() );
+	ParentToolBar.Pin()->SetOpenMenu( MenuAnchor );
+	return FReply::Handled();
+}
+	
+void SViewportToolBarIconMenu::OnMouseEnter( const FGeometry& MyGeometry, const FPointerEvent& MouseEvent )
+{
+	// See if there is another menu on the same tool bar already open
+	TWeakPtr<SMenuAnchor> OpenedMenu = ParentToolBar.Pin()->GetOpenMenu();
+	if( OpenedMenu.IsValid() && OpenedMenu.Pin()->IsOpen() && OpenedMenu.Pin() != MenuAnchor )
+	{
+		// There is another menu open so we open this menu and close the other
+		ParentToolBar.Pin()->SetOpenMenu( MenuAnchor ); 
+		MenuAnchor->SetIsOpen( true );
+	}
+}
diff --git a/Engine/Source/Editor/UnrealEd/Public/SViewportToolBarComboMenu.h b/Engine/Source/Editor/UnrealEd/Public/SViewportToolBarComboMenu.h
new file mode 100644
--- /dev/null
+++ b/Engine/Source/Editor/UnrealEd/Public/SViewportToolBarComboMenu.h
@@ -0,0 +1,76 @@
+// Copyright 1998-2014 Epic Games, Inc. All Rights Reserved.
+
+#pragma once
+
+#include "SViewportToolBar.h"
+
+/**
+ * Custom widget to display a toggle/drop down menu. 
+ *	Displayed as shown:
+ * +--------+-------------+
+ * | Toggle | Menu button |
+ * +--------+-------------+
+ */
+class UNREALED_API SViewportToolBarComboMenu : public SCompoundWidget 
+{
+public:
+	SLATE_BEGIN_ARGS(SViewportToolBarComboMenu) : _BlockLocation(EMultiBlockLocation::Start) {}
+	
+		/** We need to know about the toolbar we are in */
+		SLATE_ARGUMENT( TSharedPtr<class SViewportToolBar>, ParentToolBar );
+
+		/** Content for the drop down menu  */
+		SLATE_EVENT( FOnGetContent, OnGetMenuContent )
+
+		/** Called upon state change with the value of the next state */
+		SLATE_EVENT( FOnCheckStateChanged, OnCheckStateChanged )
+			
+		/** Sets the current checked state of the checkbox */
+		SLATE_ATTRIBUTE( ESlateCheckBoxState::Type, IsChecked )
+
+		/** Icon shown on the toggle button */
+		SLATE_ATTRIBUTE( FSlateIcon, Icon )
+
+		/** Label shown on the menu button */
+		SLATE_ATTRIBUTE( FText, Label )
+
+		/** Overall style */
+		SLATE_ATTRIBUTE( FName, Style )
+
+		/** ToolTip shown on the menu button */
+		SLATE_ATTRIBUTE( FText, MenuButtonToolTip )
+
+		/** ToolTip shown on the toggle button */
+		SLATE_ATTRIBUTE( FText, ToggleButtonToolTip )
+
+		/** The button location */
+		SLATE_ARGUMENT( EMultiBlockLocation::Type, BlockLocation )
+
+	SLATE_END_ARGS( )
+
+	/**
+	 * Construct this widget
+	 *
+	 * @param	InArgs	The declaration data for this widget
+	 */
+	void Construct( const FArguments& InArgs );
+	
+private:
+	/**
+	 * Called when the menu button is clicked.  Will toggle the visibility of the menu content                   
+	 */
+	FReply OnMenuClicked();
+
+	/**
+	 * Called when the mouse enters a menu button.  If there was a menu previously opened
+	 * we open this menu automatically
+	 */
+	void OnMouseEnter( const FGeometry& MyGeometry, const FPointerEvent& MouseEvent );
+
+private:
+	/** Our menus anchor */
+	TSharedPtr<SMenuAnchor> MenuAnchor;
+
+	/** Parent tool bar for querying other open menus */
+	TWeakPtr<class SViewportToolBar> ParentToolBar;
+};
diff --git a/Engine/Source/Editor/UnrealEd/Public/SViewportToolBarIconMenu.h b/Engine/Source/Editor/UnrealEd/Public/SViewportToolBarIconMenu.h
new file mode 100644
--- /dev/null
+++ b/Engine/Source/Editor/UnrealEd/Public/SViewportToolBarIconMenu.h
@@ -0,0 +1,61 @@
+// Copyright 1998-2014 Epic Games, Inc. All Rights Reserved.
+
+#pragma once
+
+#include "SViewportToolBar.h"
+
+/**
+ * Custom widget to display an icon/drop down menu. 
+ *	Displayed as shown:
+ * +--------+-------------+
+ * | Icon     Menu button |
+ * +--------+-------------+
+ */
+class UNREALED_API SViewportToolBarIconMenu : public SCompoundWidget 
+{
+public:
+	SLATE_BEGIN_ARGS( SViewportToolBarIconMenu ){}
+	
+		/** We need to know about the toolbar we are in */
+		SLATE_ARGUMENT( TSharedPtr<class SViewportToolBar>, ParentToolBar );
+
+		/** Content for the drop down menu  */
+		SLATE_EVENT( FOnGetContent, OnGetMenuContent )
+
+		/** Icon shown */
+		SLATE_ATTRIBUTE( FSlateIcon, Icon )
+
+		/** Label shown on the menu button */
+		SLATE_ATTRIBUTE( FText, Label )
+
+		/** Overall style */
+		SLATE_ATTRIBUTE( FName, Style )
+
+	SLATE_END_ARGS( )
+
+	/**
+	 * Construct this widget
+	 *
+	 * @param	InArgs	The declaration data for this widget
+	 */
+	void Construct( const FArguments& InArgs );
+	
+private:
+	/**
+	 * Called when the menu button is clicked.  Will toggle the visibility of the menu content                   
+	 */
+	FReply OnMenuClicked();
+
+	/**
+	 * Called when the mouse enters a menu button.  If there was a menu previously opened
+	 * we open this menu automatically
+	 */
+	void OnMouseEnter(const FGeometry& MyGeometry, const FPointerEvent& MouseEvent);
+
+private:
+	/** Our menus anchor */
+	TSharedPtr<SMenuAnchor> MenuAnchor;
+
+	/** Parent tool bar for querying other open menus */
+	TWeakPtr<class SViewportToolBar> ParentToolBar;
+};
diff --git a/Engine/Source/Runtime/Engine/Private/ObjectEditorUtils.cpp b/Engine/Source/Runtime/Engine/Private/ObjectEditorUtils.cpp
--- a/Engine/Source/Runtime/Engine/Private/ObjectEditorUtils.cpp
+++ b/Engine/Source/Runtime/Engine/Private/ObjectEditorUtils.cpp
@@ -94,7 +94,9 @@ namespace FObjectEditorUtils
 
 		if ( !DestinationObject->HasAnyFlags(RF_ClassDefaultObject) )
 		{
-			DestinationObject->PreEditChange(DestinationProperty);
+			FEditPropertyChain PropertyChain;
+			PropertyChain.AddHead(DestinationProperty);
+			DestinationObject->PreEditChange(PropertyChain);
 		}
 
 		// Set the value on the destination object.
diff --git a/Engine/Source/Runtime/Engine/Public/ObjectEditorUtils.h b/Engine/Source/Runtime/Engine/Public/ObjectEditorUtils.h
--- a/Engine/Source/Runtime/Engine/Public/ObjectEditorUtils.h
+++ b/Engine/Source/Runtime/Engine/Public/ObjectEditorUtils.h
@@ -109,7 +109,9 @@ namespace FObjectEditorUtils
 
 		if ( !Object->HasAnyFlags(RF_ClassDefaultObject) )
 		{
-			Object->PreEditChange(Property);
+			FEditPropertyChain PropertyChain;
+			PropertyChain.AddHead(Property);
+			Object->PreEditChange(PropertyChain);
 		}
 
 		// Set the value on the destination object.
diff --git a/Engine/Source/Runtime/UMG/Private/SlateBlueprintLibrary.cpp b/Engine/Source/Runtime/UMG/Private/SlateBlueprintLibrary.cpp
--- a/Engine/Source/Runtime/UMG/Private/SlateBlueprintLibrary.cpp
+++ b/Engine/Source/Runtime/UMG/Private/SlateBlueprintLibrary.cpp
@@ -30,4 +30,9 @@ FVector2D USlateBlueprintLibrary::LocalToAbsolute(const FGeometry& Geometry, FVe
 	return Geometry.LocalToAbsolute(LocalCoordinate);
 }
 
+FVector2D USlateBlueprintLibrary::GetLocalSize(const FGeometry& Geometry)
+{
+	return Geometry.GetLocalSize();
+}
+
 #undef LOCTEXT_NAMESPACE
\ No newline at end of file
diff --git a/Engine/Source/Runtime/UMG/Private/UserWidget.cpp b/Engine/Source/Runtime/UMG/Private/UserWidget.cpp
--- a/Engine/Source/Runtime/UMG/Private/UserWidget.cpp
+++ b/Engine/Source/Runtime/UMG/Private/UserWidget.cpp
@@ -497,7 +497,7 @@ void UUserWidget::NativeTick(const FGeometry& MyGeometry, float InDeltaTime )
 	Tick( MyGeometry, InDeltaTime );
 }
 
-TSharedRef<SWidget> UUserWidget::MakeViewportWidget(bool bModal, bool bShowCursor, TSharedPtr<SWidget>& UserSlateWidget)
+TSharedRef<SWidget> UUserWidget::MakeViewportWidget(TSharedPtr<SWidget>& UserSlateWidget)
 {
 	UserSlateWidget = TakeWidget();
 
@@ -523,12 +523,12 @@ UWidget* UUserWidget::GetRootWidgetComponent()
 	return NULL;
 }
 
-void UUserWidget::AddToViewport(bool bModal, bool bShowCursor)
+void UUserWidget::AddToViewport(bool bModal)
 {
 	if ( !FullScreenWidget.IsValid() )
 	{
 		TSharedPtr<SWidget> OutUserSlateWidget;
-		TSharedRef<SWidget> RootWidget = MakeViewportWidget(bModal, bShowCursor, OutUserSlateWidget);
+		TSharedRef<SWidget> RootWidget = MakeViewportWidget(OutUserSlateWidget);
 
 		TSharedRef<SViewportWidgetHost> WidgetHost = SNew(SViewportWidgetHost, OutUserSlateWidget, bModal)
 			[
diff --git a/Engine/Source/Runtime/UMG/Public/Blueprint/SlateBlueprintLibrary.h b/Engine/Source/Runtime/UMG/Public/Blueprint/SlateBlueprintLibrary.h
--- a/Engine/Source/Runtime/UMG/Public/Blueprint/SlateBlueprintLibrary.h
+++ b/Engine/Source/Runtime/UMG/Public/Blueprint/SlateBlueprintLibrary.h
@@ -14,7 +14,7 @@ class USlateBlueprintLibrary : public UBlueprintFunctionLibrary
 	 *
 	 * @return true if the provided location in absolute coordinates is within the bounds of this geometry.
 	 */
-	UFUNCTION(BlueprintCallable, Category="User Interface|Geometry")
+	 UFUNCTION(BlueprintPure, Category="User Interface|Geometry")
 	static bool IsUnderLocation(const FGeometry& Geometry, const FVector2D& AbsoluteCoordinate);
 
 	/**
@@ -22,7 +22,7 @@ class USlateBlueprintLibrary : public UBlueprintFunctionLibrary
 	 *
 	 * @return Transforms AbsoluteCoordinate into the local space of this Geometry.
 	 */
-	UFUNCTION(BlueprintCallable, Category="User Interface|Geometry")
+	UFUNCTION(BlueprintPure, Category="User Interface|Geometry")
 	static FVector2D AbsoluteToLocal(const FGeometry& Geometry, FVector2D AbsoluteCoordinate);
 
 	/**
@@ -32,6 +32,10 @@ class USlateBlueprintLibrary : public UBlueprintFunctionLibrary
 	 *
 	 * @return  Absolute coordinates
 	 */
-	UFUNCTION(BlueprintCallable, Category="User Interface|Geometry")
+	UFUNCTION(BlueprintPure, Category="User Interface|Geometry")
 	static FVector2D LocalToAbsolute(const FGeometry& Geometry, FVector2D LocalCoordinate);
+
+	/** @return the size of the geometry in local space. */
+	UFUNCTION(BlueprintPure, Category="User Interface|Geometry")
+	static FVector2D GetLocalSize(const FGeometry& Geometry);
 };
diff --git a/Engine/Source/Runtime/UMG/Public/Blueprint/UserWidget.h b/Engine/Source/Runtime/UMG/Public/Blueprint/UserWidget.h
--- a/Engine/Source/Runtime/UMG/Public/Blueprint/UserWidget.h
+++ b/Engine/Source/Runtime/UMG/Public/Blueprint/UserWidget.h
@@ -97,11 +97,16 @@ public:
 	virtual void ReleaseNativeWidget() override;
 	// End of UVisual interface
 
-	/*  */
+	/**
+	 * Adds it to the game's viewport, defaults to filling the entire viewport area.
+	 * @param bModal If this dialog should steal keyboard/mouse focus and consume all input. Great for a fullscreen menu. Terrible for HUDs.
+	 */
 	UFUNCTION(BlueprintCallable, Category="User Interface|Viewport")
-	void AddToViewport(bool bModal = false, bool bShowCursor = false);
+	void AddToViewport(bool bModal = false);
 
-	/*  */
+	/**
+	 * Removes the widget from the viewport.
+	 */
 	UFUNCTION(BlueprintCallable, Category="User Interface|Viewport")
 	void RemoveFromViewport();
 
@@ -258,7 +263,7 @@ public:
 	UWidget* GetWidgetHandle(TSharedRef<SWidget> InWidget);
 
 	/** Creates a fullscreen host widget, that wraps this widget. */
-	TSharedRef<SWidget> MakeViewportWidget(bool bModal, bool bShowCursor, TSharedPtr<SWidget>& UserSlateWidget);
+	TSharedRef<SWidget> MakeViewportWidget(TSharedPtr<SWidget>& UserSlateWidget);
 
 	/** @returns The root UObject widget wrapper */
 	UWidget* GetRootWidgetComponent();
