Build - Adding missing file.

diff --git a/Engine/Source/Editor/UMGEditor/Private/Settings/WidgetDesignerSettings.cpp b/Engine/Source/Editor/UMGEditor/Private/Settings/WidgetDesignerSettings.cpp
new file mode 100644
--- /dev/null
+++ b/Engine/Source/Editor/UMGEditor/Private/Settings/WidgetDesignerSettings.cpp
@@ -0,0 +1,10 @@
+// Copyright 1998-2014 Epic Games, Inc. All Rights Reserved.
+
+#include "UMGEditorPrivatePCH.h"
+
+UWidgetDesignerSettings::UWidgetDesignerSettings(const class FPostConstructInitializeProperties& PCIP)
+	: Super(PCIP)
+{
+	GridSnapEnabled = true;
+	GridSnapSize = 4;
+}
