Build - Trying to fix non-unity builds.

diff --git a/Engine/Source/Editor/UMGEditor/Private/UMGEditorPrivatePCH.h b/Engine/Source/Editor/UMGEditor/Private/UMGEditorPrivatePCH.h
--- a/Engine/Source/Editor/UMGEditor/Private/UMGEditorPrivatePCH.h
+++ b/Engine/Source/Editor/UMGEditor/Private/UMGEditorPrivatePCH.h
@@ -51,3 +51,4 @@
 
 #include "WidgetBlueprintLibrary.h"
 
+#include "WidgetBlueprintCompiler.h"
