UMG - Improving the anchors drop down menu to show a preview of the effect it has.  Adding new button styles to the editor, fixing several bugs with making widgets design time flagged.  No longer asserting on incorrect visibility flags.  Adding major (0.1) snapping to the anchor widget.  Making several addChild functions specific to their widgets.

diff --git a/Engine/Source/Editor/EditorStyle/Private/SlateEditorStyle.cpp b/Engine/Source/Editor/EditorStyle/Private/SlateEditorStyle.cpp
--- a/Engine/Source/Editor/EditorStyle/Private/SlateEditorStyle.cpp
+++ b/Engine/Source/Editor/EditorStyle/Private/SlateEditorStyle.cpp
@@ -249,6 +249,23 @@ void FSlateEditorStyle::FStyle::SetupGeneralStyles()
 			.SetPressedPadding( FMargin(0,1,0,0) );
 	Set( "HoverHintOnly", HoverHintOnly );
 
+
+	FButtonStyle SimpleSharpButton = FButtonStyle()
+		.SetNormal(BOX_BRUSH("Common/Button/simple_sharp_normal", FMargin(4 / 16.0f), FLinearColor(1, 1, 1, 1)))
+		.SetHovered(BOX_BRUSH("Common/Button/simple_sharp_hovered", FMargin(4 / 16.0f), FLinearColor(1, 1, 1, 1)))
+		.SetPressed(BOX_BRUSH("Common/Button/simple_sharp_hovered", FMargin(4 / 16.0f), FLinearColor(1, 1, 1, 1)))
+		.SetNormalPadding(FMargin(0, 0, 0, 1))
+		.SetPressedPadding(FMargin(0, 1, 0, 0));
+	Set("SimpleSharpButton", SimpleSharpButton);
+
+	FButtonStyle SimpleRoundButton = FButtonStyle()
+		.SetNormal(BOX_BRUSH("Common/Button/simple_round_normal", FMargin(4 / 16.0f), FLinearColor(1, 1, 1, 1)))
+		.SetHovered(BOX_BRUSH("Common/Button/simple_round_hovered", FMargin(4 / 16.0f), FLinearColor(1, 1, 1, 1)))
+		.SetPressed(BOX_BRUSH("Common/Button/simple_round_hovered", FMargin(4 / 16.0f), FLinearColor(1, 1, 1, 1)))
+		.SetNormalPadding(FMargin(0, 0, 0, 1))
+		.SetPressedPadding(FMargin(0, 1, 0, 0));
+	Set("SimpleRoundButton", SimpleRoundButton);
+
 	// Common glyphs
 	{
 		Set( "Symbols.SearchGlass", new IMAGE_BRUSH( "Common/SearchGlass", Icon16x16 ) );
diff --git a/Engine/Source/Editor/UMGEditor/Private/Customizations/CanvasSlotCustomization.cpp b/Engine/Source/Editor/UMGEditor/Private/Customizations/CanvasSlotCustomization.cpp
--- a/Engine/Source/Editor/UMGEditor/Private/Customizations/CanvasSlotCustomization.cpp
+++ b/Engine/Source/Editor/UMGEditor/Private/Customizations/CanvasSlotCustomization.cpp
@@ -8,10 +8,167 @@
 #include "ScopedTransaction.h"
 #include "BlueprintEditorUtils.h"
 
+#include "SConstraintCanvas.h"
+
 #include "CanvasSlotCustomization.h"
 
 #define LOCTEXT_NAMESPACE "UMG"
 
+class SAnchorPreviewWidget : public SCompoundWidget
+{
+public:
+
+	SLATE_BEGIN_ARGS(SAnchorPreviewWidget) {}
+	SLATE_END_ARGS()
+
+	void Construct(const FArguments& InArgs, 
+		TSharedPtr<IPropertyHandle> AnchorsHandle,
+		TSharedPtr<IPropertyHandle> AlignmentHandle,
+		TSharedPtr<IPropertyHandle> OffsetsHandle,
+		FText LabelText,
+		FAnchors Anchors)
+	{
+		ResizeCurve = FCurveSequence(0, 0.40f);
+
+		ChildSlot
+		[
+			SNew(SButton)
+			.ButtonStyle(FEditorStyle::Get(), "SimpleSharpButton")
+			//.ButtonColorAndOpacity(FLinearColor(FColor(64, 64, 64)))
+			.ButtonColorAndOpacity(FLinearColor(FColor(48, 48, 48)))
+			.OnClicked(this, &SAnchorPreviewWidget::OnAnchorClicked, AnchorsHandle, AlignmentHandle, OffsetsHandle, Anchors)
+			.ContentPadding(FMargin(2.0f, 2.0f))
+			[
+				SNew(SVerticalBox)
+
+				+ SVerticalBox::Slot()
+				.AutoHeight()
+				[
+					SNew(SBox)
+					.WidthOverride(64)
+					.HeightOverride(64)
+					.HAlign(HAlign_Center)
+					.VAlign(VAlign_Center)
+					[
+						SNew(SBox)
+						.WidthOverride(this, &SAnchorPreviewWidget::GetCurrentWidth)
+						.HeightOverride(this, &SAnchorPreviewWidget::GetCurrentHeight)
+						[
+							SNew(SBorder)
+							.BorderImage(FEditorStyle::GetBrush("ToolPanel.GroupBorder"))
+							.Padding(1)
+							[
+								SNew(SConstraintCanvas)
+
+								+ SConstraintCanvas::Slot()
+								.Anchors(Anchors)
+								.Offset(FMargin(0, 0, Anchors.IsStretchedHorizontal() ? 0 : 15, Anchors.IsStretchedVertical() ? 0 : 15))
+								.Alignment(FVector2D(Anchors.IsStretchedHorizontal() ? 0 : Anchors.Minimum.X, Anchors.IsStretchedVertical() ? 0 : Anchors.Minimum.Y))
+								[
+									SNew(SImage)
+									.Image(FEditorStyle::Get().GetBrush("UMGEditor.AnchoredWidget"))
+								]
+
+								//+ SConstraintCanvas::Slot()
+								//.Anchors(Anchors)
+								//.Offset(FMargin(0, 0, 8, 8))
+								//.Alignment(FVector2D(Anchors.IsStretchedHorizontal() ? 0 : Anchors.Minimum.X, Anchors.IsStretchedVertical() ? 0 : Anchors.Minimum.Y))
+								//[
+								//	SNew(SImage)
+								//	.Image(FEditorStyle::Get().GetBrush("UMGEditor.AnchoredWidgetAlignment"))
+								//]
+							]
+						]
+					]
+				]
+
+				//+ SVerticalBox::Slot()
+				//.AutoHeight()
+				//.HAlign(HAlign_Center)
+				//.Padding(0,2,0,0)
+				//[
+				//	SNew(STextBlock)
+				//	.Text(LabelText)
+				//]
+			]
+		];
+	}
+
+	void Tick(const FGeometry& AllottedGeometry, const double InCurrentTime, const float InDeltaTime) override
+	{
+		if ( AllottedGeometry.IsUnderLocation(FSlateApplication::Get().GetCursorPos()) )
+		{
+			if ( !ResizeCurve.IsPlaying() )
+			{
+				if ( ResizeCurve.IsAtStart() )
+				{
+					ResizeCurve.Play();
+				}
+				else if ( ResizeCurve.IsAtEnd() )
+				{
+					ResizeCurve.PlayReverse();
+				}
+			}
+		}
+		else
+		{
+			if ( !ResizeCurve.IsAtStart() && ((!ResizeCurve.IsInReverse() && ResizeCurve.IsPlaying()) || !ResizeCurve.IsPlaying()) )
+			{
+				ResizeCurve.PlayReverse();
+			}
+		}
+	}
+
+private:
+
+	FOptionalSize GetCurrentWidth() const
+	{
+		return 48 + ( 16 * ResizeCurve.GetLerp() );
+	}
+
+	FOptionalSize GetCurrentHeight() const
+	{
+		return 48 + ( 16 * ResizeCurve.GetLerp() );
+	}
+
+	FReply OnAnchorClicked(
+		TSharedPtr<IPropertyHandle> AnchorsHandle,
+		TSharedPtr<IPropertyHandle> AlignmentHandle,
+		TSharedPtr<IPropertyHandle> OffsetsHandle,
+		FAnchors Anchors)
+	{
+		FScopedTransaction Transaction(LOCTEXT("ChangeAnchors", "Changed Anchors"));
+
+		FString Value = FString::Printf(TEXT("(Minimum=(X=%f,Y=%f),Maximum=(X=%f,Y=%f))"), Anchors.Minimum.X, Anchors.Minimum.Y, Anchors.Maximum.X, Anchors.Maximum.Y);
+		AnchorsHandle->SetValueFromFormattedString(Value);
+
+		if ( FSlateApplication::Get().GetModifierKeys().IsShiftDown() )
+		{
+			FString Value = FString::Printf(TEXT("(X=%f,Y=%f)"), Anchors.IsStretchedHorizontal() ? 0 : Anchors.Minimum.X, Anchors.IsStretchedVertical() ? 0 : Anchors.Minimum.Y);
+			AlignmentHandle->SetValueFromFormattedString(Value);
+		}
+
+		if ( FSlateApplication::Get().GetModifierKeys().IsControlDown() )
+		{
+			TArray<void*> RawOffsetData;
+			OffsetsHandle->AccessRawData(RawOffsetData);
+			FMargin* Offsets = reinterpret_cast<FMargin*>( RawOffsetData[0] );
+
+			FString Value = FString::Printf(TEXT("(Left=%f,Top=%f,Right=%f,Bottom=%f)"), 0, 0, Anchors.IsStretchedHorizontal() ? 0 : Offsets->Right, Anchors.IsStretchedVertical() ? 0 : Offsets->Bottom);
+			OffsetsHandle->SetValueFromFormattedString(Value);
+		}
+
+		// Close the menu
+		FSlateApplication::Get().DismissAllMenus();
+
+		return FReply::Handled();
+	}
+
+private:
+	FCurveSequence ResizeCurve;
+};
+
+
 // FCanvasSlotCustomization
 ////////////////////////////////////////////////////////////////////////////////
 
@@ -140,196 +297,156 @@ FText FCanvasSlotCustomization::GetOffsetLabel(TSharedPtr<IPropertyHandle> Prope
 void FCanvasSlotCustomization::CustomizeAnchors(TSharedPtr<IPropertyHandle> PropertyHandle, IDetailChildrenBuilder& ChildBuilder, IPropertyTypeCustomizationUtils& CustomizationUtils)
 {
 	TSharedPtr<IPropertyHandle> AnchorsHandle = PropertyHandle->GetChildHandle(GET_MEMBER_NAME_CHECKED(FAnchorData, Anchors));
+	TSharedPtr<IPropertyHandle> AlignmentHandle = PropertyHandle->GetChildHandle(GET_MEMBER_NAME_CHECKED(FAnchorData, Alignment));
+	TSharedPtr<IPropertyHandle> OffsetsHandle = PropertyHandle->GetChildHandle(GET_MEMBER_NAME_CHECKED(FAnchorData, Offsets));
 
 	AnchorsHandle->MarkHiddenByCustomization();
 
 	IDetailPropertyRow& AnchorsPropertyRow = ChildBuilder.AddChildProperty(AnchorsHandle.ToSharedRef());
 
+	const float FillDividePadding = 1;
+
 	AnchorsPropertyRow.CustomWidget(/*bShowChildren*/ true)
 		.NameContent()
 		[
-			SNew(STextBlock)
+			SNew( STextBlock )
 			.Font( IDetailLayoutBuilder::GetDetailFont() )
-			.Text(LOCTEXT("Anchors", "Anchors"))
+			.Text( LOCTEXT("Anchors", "Anchors") )
 		]
 		.ValueContent()
 		[
-			SNew(SHorizontalBox)
-					
-			+ SHorizontalBox::Slot()
-			.AutoWidth()
+			SNew(SComboButton)
+			.ButtonContent()
 			[
-				SNew(SComboButton)
-				.ButtonContent()
-				[
-					SNew(STextBlock)
-					.Text(LOCTEXT("AnchorsText", "Anchors"))
-				]
-				.MenuContent()
+				SNew(STextBlock)
+				.Text(LOCTEXT("AnchorsText", "Anchors"))
+			]
+			.MenuContent()
+			[
+				SNew(SBorder)
+				//.BorderImage(FEditorStyle::GetBrush("ToolPanel.GroupBorder"))FEditorStyle::GetBrush("WhiteBrush")/*FEditorStyle::GetBrush("ToolPanel.GroupBorder")*/)
+				.Padding(5)
 				[
-					SNew(SUniformGridPanel)
-
-					// Top Row
-					+ SUniformGridPanel::Slot(0, 0)
+					SNew(SBorder)
+					.BorderImage(FEditorStyle::GetBrush("WhiteBrush"))
+					.BorderBackgroundColor(FLinearColor(FColor(66, 139, 202)))
+					.Padding(0)
 					[
-						SNew(SButton)
-						.OnClicked(this, &FCanvasSlotCustomization::OnAnchorClicked, AnchorsHandle, FAnchors(0, 0, 0, 0))
-						[
-							SNew(SImage)
-							.Image(FEditorStyle::Get().GetBrush("UMGEditor.TopLeft"))
-						]
-					]
+						SNew(SVerticalBox)
 
-					+ SUniformGridPanel::Slot(1, 0)
-					[
-						SNew(SButton)
-						.OnClicked(this, &FCanvasSlotCustomization::OnAnchorClicked, AnchorsHandle, FAnchors(0.5, 0, 0.5, 0))
-						[
-							SNew(SImage)
-							.Image(FEditorStyle::Get().GetBrush("UMGEditor.TopCenter"))
-						]
-					]
-
-					+ SUniformGridPanel::Slot(2, 0)
-					[
-						SNew(SButton)
-						.OnClicked(this, &FCanvasSlotCustomization::OnAnchorClicked, AnchorsHandle, FAnchors(1, 0, 1, 0))
+						+ SVerticalBox::Slot()
+						.AutoHeight()
 						[
-							SNew(SImage)
-							.Image(FEditorStyle::Get().GetBrush("UMGEditor.TopRight"))
+							SNew(SHorizontalBox)
+
+							+ SHorizontalBox::Slot()
+							.AutoWidth()
+							[
+								SNew(SUniformGridPanel)
+
+								// Top Row
+								+ SUniformGridPanel::Slot(0, 0)
+								[
+									SNew(SAnchorPreviewWidget, AnchorsHandle, AlignmentHandle, OffsetsHandle, LOCTEXT("TopLeft", "Top/Left"), FAnchors(0, 0, 0, 0))
+								]
+
+								+ SUniformGridPanel::Slot(1, 0)
+								[
+									SNew(SAnchorPreviewWidget, AnchorsHandle, AlignmentHandle, OffsetsHandle, LOCTEXT("TopCenter", "Top/Center"), FAnchors(0.5, 0, 0.5, 0))
+								]
+
+								+ SUniformGridPanel::Slot(2, 0)
+								[
+									SNew(SAnchorPreviewWidget, AnchorsHandle, AlignmentHandle, OffsetsHandle, LOCTEXT("TopRight", "Top/Right"), FAnchors(1, 0, 1, 0))
+								]
+
+								// Center Row
+								+ SUniformGridPanel::Slot(0, 1)
+								[
+									SNew(SAnchorPreviewWidget, AnchorsHandle, AlignmentHandle, OffsetsHandle, LOCTEXT("CenterLeft", "Center/Left"), FAnchors(0, 0.5, 0, 0.5))
+								]
+
+								+ SUniformGridPanel::Slot(1, 1)
+								[
+									SNew(SAnchorPreviewWidget, AnchorsHandle, AlignmentHandle, OffsetsHandle, LOCTEXT("CenterCenter", "Center/Center"), FAnchors(0.5, 0.5, 0.5, 0.5))
+								]
+
+								+ SUniformGridPanel::Slot(2, 1)
+								[
+									SNew(SAnchorPreviewWidget, AnchorsHandle, AlignmentHandle, OffsetsHandle, LOCTEXT("CenterRight", "Center/Right"), FAnchors(1, 0.5, 1, 0.5))
+								]
+
+								// Bottom Row
+								+ SUniformGridPanel::Slot(0, 2)
+								[
+									SNew(SAnchorPreviewWidget, AnchorsHandle, AlignmentHandle, OffsetsHandle, LOCTEXT("BottomLeft", "Bottom/Left"), FAnchors(0, 1, 0, 1))
+								]
+
+								+ SUniformGridPanel::Slot(1, 2)
+								[
+									SNew(SAnchorPreviewWidget, AnchorsHandle, AlignmentHandle, OffsetsHandle, LOCTEXT("BottomCenter", "Bottom/Center"), FAnchors(0.5, 1, 0.5, 1))
+								]
+
+								+ SUniformGridPanel::Slot(2, 2)
+								[
+									SNew(SAnchorPreviewWidget, AnchorsHandle, AlignmentHandle, OffsetsHandle, LOCTEXT("BottomRight", "Bottom/Right"), FAnchors(1, 1, 1, 1))
+								]
+							]
+
+							+ SHorizontalBox::Slot()
+							.AutoWidth()
+							.Padding(FillDividePadding, 0, 0, 0)
+							[
+								SNew(SUniformGridPanel)
+
+								+ SUniformGridPanel::Slot(0, 0)
+								[
+									SNew(SAnchorPreviewWidget, AnchorsHandle, AlignmentHandle, OffsetsHandle, LOCTEXT("TopFill", "Top/Fill"), FAnchors(0, 0, 1, 0))
+								]
+
+								+ SUniformGridPanel::Slot(0, 1)
+								[
+									SNew(SAnchorPreviewWidget, AnchorsHandle, AlignmentHandle, OffsetsHandle, LOCTEXT("CenterFill", "Center/Fill"), FAnchors(0, 0.5f, 1, 0.5f))
+								]
+
+								+ SUniformGridPanel::Slot(0, 2)
+								[
+									SNew(SAnchorPreviewWidget, AnchorsHandle, AlignmentHandle, OffsetsHandle, LOCTEXT("BottomFill", "Bottom/Fill"), FAnchors(0, 1, 1, 1))
+								]
+							]
 						]
-					]
 
-					+ SUniformGridPanel::Slot(3, 0)
-					[
-						SNew(SButton)
-						.OnClicked(this, &FCanvasSlotCustomization::OnAnchorClicked, AnchorsHandle, FAnchors(0, 0, 1, 0))
+						+ SVerticalBox::Slot()
+						.AutoHeight()
+						.Padding(0, FillDividePadding, 0, 0)
 						[
-							SNew(SImage)
-							.Image(FEditorStyle::Get().GetBrush("UMGEditor.TopFill"))
-						]
-					]
-
-					// Center Row
-					+ SUniformGridPanel::Slot(0, 1)
-					[
-						SNew(SButton)
-						.OnClicked(this, &FCanvasSlotCustomization::OnAnchorClicked, AnchorsHandle, FAnchors(0, 0.5, 0, 0.5))
-						[
-							SNew(SImage)
-							.Image(FEditorStyle::Get().GetBrush("UMGEditor.CenterLeft"))
-						]
-					]
-
-					+ SUniformGridPanel::Slot(1, 1)
-					[
-						SNew(SButton)
-						.OnClicked(this, &FCanvasSlotCustomization::OnAnchorClicked, AnchorsHandle, FAnchors(0.5, 0.5, 0.5, 0.5))
-						[
-							SNew(SImage)
-							.Image(FEditorStyle::Get().GetBrush("UMGEditor.CenterCenter"))
-						]
-					]
-
-					+ SUniformGridPanel::Slot(2, 1)
-					[
-						SNew(SButton)
-						.OnClicked(this, &FCanvasSlotCustomization::OnAnchorClicked, AnchorsHandle, FAnchors(1, 0.5, 1, 0.5))
-						[
-							SNew(SImage)
-							.Image(FEditorStyle::Get().GetBrush("UMGEditor.CenterRight"))
-						]
-					]
-
-					+ SUniformGridPanel::Slot(3, 1)
-					[
-						SNew(SButton)
-						.OnClicked(this, &FCanvasSlotCustomization::OnAnchorClicked, AnchorsHandle, FAnchors(0, 0.5f, 1, 0.5f))
-						[
-							SNew(SImage)
-							.Image(FEditorStyle::Get().GetBrush("UMGEditor.CenterFill"))
-						]
-					]
-
-					// Bottom Row
-					+ SUniformGridPanel::Slot(0, 2)
-					[
-						SNew(SButton)
-						.OnClicked(this, &FCanvasSlotCustomization::OnAnchorClicked, AnchorsHandle, FAnchors(0, 1, 0, 1))
-						[
-							SNew(SImage)
-							.Image(FEditorStyle::Get().GetBrush("UMGEditor.BottomLeft"))
-						]
-					]
-
-					+ SUniformGridPanel::Slot(1, 2)
-					[
-						SNew(SButton)
-						.OnClicked(this, &FCanvasSlotCustomization::OnAnchorClicked, AnchorsHandle, FAnchors(0.5, 1, 0.5, 1))
-						[
-							SNew(SImage)
-							.Image(FEditorStyle::Get().GetBrush("UMGEditor.BottomCenter"))
-						]
-					]
-
-					+ SUniformGridPanel::Slot(2, 2)
-					[
-						SNew(SButton)
-						.OnClicked(this, &FCanvasSlotCustomization::OnAnchorClicked, AnchorsHandle, FAnchors(1, 1, 1, 1))
-						[
-							SNew(SImage)
-							.Image(FEditorStyle::Get().GetBrush("UMGEditor.BottomRight"))
-						]
-					]
-
-					+ SUniformGridPanel::Slot(3, 2)
-					[
-						SNew(SButton)
-						.OnClicked(this, &FCanvasSlotCustomization::OnAnchorClicked, AnchorsHandle, FAnchors(0, 1, 1, 1))
-						[
-							SNew(SImage)
-							.Image(FEditorStyle::Get().GetBrush("UMGEditor.BottomFill"))
-						]
-					]
-
-					// Fill Row
-					+ SUniformGridPanel::Slot(0, 3)
-					[
-						SNew(SButton)
-						.OnClicked(this, &FCanvasSlotCustomization::OnAnchorClicked, AnchorsHandle, FAnchors(0, 0, 0, 1))
-						[
-							SNew(SImage)
-							.Image(FEditorStyle::Get().GetBrush("UMGEditor.FillLeft"))
-						]
-					]
-
-					+ SUniformGridPanel::Slot(1, 3)
-					[
-						SNew(SButton)
-						.OnClicked(this, &FCanvasSlotCustomization::OnAnchorClicked, AnchorsHandle, FAnchors(0.5, 0, 0.5, 1))
-						[
-							SNew(SImage)
-							.Image(FEditorStyle::Get().GetBrush("UMGEditor.FillCenter"))
-						]
-					]
-
-					+ SUniformGridPanel::Slot(2, 3)
-					[
-						SNew(SButton)
-						.OnClicked(this, &FCanvasSlotCustomization::OnAnchorClicked, AnchorsHandle, FAnchors(1, 0, 1, 1))
-						[
-							SNew(SImage)
-							.Image(FEditorStyle::Get().GetBrush("UMGEditor.FillRight"))
-						]
-					]
-
-					+ SUniformGridPanel::Slot(3, 3)
-					[
-						SNew(SButton)
-						.OnClicked(this, &FCanvasSlotCustomization::OnAnchorClicked, AnchorsHandle, FAnchors(0, 0, 1, 1))
-						[
-							SNew(SImage)
-							.Image(FEditorStyle::Get().GetBrush("UMGEditor.FillFill"))
+							SNew(SHorizontalBox)
+
+							// Fill Row
+							+ SHorizontalBox::Slot()
+							.AutoWidth()
+							[
+								SNew(SAnchorPreviewWidget, AnchorsHandle, AlignmentHandle, OffsetsHandle, LOCTEXT("FillLeft", "Fill/Left"), FAnchors(0, 0, 0, 1))
+							]
+
+							+ SHorizontalBox::Slot()
+							.AutoWidth()
+							[
+								SNew(SAnchorPreviewWidget, AnchorsHandle, AlignmentHandle, OffsetsHandle, LOCTEXT("FillCenter", "Fill/Center"), FAnchors(0.5, 0, 0.5, 1))
+							]
+
+							+ SHorizontalBox::Slot()
+							.AutoWidth()
+							[
+								SNew(SAnchorPreviewWidget, AnchorsHandle, AlignmentHandle, OffsetsHandle, LOCTEXT("FillRight", "Fill/Right"), FAnchors(1, 0, 1, 1))
+							]
+
+							+ SHorizontalBox::Slot()
+							.AutoWidth()
+							.Padding(FillDividePadding, 0, 0, 0)
+							[
+								SNew(SAnchorPreviewWidget, AnchorsHandle, AlignmentHandle, OffsetsHandle, LOCTEXT("FillFill", "Fill/Fill"), FAnchors(0, 0, 1, 1))
+							]
 						]
 					]
 				]
@@ -337,20 +454,4 @@ void FCanvasSlotCustomization::CustomizeAnchors(TSharedPtr<IPropertyHandle> Prop
 		];
 }
 
-FReply FCanvasSlotCustomization::OnAnchorClicked(TSharedPtr<IPropertyHandle> AnchorsHandle, FAnchors Anchors)
-{
-	//TODO UMG Group under single undo transaction.
-
-	FString Value = FString::Printf(TEXT("(Minimum=(X=%f,Y=%f),Maximum=(X=%f,Y=%f))"), Anchors.Minimum.X, Anchors.Minimum.Y, Anchors.Maximum.X, Anchors.Maximum.Y);
-	AnchorsHandle->SetValueFromFormattedString(Value);
-
-	//TODO UMG Reset position, pivot.
-	//TODO UMG Moving the anchors shouldn't change the position of the widget, unless stretched.  But maybe that too should be based off current position and size.
-	//         Either way need access to cached geometry of the slot in the preview world...
-	
-	FSlateApplication::Get().DismissAllMenus();
-
-	return FReply::Handled();
-}
-
 #undef LOCTEXT_NAMESPACE
diff --git a/Engine/Source/Editor/UMGEditor/Private/Customizations/CanvasSlotCustomization.h b/Engine/Source/Editor/UMGEditor/Private/Customizations/CanvasSlotCustomization.h
--- a/Engine/Source/Editor/UMGEditor/Private/Customizations/CanvasSlotCustomization.h
+++ b/Engine/Source/Editor/UMGEditor/Private/Customizations/CanvasSlotCustomization.h
@@ -31,8 +31,6 @@ private:
 	void CreateEditorWithDynamicLabel(IDetailPropertyRow& PropertyRow, TAttribute<FText> TextAttribute);
 
 	static FText GetOffsetLabel(TSharedPtr<IPropertyHandle> AnchorStructureHandle, EOrientation Orientation, FText NonStretchingLabel, FText StretchingLabel);
-	
-	FReply OnAnchorClicked(TSharedPtr<IPropertyHandle> AnchorsHandle, FAnchors Anchors);
 
 private:
 	UWidgetBlueprint* Blueprint;
diff --git a/Engine/Source/Editor/UMGEditor/Private/Designer/SDesignerView.cpp b/Engine/Source/Editor/UMGEditor/Private/Designer/SDesignerView.cpp
--- a/Engine/Source/Editor/UMGEditor/Private/Designer/SDesignerView.cpp
+++ b/Engine/Source/Editor/UMGEditor/Private/Designer/SDesignerView.cpp
@@ -849,8 +849,10 @@ FReply SDesignerView::OnMouseMove(const FGeometry& MyGeometry, const FPointerEve
 						FWidgetTransform RenderTransform = PreviewWidget->RenderTransform;
 						RenderTransform.Translation += AbsoluteToLocalTransform.TransformVector(MouseEvent.GetCursorDelta());
 
-						FObjectEditorUtils::SetPropertyValue<UWidget, FWidgetTransform>(PreviewWidget, "RenderTransform", RenderTransform);
-						FObjectEditorUtils::SetPropertyValue<UWidget, FWidgetTransform>(SelectedWidget.GetTemplate(), "RenderTransform", RenderTransform);
+						static const FName RenderTransformName(TEXT("RenderTransform"));
+
+						FObjectEditorUtils::SetPropertyValue<UWidget, FWidgetTransform>(PreviewWidget, RenderTransformName, RenderTransform);
+						FObjectEditorUtils::SetPropertyValue<UWidget, FWidgetTransform>(SelectedWidget.GetTemplate(), RenderTransformName, RenderTransform);
 					}
 				}
 			}
@@ -1238,7 +1240,7 @@ UWidget* SDesignerView::ProcessDropAndAddWidget(const FGeometry& MyGeometry, con
 
 			// Add it to the root if there are no other widgets to add it to.
 			UWidget* Widget = TemplateDragDropOp->Template->Create(BP->WidgetTree);
-			Widget->IsDesignTime(true);
+			Widget->SetIsDesignTime(true);
 
 			BP->WidgetTree->RootWidget = Widget;
 
@@ -1272,7 +1274,7 @@ UWidget* SDesignerView::ProcessDropAndAddWidget(const FGeometry& MyGeometry, con
 
 			// Construct the widget and mark it for design time rendering.
 			UWidget* Widget = TemplateDragDropOp->Template->Create(BP->WidgetTree);
-			Widget->IsDesignTime(true);
+			Widget->SetIsDesignTime(true);
 
 			// Determine local position inside the parent widget and add the widget to the slot.
 			FVector2D LocalPosition = WidgetUnderCursorGeometry.AbsoluteToLocal(DragDropEvent.GetScreenSpacePosition());
@@ -1368,77 +1370,81 @@ UWidget* SDesignerView::ProcessDropAndAddWidget(const FGeometry& MyGeometry, con
 			}
 
 			UWidget* Widget = bIsPreview ? SelectedDragDropOp->Widget.GetPreview() : SelectedDragDropOp->Widget.GetTemplate();
+			check(Widget);
 
-			if ( Widget->GetParent() )
+			if ( Widget )
 			{
-				if ( !bIsPreview )
+				if ( Widget->GetParent() )
 				{
-					Widget->GetParent()->Modify();
-				}
-
-				Widget->GetParent()->RemoveChild(Widget);
-			}
-
-			FVector2D LocalPosition = WidgetUnderCursorGeometry.AbsoluteToLocal(DragDropEvent.GetScreenSpacePosition());
-			if ( UPanelSlot* Slot = NewParent->AddChild(Widget) )
-			{
-				FVector2D NewPosition = LocalPosition - SelectedWidgetContextMenuLocation;
+					if ( !bIsPreview )
+					{
+						Widget->GetParent()->Modify();
+					}
 
-				const UWidgetDesignerSettings* DesignerSettings = GetDefault<UWidgetDesignerSettings>();
-				if ( DesignerSettings->GridSnapEnabled )
-				{
-					NewPosition.X = ( (int32)NewPosition.X ) - ( ( (int32)NewPosition.X ) % DesignerSettings->GridSnapSize );
-					NewPosition.Y = ( (int32)NewPosition.Y ) - ( ( (int32)NewPosition.Y ) % DesignerSettings->GridSnapSize );
+					Widget->GetParent()->RemoveChild(Widget);
 				}
 
-				// HACK UMG: In order to correctly drop items into the canvas that have a non-zero anchor,
-				// we need to know the layout information after slate has performed a prepass.  So we have
-				// to rebase the layout and reinterpret the new position based on anchor point layout data.
-				// This should be pulled out into an extension of some kind so that this can be fixed for
-				// other widgets as well that may need to do work like this.
-				if ( UCanvasPanelSlot* CanvasSlot = Cast<UCanvasPanelSlot>(Slot) )
+				FVector2D LocalPosition = WidgetUnderCursorGeometry.AbsoluteToLocal(DragDropEvent.GetScreenSpacePosition());
+				if ( UPanelSlot* Slot = NewParent->AddChild(Widget) )
 				{
-					if ( bIsPreview )
-					{
-						FWidgetBlueprintEditorUtils::ImportPropertiesFromText(Slot, SelectedDragDropOp->ExportedSlotProperties);
+					FVector2D NewPosition = LocalPosition - SelectedWidgetContextMenuLocation;
 
-						CanvasSlot->SaveBaseLayout();
-						Slot->SetDesiredPosition(NewPosition);
-						CanvasSlot->RebaseLayout();
+					const UWidgetDesignerSettings* DesignerSettings = GetDefault<UWidgetDesignerSettings>();
+					if ( DesignerSettings->GridSnapEnabled )
+					{
+						NewPosition.X = ( (int32)NewPosition.X ) - ( ( (int32)NewPosition.X ) % DesignerSettings->GridSnapSize );
+						NewPosition.Y = ( (int32)NewPosition.Y ) - ( ( (int32)NewPosition.Y ) % DesignerSettings->GridSnapSize );
+					}
 
-						FWidgetBlueprintEditorUtils::ExportPropertiesToText(Slot, SelectedDragDropOp->ExportedSlotProperties);
+					// HACK UMG: In order to correctly drop items into the canvas that have a non-zero anchor,
+					// we need to know the layout information after slate has performed a prepass.  So we have
+					// to rebase the layout and reinterpret the new position based on anchor point layout data.
+					// This should be pulled out into an extension of some kind so that this can be fixed for
+					// other widgets as well that may need to do work like this.
+					if ( UCanvasPanelSlot* CanvasSlot = Cast<UCanvasPanelSlot>(Slot) )
+					{
+						if ( bIsPreview )
+						{
+							FWidgetBlueprintEditorUtils::ImportPropertiesFromText(Slot, SelectedDragDropOp->ExportedSlotProperties);
+
+							CanvasSlot->SaveBaseLayout();
+							Slot->SetDesiredPosition(NewPosition);
+							CanvasSlot->RebaseLayout();
+
+							FWidgetBlueprintEditorUtils::ExportPropertiesToText(Slot, SelectedDragDropOp->ExportedSlotProperties);
+						}
+						else
+						{
+							FWidgetBlueprintEditorUtils::ImportPropertiesFromText(Slot, SelectedDragDropOp->ExportedSlotProperties);
+						}
 					}
 					else
 					{
 						FWidgetBlueprintEditorUtils::ImportPropertiesFromText(Slot, SelectedDragDropOp->ExportedSlotProperties);
+						Slot->SetDesiredPosition(NewPosition);
 					}
+
+					DropPreviewParent = NewParent;
+
+					if ( bIsPreview )
+					{
+						Transaction.Cancel();
+					}
+
+					return Widget;
 				}
 				else
 				{
-					FWidgetBlueprintEditorUtils::ImportPropertiesFromText(Slot, SelectedDragDropOp->ExportedSlotProperties);
-					Slot->SetDesiredPosition(NewPosition);
-				}
+					SelectedDragDropOp->SetCursorOverride(EMouseCursor::SlashedCircle);
 
-				DropPreviewParent = NewParent;
+					// TODO UMG ERROR Slot can not be created because maybe the max children has been reached.
+					//          Maybe we can traverse the hierarchy and add it to the first parent that will accept it?
+				}
 
 				if ( bIsPreview )
 				{
 					Transaction.Cancel();
 				}
-
-				return Widget;
-			}
-			else
-			{
-				SelectedDragDropOp->SetCursorOverride(EMouseCursor::SlashedCircle);
-
-				// TODO UMG ERROR Slot can not be created because maybe the max children has been reached.
-				//          Maybe we can traverse the hierarchy and add it to the first parent that will accept it?
-			}
-
-			if ( bIsPreview )
-			{
-				Transaction.Cancel();
 			}
 		}
 		else
diff --git a/Engine/Source/Editor/UMGEditor/Private/Designer/STransformHandle.cpp b/Engine/Source/Editor/UMGEditor/Private/Designer/STransformHandle.cpp
--- a/Engine/Source/Editor/UMGEditor/Private/Designer/STransformHandle.cpp
+++ b/Engine/Source/Editor/UMGEditor/Private/Designer/STransformHandle.cpp
@@ -137,8 +137,10 @@ FReply STransformHandle::OnMouseMove(const FGeometry& MyGeometry, const FPointer
 				RenderTransform.Angle += RotationDelta.Size();
 			}
 
-			FObjectEditorUtils::SetPropertyValue<UWidget, FWidgetTransform>(Preview, "RenderTransform", RenderTransform);
-			FObjectEditorUtils::SetPropertyValue<UWidget, FWidgetTransform>(Template, "RenderTransform", RenderTransform);
+			static const FName RenderTransformName(TEXT("RenderTransform"));
+
+			FObjectEditorUtils::SetPropertyValue<UWidget, FWidgetTransform>(Preview, RenderTransformName, RenderTransform);
+			FObjectEditorUtils::SetPropertyValue<UWidget, FWidgetTransform>(Template, RenderTransformName, RenderTransform);
 		}
 	}
 
diff --git a/Engine/Source/Editor/UMGEditor/Private/Extensions/CanvasSlotExtension.cpp b/Engine/Source/Editor/UMGEditor/Private/Extensions/CanvasSlotExtension.cpp
--- a/Engine/Source/Editor/UMGEditor/Private/Extensions/CanvasSlotExtension.cpp
+++ b/Engine/Source/Editor/UMGEditor/Private/Extensions/CanvasSlotExtension.cpp
@@ -266,6 +266,17 @@ FReply FCanvasSlotExtension::HandleAnchorEndDrag(const FGeometry& Geometry, cons
 	return FReply::Handled().ReleaseMouseCapture();
 }
 
+void FCanvasSlotExtension::PoximitySnapValue(float SnapFrequency, float SnapProximity, float& Value)
+{
+	float MajorAnchorDiv = Value / SnapFrequency;
+	float SubAnchorLinePos = MajorAnchorDiv - FMath::Round(MajorAnchorDiv);
+
+	if ( FMath::Abs(SubAnchorLinePos) <= SnapProximity )
+	{
+		Value = FMath::Round(MajorAnchorDiv) * SnapFrequency;
+	}
+}
+
 FReply FCanvasSlotExtension::HandleAnchorDragging(const FGeometry& Geometry, const FPointerEvent& Event, EAnchorWidget::Type AnchorType)
 {
 	if ( bMovingAnchor && !Event.GetCursorDelta().IsZero() )
@@ -288,17 +299,19 @@ FReply FCanvasSlotExtension::HandleAnchorDragging(const FGeometry& Geometry, con
 					FVector2D LocalPositionDelta = NewLocalPosition - StartLocalPosition;
 
 					FVector2D AnchorDelta = LocalPositionDelta / CanvasGeometry.Size;
+
+					const FAnchorData OldLayoutData = PreviewCanvasSlot->LayoutData;
+					FAnchorData LayoutData = OldLayoutData;
 					
-					PreviewCanvasSlot->SaveBaseLayout();
 					switch ( AnchorType )
 					{
 					case EAnchorWidget::Center:
-						PreviewCanvasSlot->LayoutData.Anchors.Maximum = BeginAnchors.Maximum + AnchorDelta;
-						PreviewCanvasSlot->LayoutData.Anchors.Minimum = BeginAnchors.Minimum + AnchorDelta;
-						PreviewCanvasSlot->LayoutData.Anchors.Minimum.X = FMath::Clamp(PreviewCanvasSlot->LayoutData.Anchors.Minimum.X, 0.0f, 1.0f);
-						PreviewCanvasSlot->LayoutData.Anchors.Maximum.X = FMath::Clamp(PreviewCanvasSlot->LayoutData.Anchors.Maximum.X, 0.0f, 1.0f);
-						PreviewCanvasSlot->LayoutData.Anchors.Minimum.Y = FMath::Clamp(PreviewCanvasSlot->LayoutData.Anchors.Minimum.Y, 0.0f, 1.0f);
-						PreviewCanvasSlot->LayoutData.Anchors.Maximum.Y = FMath::Clamp(PreviewCanvasSlot->LayoutData.Anchors.Maximum.Y, 0.0f, 1.0f);
+						LayoutData.Anchors.Maximum = BeginAnchors.Maximum + AnchorDelta;
+						LayoutData.Anchors.Minimum = BeginAnchors.Minimum + AnchorDelta;
+						LayoutData.Anchors.Minimum.X = FMath::Clamp(LayoutData.Anchors.Minimum.X, 0.0f, 1.0f);
+						LayoutData.Anchors.Maximum.X = FMath::Clamp(LayoutData.Anchors.Maximum.X, 0.0f, 1.0f);
+						LayoutData.Anchors.Minimum.Y = FMath::Clamp(LayoutData.Anchors.Minimum.Y, 0.0f, 1.0f);
+						LayoutData.Anchors.Maximum.Y = FMath::Clamp(LayoutData.Anchors.Maximum.Y, 0.0f, 1.0f);
 						break;
 					}
 
@@ -307,8 +320,8 @@ FReply FCanvasSlotExtension::HandleAnchorDragging(const FGeometry& Geometry, con
 					case EAnchorWidget::Left:
 					case EAnchorWidget::TopLeft:
 					case EAnchorWidget::BottomLeft:
-						PreviewCanvasSlot->LayoutData.Anchors.Minimum.X = BeginAnchors.Minimum.X + AnchorDelta.X;
-						PreviewCanvasSlot->LayoutData.Anchors.Minimum.X = FMath::Clamp(PreviewCanvasSlot->LayoutData.Anchors.Minimum.X, 0.0f, PreviewCanvasSlot->LayoutData.Anchors.Maximum.X);
+						LayoutData.Anchors.Minimum.X = BeginAnchors.Minimum.X + AnchorDelta.X;
+						LayoutData.Anchors.Minimum.X = FMath::Clamp(LayoutData.Anchors.Minimum.X, 0.0f, LayoutData.Anchors.Maximum.X);
 						break;
 					}
 
@@ -317,8 +330,8 @@ FReply FCanvasSlotExtension::HandleAnchorDragging(const FGeometry& Geometry, con
 					case EAnchorWidget::Right:
 					case EAnchorWidget::TopRight:
 					case EAnchorWidget::BottomRight:
-						PreviewCanvasSlot->LayoutData.Anchors.Maximum.X = BeginAnchors.Maximum.X + AnchorDelta.X;
-						PreviewCanvasSlot->LayoutData.Anchors.Maximum.X = FMath::Clamp(PreviewCanvasSlot->LayoutData.Anchors.Maximum.X, PreviewCanvasSlot->LayoutData.Anchors.Minimum.X, 1.0f);
+						LayoutData.Anchors.Maximum.X = BeginAnchors.Maximum.X + AnchorDelta.X;
+						LayoutData.Anchors.Maximum.X = FMath::Clamp(LayoutData.Anchors.Maximum.X, LayoutData.Anchors.Minimum.X, 1.0f);
 						break;
 					}
 
@@ -327,8 +340,8 @@ FReply FCanvasSlotExtension::HandleAnchorDragging(const FGeometry& Geometry, con
 					case EAnchorWidget::Top:
 					case EAnchorWidget::TopLeft:
 					case EAnchorWidget::TopRight:
-						PreviewCanvasSlot->LayoutData.Anchors.Minimum.Y = BeginAnchors.Minimum.Y + AnchorDelta.Y;
-						PreviewCanvasSlot->LayoutData.Anchors.Minimum.Y = FMath::Clamp(PreviewCanvasSlot->LayoutData.Anchors.Minimum.Y, 0.0f, PreviewCanvasSlot->LayoutData.Anchors.Maximum.Y);
+						LayoutData.Anchors.Minimum.Y = BeginAnchors.Minimum.Y + AnchorDelta.Y;
+						LayoutData.Anchors.Minimum.Y = FMath::Clamp(LayoutData.Anchors.Minimum.Y, 0.0f, LayoutData.Anchors.Maximum.Y);
 						break;
 					}
 
@@ -337,19 +350,62 @@ FReply FCanvasSlotExtension::HandleAnchorDragging(const FGeometry& Geometry, con
 					case EAnchorWidget::Bottom:
 					case EAnchorWidget::BottomLeft:
 					case EAnchorWidget::BottomRight:
-						PreviewCanvasSlot->LayoutData.Anchors.Maximum.Y = BeginAnchors.Maximum.Y + AnchorDelta.Y;
-						PreviewCanvasSlot->LayoutData.Anchors.Maximum.Y = FMath::Clamp(PreviewCanvasSlot->LayoutData.Anchors.Maximum.Y, PreviewCanvasSlot->LayoutData.Anchors.Minimum.Y, 1.0f);
+						LayoutData.Anchors.Maximum.Y = BeginAnchors.Maximum.Y + AnchorDelta.Y;
+						LayoutData.Anchors.Maximum.Y = FMath::Clamp(LayoutData.Anchors.Maximum.Y, LayoutData.Anchors.Minimum.Y, 1.0f);
 						break;
 					}
 
-					PreviewCanvasSlot->RebaseLayout();
+					// Major percentage snapping
+					{
+						const float MajorAnchorLine = 0.1f;
+						const float MajorAnchorLineSnapDistance = 0.1f;
+
+						if ( LayoutData.Anchors.Minimum.X != OldLayoutData.Anchors.Minimum.X )
+						{
+							PoximitySnapValue(MajorAnchorLine, MajorAnchorLineSnapDistance, LayoutData.Anchors.Minimum.X);
+						}
+
+						if ( LayoutData.Anchors.Minimum.Y != OldLayoutData.Anchors.Minimum.Y )
+						{
+							PoximitySnapValue(MajorAnchorLine, MajorAnchorLineSnapDistance, LayoutData.Anchors.Minimum.Y);
+						}
+
+						if ( LayoutData.Anchors.Maximum.X != OldLayoutData.Anchors.Maximum.X )
+						{
+							PoximitySnapValue(MajorAnchorLine, MajorAnchorLineSnapDistance, LayoutData.Anchors.Maximum.X);
+						}
+
+						if ( LayoutData.Anchors.Maximum.Y != OldLayoutData.Anchors.Maximum.Y )
+						{
+							PoximitySnapValue(MajorAnchorLine, MajorAnchorLineSnapDistance, LayoutData.Anchors.Maximum.Y);
+						}
+					}
+
+					// Rebase the layout and restore the old value after calculating the new final layout
+					// result.
+					{
+						PreviewCanvasSlot->SaveBaseLayout();
+						PreviewCanvasSlot->LayoutData = LayoutData;
+						PreviewCanvasSlot->RebaseLayout();
 
+						LayoutData = PreviewCanvasSlot->LayoutData;
+						PreviewCanvasSlot->LayoutData = OldLayoutData;
+					}
+
+					// If control is pressed reset all positional offset information
+					if ( FSlateApplication::Get().GetModifierKeys().IsControlDown() )
+					{
+						FMargin NewOffsets = FMargin(0, 0, LayoutData.Anchors.IsStretchedHorizontal() ? 0 : LayoutData.Offsets.Right, LayoutData.Anchors.IsStretchedVertical() ? 0 : LayoutData.Offsets.Bottom);
+						LayoutData.Offsets = NewOffsets;
+					}
 
-					// Update the Template widget to match
 					UWidget* TemplateWidget = Selection.GetTemplate();
-					UCanvasPanelSlot* TemplateSlot = CastChecked<UCanvasPanelSlot>(TemplateWidget->Slot);
+					UCanvasPanelSlot* TemplateCanvasSlot = CastChecked<UCanvasPanelSlot>(TemplateWidget->Slot);
+
+					static const FName LayoutDataName(TEXT("LayoutData"));
 
-					TemplateSlot->LayoutData = PreviewCanvasSlot->LayoutData;
+					FObjectEditorUtils::SetPropertyValue<UCanvasPanelSlot, FAnchorData>(PreviewCanvasSlot, LayoutDataName, LayoutData);
+					FObjectEditorUtils::SetPropertyValue<UCanvasPanelSlot, FAnchorData>(TemplateCanvasSlot, LayoutDataName, LayoutData);
 				}
 			};
 
diff --git a/Engine/Source/Editor/UMGEditor/Private/WidgetBlueprintEditor.cpp b/Engine/Source/Editor/UMGEditor/Private/WidgetBlueprintEditor.cpp
--- a/Engine/Source/Editor/UMGEditor/Private/WidgetBlueprintEditor.cpp
+++ b/Engine/Source/Editor/UMGEditor/Private/WidgetBlueprintEditor.cpp
@@ -58,7 +58,7 @@ void FWidgetBlueprintEditor::InitWidgetBlueprintEditor(const EToolkitMode::Type
 	if ( Blueprint->WidgetTree->RootWidget == NULL )
 	{
 		UWidget* RootWidget = Blueprint->WidgetTree->ConstructWidget<UCanvasPanel>(UCanvasPanel::StaticClass());
-		RootWidget->IsDesignTime(true);
+		RootWidget->SetIsDesignTime(true);
 		Blueprint->WidgetTree->RootWidget = RootWidget;
 	}
 
@@ -630,11 +630,7 @@ void FWidgetBlueprintEditor::UpdatePreview(UBlueprint* InBlueprint, bool bInForc
 		PreviewActor->SetFlags(RF_Transactional);
 		
 		// Configure all the widgets to be set to design time.
-		PreviewActor->IsDesignTime(true);
-		for(UWidget* SubPreviewWidget : PreviewActor->Components)
-		{
-			SubPreviewWidget->IsDesignTime(true);
-		}
+		PreviewActor->SetIsDesignTime(true);
 
 		// Store a reference to the preview actor.
 		PreviewWidgetPtr = PreviewActor;
diff --git a/Engine/Source/Editor/UnrealEd/Private/SViewportToolBarIconMenu.cpp b/Engine/Source/Editor/UnrealEd/Private/SViewportToolBarIconMenu.cpp
--- a/Engine/Source/Editor/UnrealEd/Private/SViewportToolBarIconMenu.cpp
+++ b/Engine/Source/Editor/UnrealEd/Private/SViewportToolBarIconMenu.cpp
@@ -1,6 +1,7 @@
 // Copyright 1998-2014 Epic Games, Inc. All Rights Reserved.
 
 #include "UnrealEd.h"
+#include "MultiBox.h"
 #include "SViewportToolBarIconMenu.h"
 
 void SViewportToolBarIconMenu::Construct( const FArguments& InArgs )
diff --git a/Engine/Source/Runtime/UMG/Private/Components/CanvasPanelSlot.cpp b/Engine/Source/Runtime/UMG/Private/Components/CanvasPanelSlot.cpp
--- a/Engine/Source/Runtime/UMG/Private/Components/CanvasPanelSlot.cpp
+++ b/Engine/Source/Runtime/UMG/Private/Components/CanvasPanelSlot.cpp
@@ -12,6 +12,7 @@ UCanvasPanelSlot::UCanvasPanelSlot(const FPostConstructInitializeProperties& PCI
 	LayoutData.Offsets = FMargin(0, 0, 100, 30);
 	LayoutData.Anchors = FAnchors(0.0f, 0.0f);
 	LayoutData.Alignment = FVector2D(0.0f, 0.0f);
+	bAutoSize = false;
 	ZOrder = 0;
 }
 
@@ -172,6 +173,15 @@ void UCanvasPanelSlot::SetAlignment(FVector2D InAlignment)
 	}
 }
 
+void UCanvasPanelSlot::SetAutoSize(bool InbAutoSize)
+{
+	bAutoSize = InbAutoSize;
+	if ( Slot )
+	{
+		Slot->AutoSize(InbAutoSize);
+	}
+}
+
 void UCanvasPanelSlot::SetZOrder(int32 InZOrder)
 {
 	ZOrder = InZOrder;
@@ -186,12 +196,13 @@ void UCanvasPanelSlot::SynchronizeProperties()
 	SetOffsets(LayoutData.Offsets);
 	SetAnchors(LayoutData.Anchors);
 	SetAlignment(LayoutData.Alignment);
+	SetAutoSize(bAutoSize);
 	SetZOrder(ZOrder);
 }
 
 #if WITH_EDITOR
 
-void UCanvasPanelSlot::PreEditChange(UProperty* PropertyAboutToChange)
+void UCanvasPanelSlot::PreEditChange(class FEditPropertyChain& PropertyAboutToChange)
 {
 	Super::PreEditChange(PropertyAboutToChange);
 
diff --git a/Engine/Source/Runtime/UMG/Private/Components/GridPanel.cpp b/Engine/Source/Runtime/UMG/Private/Components/GridPanel.cpp
--- a/Engine/Source/Runtime/UMG/Private/Components/GridPanel.cpp
+++ b/Engine/Source/Runtime/UMG/Private/Components/GridPanel.cpp
@@ -66,6 +66,11 @@ TSharedRef<SWidget> UGridPanel::RebuildWidget()
 	return BuildDesignTimeWidget( MyGridPanel.ToSharedRef() );
 }
 
+UGridSlot* UGridPanel::AddChildToGrid(UWidget* Content)
+{
+	return Cast<UGridSlot>(Super::AddChild(Content));
+}
+
 void UGridPanel::SynchronizeProperties()
 {
 	Super::SynchronizeProperties();
diff --git a/Engine/Source/Runtime/UMG/Private/Components/HorizontalBox.cpp b/Engine/Source/Runtime/UMG/Private/Components/HorizontalBox.cpp
--- a/Engine/Source/Runtime/UMG/Private/Components/HorizontalBox.cpp
+++ b/Engine/Source/Runtime/UMG/Private/Components/HorizontalBox.cpp
@@ -50,16 +50,11 @@ void UHorizontalBox::OnSlotRemoved(UPanelSlot* Slot)
 	}
 }
 
-UHorizontalBoxSlot* UHorizontalBox::AddSlot(UWidget* Content)
+UHorizontalBoxSlot* UHorizontalBox::AddChildToHorizontalBox(UWidget* Content)
 {
 	return Cast<UHorizontalBoxSlot>( Super::AddChild(Content) );
 }
 
-UHorizontalBoxSlot* UHorizontalBox::HorizontalBoxSlot(UWidget* Child)
-{
-	return Cast<UHorizontalBoxSlot>(Child->Slot);
-}
-
 TSharedRef<SWidget> UHorizontalBox::RebuildWidget()
 {
 	MyHorizontalBox = SNew(SHorizontalBox);
diff --git a/Engine/Source/Runtime/UMG/Private/Components/Overlay.cpp b/Engine/Source/Runtime/UMG/Private/Components/Overlay.cpp
--- a/Engine/Source/Runtime/UMG/Private/Components/Overlay.cpp
+++ b/Engine/Source/Runtime/UMG/Private/Components/Overlay.cpp
@@ -23,6 +23,11 @@ void UOverlay::ReleaseNativeWidget()
 	MyOverlay.Reset();
 }
 
+UOverlaySlot* UOverlay::AddChildToOverlay(UWidget* Content)
+{
+	return Cast<UOverlaySlot>(Super::AddChild(Content));
+}
+
 UClass* UOverlay::GetSlotClass() const
 {
 	return UOverlaySlot::StaticClass();
diff --git a/Engine/Source/Runtime/UMG/Private/Components/PanelWidget.cpp b/Engine/Source/Runtime/UMG/Private/Components/PanelWidget.cpp
--- a/Engine/Source/Runtime/UMG/Private/Components/PanelWidget.cpp
+++ b/Engine/Source/Runtime/UMG/Private/Components/PanelWidget.cpp
@@ -160,6 +160,21 @@ void UPanelWidget::ClearChildren()
 	}
 }
 
+void UPanelWidget::SetIsDesignTime(bool bInDesignTime)
+{
+	Super::SetIsDesignTime(bInDesignTime);
+
+	// Also mark all children as design time widgets.
+	int32 Children = GetChildrenCount();
+	for ( int32 SlotIndex = 0; SlotIndex < Children; SlotIndex++ )
+	{
+		if ( Slots[SlotIndex]->Content != NULL )
+		{
+			Slots[SlotIndex]->Content->SetIsDesignTime(bInDesignTime);
+		}
+	}
+}
+
 void UPanelWidget::PostLoad()
 {
 	Super::PostLoad();
diff --git a/Engine/Source/Runtime/UMG/Private/Components/UniformGridPanel.cpp b/Engine/Source/Runtime/UMG/Private/Components/UniformGridPanel.cpp
--- a/Engine/Source/Runtime/UMG/Private/Components/UniformGridPanel.cpp
+++ b/Engine/Source/Runtime/UMG/Private/Components/UniformGridPanel.cpp
@@ -69,6 +69,11 @@ TSharedRef<SWidget> UUniformGridPanel::RebuildWidget()
 	return BuildDesignTimeWidget( MyUniformGridPanel.ToSharedRef() );
 }
 
+UUniformGridSlot* UUniformGridPanel::AddChildToUniformGrid(UWidget* Content)
+{
+	return Cast<UUniformGridSlot>(Super::AddChild(Content));
+}
+
 void UUniformGridPanel::SynchronizeProperties()
 {
 	Super::SynchronizeProperties();
diff --git a/Engine/Source/Runtime/UMG/Private/Components/VerticalBox.cpp b/Engine/Source/Runtime/UMG/Private/Components/VerticalBox.cpp
--- a/Engine/Source/Runtime/UMG/Private/Components/VerticalBox.cpp
+++ b/Engine/Source/Runtime/UMG/Private/Components/VerticalBox.cpp
@@ -50,7 +50,7 @@ void UVerticalBox::OnSlotRemoved(UPanelSlot* Slot)
 	}
 }
 
-UVerticalBoxSlot* UVerticalBox::AddSlot(UWidget* Content)
+UVerticalBoxSlot* UVerticalBox::AddChildVerticalBox(UWidget* Content)
 {
 	return Cast<UVerticalBoxSlot>(Super::AddChild(Content));
 }
diff --git a/Engine/Source/Runtime/UMG/Private/Components/Widget.cpp b/Engine/Source/Runtime/UMG/Private/Components/Widget.cpp
--- a/Engine/Source/Runtime/UMG/Private/Components/Widget.cpp
+++ b/Engine/Source/Runtime/UMG/Private/Components/Widget.cpp
@@ -167,6 +167,16 @@ void UWidget::ForceLayoutPrepass()
 	}
 }
 
+FVector2D UWidget::GetDesiredSize() const
+{
+	if ( MyWidget.IsValid() )
+	{
+		return MyWidget.Pin()->GetDesiredSize();
+	}
+
+	return FVector2D(0, 0);
+}
+
 UPanelWidget* UWidget::GetParent() const
 {
 	if ( Slot )
@@ -174,7 +184,7 @@ UPanelWidget* UWidget::GetParent() const
 		return Slot->Parent;
 	}
 
-	return NULL;
+	return nullptr;
 }
 
 void UWidget::RemoveFromParent()
@@ -335,7 +345,7 @@ void UWidget::Select()
 	OnSelected();
 
 	UWidget* Parent = GetParent();
-	while ( Parent != NULL )
+	while ( Parent != nullptr )
 	{
 		Parent->OnDescendantSelected(this);
 		Parent = Parent->GetParent();
@@ -347,7 +357,7 @@ void UWidget::Deselect()
 	OnDeselected();
 
 	UWidget* Parent = GetParent();
-	while ( Parent != NULL )
+	while ( Parent != nullptr )
 	{
 		Parent->OnDescendantDeselected(this);
 		Parent = Parent->GetParent();
@@ -362,7 +372,9 @@ bool UWidget::Modify(bool bAlwaysMarkDirty)
 	bool Modified = Super::Modify(bAlwaysMarkDirty);
 
 	if ( Slot )
+	{
 		Modified &= Slot->Modify(bAlwaysMarkDirty);
+	}
 
 	return Modified;
 }
@@ -370,7 +382,7 @@ bool UWidget::Modify(bool bAlwaysMarkDirty)
 bool UWidget::IsChildOf(UWidget* PossibleParent)
 {
 	UPanelWidget* Parent = GetParent();
-	if ( Parent == NULL )
+	if ( Parent == nullptr )
 	{
 		return false;
 	}
@@ -421,12 +433,12 @@ bool UWidget::IsDesignTime() const
 	return bDesignTime;
 }
 
-void UWidget::IsDesignTime(bool bInDesignTime)
+void UWidget::SetIsDesignTime(bool bInDesignTime)
 {
 	bDesignTime = bInDesignTime;
 }
 
-EVisibility UWidget::ConvertSerializedVisibilityToRuntime(TEnumAsByte<ESlateVisibility::Type> Input)
+EVisibility UWidget::ConvertSerializedVisibilityToRuntime(ESlateVisibility::Type Input)
 {
 	switch ( Input )
 	{
@@ -441,12 +453,12 @@ EVisibility UWidget::ConvertSerializedVisibilityToRuntime(TEnumAsByte<ESlateVisi
 	case ESlateVisibility::SelfHitTestInvisible:
 		return EVisibility::SelfHitTestInvisible;
 	default:
-		check(false);
+		//check(false);
 		return EVisibility::Visible;
 	}
 }
 
-TEnumAsByte<ESlateVisibility::Type> UWidget::ConvertRuntimeToSerializedVisiblity(const EVisibility& Input)
+ESlateVisibility::Type UWidget::ConvertRuntimeToSerializedVisiblity(const EVisibility& Input)
 {
 	if ( Input == EVisibility::Visible )
 	{
@@ -470,7 +482,7 @@ TEnumAsByte<ESlateVisibility::Type> UWidget::ConvertRuntimeToSerializedVisiblity
 	}
 	else
 	{
-		check(false);
+		//check(false);
 		return ESlateVisibility::Visible;
 	}
 }
@@ -485,6 +497,8 @@ FSizeParam UWidget::ConvertSerializedSizeParamToRuntime(const FSlateChildSize& I
 	case ESlateSizeRule::Fill:
 		return FStretch(Input.Value);
 	}
+
+	return FAuto();
 }
 
 void UWidget::GatherChildren(UWidget* Root, TSet<UWidget*>& Children)
@@ -519,7 +533,7 @@ UWidget* UWidget::FindChildContainingDescendant(UWidget* Root, UWidget* Descenda
 {
 	UWidget* Parent = Descendant->GetParent();
 
-	while ( Parent != NULL )
+	while ( Parent != nullptr )
 	{
 		// If the Descendant's parent is the root, then the child containing the descendant is the descendant.
 		if ( Parent == Root )
@@ -531,5 +545,5 @@ UWidget* UWidget::FindChildContainingDescendant(UWidget* Root, UWidget* Descenda
 		Parent = Parent->GetParent();
 	}
 
-	return NULL;
+	return nullptr;
 }
diff --git a/Engine/Source/Runtime/UMG/Private/Slate/SConstraintCanvas.cpp b/Engine/Source/Runtime/UMG/Private/Slate/SConstraintCanvas.cpp
--- a/Engine/Source/Runtime/UMG/Private/Slate/SConstraintCanvas.cpp
+++ b/Engine/Source/Runtime/UMG/Private/Slate/SConstraintCanvas.cpp
@@ -90,6 +90,8 @@ void SConstraintCanvas::OnArrangeChildren( const FGeometry& AllottedGeometry, FA
 			const FVector2D Alignment = CurChild.AlignmentAttr.Get();
 			const FAnchors Anchors = CurChild.AnchorsAttr.Get();
 
+			const bool AutoSize = CurChild.AutoSizeAttr.Get();
+
 			const FMargin AnchorPixels =
 				FMargin(Anchors.Minimum.X * AllottedGeometry.Size.X,
 						Anchors.Minimum.Y * AllottedGeometry.Size.Y,
@@ -99,7 +101,10 @@ void SConstraintCanvas::OnArrangeChildren( const FGeometry& AllottedGeometry, FA
 			const bool bIsHorizontalStretch = Anchors.Minimum.X != Anchors.Maximum.X;
 			const bool bIsVerticalStretch = Anchors.Minimum.Y != Anchors.Maximum.Y;
 			
-			const FVector2D Size = FVector2D(Offset.Right, Offset.Bottom);
+			const FVector2D SlotSize = FVector2D(Offset.Right, Offset.Bottom);
+			const FVector2D WidgetDesiredSize = CurChild.GetWidget()->GetDesiredSize();
+
+			const FVector2D Size = AutoSize ? WidgetDesiredSize : SlotSize;
 			
 			// Calculate the offset based on the pivot position.
 			FVector2D AlignmentOffset = Size * Alignment;
diff --git a/Engine/Source/Runtime/UMG/Private/UserWidget.cpp b/Engine/Source/Runtime/UMG/Private/UserWidget.cpp
--- a/Engine/Source/Runtime/UMG/Private/UserWidget.cpp
+++ b/Engine/Source/Runtime/UMG/Private/UserWidget.cpp
@@ -201,6 +201,16 @@ UWorld* UUserWidget::GetWorld() const
 	return nullptr;
 }
 
+void UUserWidget::SetIsDesignTime(bool bInDesignTime)
+{
+	Super::SetIsDesignTime(bInDesignTime);
+
+	for ( UWidget* Widget : Components )
+	{
+		Widget->SetIsDesignTime(bInDesignTime);
+	}
+}
+
 void UUserWidget::Construct_Implementation()
 {
 
@@ -456,6 +466,8 @@ TSharedPtr<SWidget> UUserWidget::GetWidgetFromName(const FString& Name) const
 
 UWidget* UUserWidget::GetHandleFromName(const FString& Name) const
 {
+	//return WidgetTree->FindWidget(Name);
+
 	for ( UWidget* Widget : Components )
 	{
 		if ( Widget->GetName().Equals(Name, ESearchCase::IgnoreCase) )
@@ -464,7 +476,7 @@ UWidget* UUserWidget::GetHandleFromName(const FString& Name) const
 		}
 	}
 
-	return NULL;
+	return nullptr;
 }
 
 void UUserWidget::NativeTick(const FGeometry& MyGeometry, float InDeltaTime )
diff --git a/Engine/Source/Runtime/UMG/Public/Blueprint/UserWidget.h b/Engine/Source/Runtime/UMG/Public/Blueprint/UserWidget.h
--- a/Engine/Source/Runtime/UMG/Public/Blueprint/UserWidget.h
+++ b/Engine/Source/Runtime/UMG/Public/Blueprint/UserWidget.h
@@ -97,6 +97,9 @@ public:
 	virtual void ReleaseNativeWidget() override;
 	// End of UVisual interface
 
+	/** Sets that this widget is being designed sets it on all children as well. */
+	virtual void SetIsDesignTime(bool bInDesignTime) override;
+
 	/**
 	 * Adds it to the game's viewport, defaults to filling the entire viewport area.
 	 * @param bModal If this dialog should steal keyboard/mouse focus and consume all input. Great for a fullscreen menu. Terrible for HUDs.
diff --git a/Engine/Source/Runtime/UMG/Public/Components/CanvasPanelSlot.h b/Engine/Source/Runtime/UMG/Public/Components/CanvasPanelSlot.h
--- a/Engine/Source/Runtime/UMG/Public/Components/CanvasPanelSlot.h
+++ b/Engine/Source/Runtime/UMG/Public/Components/CanvasPanelSlot.h
@@ -36,6 +36,10 @@ class UMG_API UCanvasPanelSlot : public UPanelSlot
 	UPROPERTY(EditDefaultsOnly, Category="Layout (Canvas Slot)")
 	FAnchorData LayoutData;
 
+	/** When AutoSize is true we use the widget's desired size */
+	UPROPERTY(EditDefaultsOnly, Category="Layout (Canvas Slot)", AdvancedDisplay)
+	bool bAutoSize;
+
 	/** The order priority this widget is rendered in.  Higher values are rendered last (and so they will appear to be on top). */
 	UPROPERTY(EditDefaultsOnly, Category="Layout (Canvas Slot)")
 	int32 ZOrder;
@@ -70,6 +74,10 @@ class UMG_API UCanvasPanelSlot : public UPanelSlot
 	UFUNCTION(BlueprintCallable, Category="Layout (Canvas Slot)")
 	void SetAlignment(FVector2D InAlignment);
 
+	/** Sets the slot to be auto-sized */
+	UFUNCTION(BlueprintCallable, Category="Layout (Canvas Slot)")
+	void SetAutoSize(bool InbAutoSize);
+
 	/** Sets the z-order on the slot */
 	UFUNCTION(BlueprintCallable, Category="Layout (Canvas Slot)")
 	void SetZOrder(int32 InZOrder);
@@ -82,7 +90,7 @@ class UMG_API UCanvasPanelSlot : public UPanelSlot
 
 #if WITH_EDITOR
 	// UObject interface
-	virtual void PreEditChange(UProperty* PropertyAboutToChange) override;
+	virtual void PreEditChange(class FEditPropertyChain& PropertyAboutToChange);
 	virtual void PostEditChangeChainProperty(struct FPropertyChangedChainEvent& PropertyChangedEvent);
 	// End of UObject interface
 
diff --git a/Engine/Source/Runtime/UMG/Public/Components/GridPanel.h b/Engine/Source/Runtime/UMG/Public/Components/GridPanel.h
--- a/Engine/Source/Runtime/UMG/Public/Components/GridPanel.h
+++ b/Engine/Source/Runtime/UMG/Public/Components/GridPanel.h
@@ -20,6 +20,12 @@ public:
 	UPROPERTY(EditDefaultsOnly, Category="Fill Rules")
 	TArray<float> RowFill;
 
+public:
+
+	/**  */
+	UFUNCTION(BlueprintCallable, Category="Widget")
+	UGridSlot* AddChildToGrid(UWidget* Content);
+
 public:
 
 	// UWidget interface
diff --git a/Engine/Source/Runtime/UMG/Public/Components/HorizontalBox.h b/Engine/Source/Runtime/UMG/Public/Components/HorizontalBox.h
--- a/Engine/Source/Runtime/UMG/Public/Components/HorizontalBox.h
+++ b/Engine/Source/Runtime/UMG/Public/Components/HorizontalBox.h
@@ -11,11 +11,7 @@ class UMG_API UHorizontalBox : public UPanelWidget
 
 	/**  */
 	UFUNCTION(BlueprintCallable, Category="Widget")
-	UHorizontalBoxSlot* AddSlot(UWidget* Content);
-
-	/**  */
-	UFUNCTION(BlueprintCallable, Category="Slot")
-	static UHorizontalBoxSlot* HorizontalBoxSlot(UWidget* Child);
+	UHorizontalBoxSlot* AddChildToHorizontalBox(UWidget* Content);
 
 #if WITH_EDITOR
 	// UWidget interface
diff --git a/Engine/Source/Runtime/UMG/Public/Components/Overlay.h b/Engine/Source/Runtime/UMG/Public/Components/Overlay.h
--- a/Engine/Source/Runtime/UMG/Public/Components/Overlay.h
+++ b/Engine/Source/Runtime/UMG/Public/Components/Overlay.h
@@ -11,6 +11,10 @@ class UMG_API UOverlay : public UPanelWidget
 
 	virtual void ReleaseNativeWidget() override;
 
+	/**  */
+	UFUNCTION(BlueprintCallable, Category="Widget")
+	UOverlaySlot* AddChildToOverlay(UWidget* Content);
+
 #if WITH_EDITOR
 	virtual const FSlateBrush* GetEditorIcon() override;
 	virtual const FText GetPaletteCategory() override;
diff --git a/Engine/Source/Runtime/UMG/Public/Components/PanelWidget.h b/Engine/Source/Runtime/UMG/Public/Components/PanelWidget.h
--- a/Engine/Source/Runtime/UMG/Public/Components/PanelWidget.h
+++ b/Engine/Source/Runtime/UMG/Public/Components/PanelWidget.h
@@ -67,6 +67,9 @@ public:
 		return bCanHaveMultipleChildren;
 	}
 
+	/** Sets that this widget is being designed sets it on all children as well. */
+	virtual void SetIsDesignTime(bool bInDesignTime) override;
+
 	virtual void ReleaseNativeWidget() override;
 
 #if WITH_EDITOR
diff --git a/Engine/Source/Runtime/UMG/Public/Components/UniformGridPanel.h b/Engine/Source/Runtime/UMG/Public/Components/UniformGridPanel.h
--- a/Engine/Source/Runtime/UMG/Public/Components/UniformGridPanel.h
+++ b/Engine/Source/Runtime/UMG/Public/Components/UniformGridPanel.h
@@ -24,6 +24,12 @@ public:
 	UPROPERTY(EditDefaultsOnly, Category="Child Layout")
 	float MinDesiredSlotHeight;
 
+public:
+
+	/**  */
+	UFUNCTION(BlueprintCallable, Category="Widget")
+	UUniformGridSlot* AddChildToUniformGrid(UWidget* Content);
+
 public:
 
 	// UWidget interface
diff --git a/Engine/Source/Runtime/UMG/Public/Components/VerticalBox.h b/Engine/Source/Runtime/UMG/Public/Components/VerticalBox.h
--- a/Engine/Source/Runtime/UMG/Public/Components/VerticalBox.h
+++ b/Engine/Source/Runtime/UMG/Public/Components/VerticalBox.h
@@ -17,7 +17,7 @@ public:
 
 	/**  */
 	UFUNCTION(BlueprintCallable, Category="Panel")
-	UVerticalBoxSlot* AddSlot(UWidget* Content);
+	UVerticalBoxSlot* AddChildVerticalBox(UWidget* Content);
 
 #if WITH_EDITOR
 	// UWidget interface
diff --git a/Engine/Source/Runtime/UMG/Public/Components/Widget.h b/Engine/Source/Runtime/UMG/Public/Components/Widget.h
--- a/Engine/Source/Runtime/UMG/Public/Components/Widget.h
+++ b/Engine/Source/Runtime/UMG/Public/Components/Widget.h
@@ -192,6 +192,16 @@ public:
 	UFUNCTION(BlueprintCallable, Category="Widget")
 	void ForceLayoutPrepass();
 
+	/**
+	 * Gets the widgets desired size.
+	 * NOTE: The underlying Slate widget must exist and be valid, also a at least one pre-pass must
+	 *       have occurred before this value will be of any use.
+	 * 
+	 * @return The widget's desired size
+	 */
+	UFUNCTION(BlueprintCallable, Category="Widget")
+	FVector2D GetDesiredSize() const;
+
 	/** Gets the parent widget */
 	UFUNCTION(BlueprintCallable, Category="Widget")
 	class UPanelWidget* GetParent() const;
@@ -221,7 +231,7 @@ public:
 	bool IsDesignTime() const;
 	
 	/** Sets that this widget is being designed */
-	void IsDesignTime(bool bInDesignTime);
+	virtual void SetIsDesignTime(bool bInDesignTime);
 
 	/** Mark this object as modified, also mark the slot as modified. */
 	virtual bool Modify(bool bAlwaysMarkDirty = true);
@@ -272,8 +282,8 @@ public:
 
 	// Utility methods
 	//@TODO UMG: Should move elsewhere
-	static EVisibility ConvertSerializedVisibilityToRuntime(TEnumAsByte<ESlateVisibility::Type> Input);
-	static TEnumAsByte<ESlateVisibility::Type> ConvertRuntimeToSerializedVisiblity(const EVisibility& Input);
+	static EVisibility ConvertSerializedVisibilityToRuntime(ESlateVisibility::Type Input);
+	static ESlateVisibility::Type ConvertRuntimeToSerializedVisiblity(const EVisibility& Input);
 
 	static FSizeParam ConvertSerializedSizeParamToRuntime(const FSlateChildSize& Input);
 
@@ -294,7 +304,8 @@ protected:
 	// Conversion functions
 	EVisibility ConvertVisibility(TAttribute<ESlateVisibility::Type> SerializedType) const
 	{
-		return ConvertSerializedVisibilityToRuntime(SerializedType.Get());
+		ESlateVisibility::Type SlateVisibility = SerializedType.Get();
+		return ConvertSerializedVisibilityToRuntime(SlateVisibility);
 	}
 
 	TOptional<float> ConvertFloatToOptionalFloat(TAttribute<float> InFloat) const
diff --git a/Engine/Source/Runtime/UMG/Public/Slate/SConstraintCanvas.h b/Engine/Source/Runtime/UMG/Public/Slate/SConstraintCanvas.h
--- a/Engine/Source/Runtime/UMG/Public/Slate/SConstraintCanvas.h
+++ b/Engine/Source/Runtime/UMG/Public/Slate/SConstraintCanvas.h
@@ -35,6 +35,12 @@ public:
 			return *this;
 		}
 
+		FSlot& AutoSize(const TAttribute<bool>& InAutoSize)
+		{
+			AutoSizeAttr = InAutoSize;
+			return *this;
+		}
+
 		FSlot& ZOrder(const TAttribute<int32>& InZOrder)
 		{
 			ZOrderAttr = InZOrder;
@@ -56,7 +62,10 @@ public:
 		/** Size */
 		TAttribute<FVector2D> AlignmentAttr;
 
-		/** Size */
+		/** Auto-Size */
+		TAttribute<bool> AutoSizeAttr;
+
+		/** Z-Order */
 		TAttribute<int32> ZOrderAttr;
 
 		/** Default values for a slot. */
@@ -65,6 +74,7 @@ public:
 			, OffsetAttr( FMargin( 0, 0, 1, 1 ) )
 			, AnchorsAttr( FAnchors( 0.0f, 0.0f ) )
 			, AlignmentAttr( FVector2D( 0.5f, 0.5f ) )
+			, AutoSizeAttr( false )
 			, ZOrderAttr( 0 )
 		{ }
 	};
