UMG - Fixing some compiler warnings about using Round instead of RoundToInt.

diff --git a/Engine/Source/Editor/UMGEditor/Private/Extensions/CanvasSlotExtension.cpp b/Engine/Source/Editor/UMGEditor/Private/Extensions/CanvasSlotExtension.cpp
--- a/Engine/Source/Editor/UMGEditor/Private/Extensions/CanvasSlotExtension.cpp
+++ b/Engine/Source/Editor/UMGEditor/Private/Extensions/CanvasSlotExtension.cpp
@@ -269,11 +269,11 @@ FReply FCanvasSlotExtension::HandleAnchorEndDrag(const FGeometry& Geometry, cons
 void FCanvasSlotExtension::PoximitySnapValue(float SnapFrequency, float SnapProximity, float& Value)
 {
 	float MajorAnchorDiv = Value / SnapFrequency;
-	float SubAnchorLinePos = MajorAnchorDiv - FMath::Round(MajorAnchorDiv);
+	float SubAnchorLinePos = MajorAnchorDiv - FMath::RoundToInt(MajorAnchorDiv);
 
 	if ( FMath::Abs(SubAnchorLinePos) <= SnapProximity )
 	{
-		Value = FMath::Round(MajorAnchorDiv) * SnapFrequency;
+		Value = FMath::RoundToInt(MajorAnchorDiv) * SnapFrequency;
 	}
 }
 
