UMG - Correcting spelling, disabling the rotation gizmo option until I have time to finish it.

diff --git a/Engine/Source/Editor/UMGEditor/Private/Designer/STransformHandle.cpp b/Engine/Source/Editor/UMGEditor/Private/Designer/STransformHandle.cpp
--- a/Engine/Source/Editor/UMGEditor/Private/Designer/STransformHandle.cpp
+++ b/Engine/Source/Editor/UMGEditor/Private/Designer/STransformHandle.cpp
@@ -117,8 +117,8 @@ FReply STransformHandle::OnMouseMove(const FGeometry& MyGeometry, const FPointer
 		{
 			FWidgetTransform RenderTransform = Preview->RenderTransform;
 
-			if ( Action == ETransformAction::Primary )
-			{
+			//if ( Action == ETransformAction::Primary )
+			//{
 				FVector2D TranslateAmount = MouseEvent.GetCursorDelta() * ( 1.0f / Designer->GetPreviewScale() );
 
 				if ( Preview->Slot )
@@ -130,12 +130,12 @@ FReply STransformHandle::OnMouseMove(const FGeometry& MyGeometry, const FPointer
 				{
 					Template->Slot->Resize(DragDirection, TranslateAmount);
 				}
-			}
-			else
-			{
-				FVector2D RotationDelta = MouseEvent.GetCursorDelta();
-				RenderTransform.Angle += RotationDelta.Size();
-			}
+			//}
+			//else
+			//{
+			//	FVector2D RotationDelta = MouseEvent.GetCursorDelta();
+			//	RenderTransform.Angle += RotationDelta.Size();
+			//}
 
 			static const FName RenderTransformName(TEXT("RenderTransform"));
 
@@ -157,8 +157,8 @@ FCursorReply STransformHandle::OnCursorQuery(const FGeometry& MyGeometry, const
 
 	ETransformMode::Type TransformMode = Designer->GetTransformMode();
 
-	if ( TransformMode == ETransformMode::Layout || CurrentAction == ETransformAction::Primary )
-	{
+	//if ( TransformMode == ETransformMode::Layout || CurrentAction == ETransformAction::Primary )
+	//{
 		switch ( TransformDirection )
 		{
 		case ETransformDirection::TopLeft:
@@ -174,11 +174,11 @@ FCursorReply STransformHandle::OnCursorQuery(const FGeometry& MyGeometry, const
 		case ETransformDirection::CenterRight:
 			return FCursorReply::Cursor(EMouseCursor::ResizeLeftRight);
 		}
-	}
-	else
-	{
-		return FCursorReply::Cursor(EMouseCursor::EyeDropper);
-	}
+	//}
+	//else
+	//{
+	//	return FCursorReply::Cursor(EMouseCursor::EyeDropper);
+	//}
 
 	return FCursorReply::Unhandled();
 }
diff --git a/Engine/Source/Editor/UMGEditor/Private/Extensions/CanvasSlotExtension.cpp b/Engine/Source/Editor/UMGEditor/Private/Extensions/CanvasSlotExtension.cpp
--- a/Engine/Source/Editor/UMGEditor/Private/Extensions/CanvasSlotExtension.cpp
+++ b/Engine/Source/Editor/UMGEditor/Private/Extensions/CanvasSlotExtension.cpp
@@ -266,7 +266,7 @@ FReply FCanvasSlotExtension::HandleAnchorEndDrag(const FGeometry& Geometry, cons
 	return FReply::Handled().ReleaseMouseCapture();
 }
 
-void FCanvasSlotExtension::PoximitySnapValue(float SnapFrequency, float SnapProximity, float& Value)
+void FCanvasSlotExtension::ProximitySnapValue(float SnapFrequency, float SnapProximity, float& Value)
 {
 	float MajorAnchorDiv = Value / SnapFrequency;
 	float SubAnchorLinePos = MajorAnchorDiv - FMath::RoundToInt(MajorAnchorDiv);
@@ -362,22 +362,22 @@ FReply FCanvasSlotExtension::HandleAnchorDragging(const FGeometry& Geometry, con
 
 						if ( LayoutData.Anchors.Minimum.X != OldLayoutData.Anchors.Minimum.X )
 						{
-							PoximitySnapValue(MajorAnchorLine, MajorAnchorLineSnapDistance, LayoutData.Anchors.Minimum.X);
+							ProximitySnapValue(MajorAnchorLine, MajorAnchorLineSnapDistance, LayoutData.Anchors.Minimum.X);
 						}
 
 						if ( LayoutData.Anchors.Minimum.Y != OldLayoutData.Anchors.Minimum.Y )
 						{
-							PoximitySnapValue(MajorAnchorLine, MajorAnchorLineSnapDistance, LayoutData.Anchors.Minimum.Y);
+							ProximitySnapValue(MajorAnchorLine, MajorAnchorLineSnapDistance, LayoutData.Anchors.Minimum.Y);
 						}
 
 						if ( LayoutData.Anchors.Maximum.X != OldLayoutData.Anchors.Maximum.X )
 						{
-							PoximitySnapValue(MajorAnchorLine, MajorAnchorLineSnapDistance, LayoutData.Anchors.Maximum.X);
+							ProximitySnapValue(MajorAnchorLine, MajorAnchorLineSnapDistance, LayoutData.Anchors.Maximum.X);
 						}
 
 						if ( LayoutData.Anchors.Maximum.Y != OldLayoutData.Anchors.Maximum.Y )
 						{
-							PoximitySnapValue(MajorAnchorLine, MajorAnchorLineSnapDistance, LayoutData.Anchors.Maximum.Y);
+							ProximitySnapValue(MajorAnchorLine, MajorAnchorLineSnapDistance, LayoutData.Anchors.Maximum.Y);
 						}
 					}
 
diff --git a/Engine/Source/Editor/UMGEditor/Private/Extensions/CanvasSlotExtension.h b/Engine/Source/Editor/UMGEditor/Private/Extensions/CanvasSlotExtension.h
--- a/Engine/Source/Editor/UMGEditor/Private/Extensions/CanvasSlotExtension.h
+++ b/Engine/Source/Editor/UMGEditor/Private/Extensions/CanvasSlotExtension.h
@@ -57,7 +57,7 @@ private:
 
 	void PaintCollisionLines(const TSet< FWidgetReference >& Selection, const FGeometry& AllottedGeometry, const FSlateRect& MyClippingRect, FSlateWindowElementList& OutDrawElements, int32 LayerId, const FWidgetStyle& InWidgetStyle, bool bParentEnabled) const;
 
-	void PoximitySnapValue(float SnapFrequency, float SnapProximity, float& Value);
+	void ProximitySnapValue(float SnapFrequency, float SnapProximity, float& Value);
 
 private:
 
