- Updated palette view appearance to look more like details panel (not 100% there yet)
- Palette view updates when widget blueprints are compiled or deleted and on hot reload
- Palette shows all widget blueprint classes including unloaded ones (that are loaded on create if need be)

diff --git a/Engine/Source/Editor/EditorStyle/Private/SlateEditorStyle.cpp b/Engine/Source/Editor/EditorStyle/Private/SlateEditorStyle.cpp
--- a/Engine/Source/Editor/EditorStyle/Private/SlateEditorStyle.cpp
+++ b/Engine/Source/Editor/EditorStyle/Private/SlateEditorStyle.cpp
@@ -5899,10 +5899,10 @@ void FSlateEditorStyle::FStyle::SetupUMGEditorStyles()
 	Set("UMGEditor.PalettePreviewWidget", new IMAGE_BRUSH("Icons/UMG/umg_anchors_fill_fill", Icon16x16));
 
 	Set("UMGEditor.PaletteHeader", FTableRowStyle()
-		.SetEvenRowBackgroundBrush( BOX_BRUSH( "PropertyView/DetailCategoryTop", FMargin( 4/16.0f, 8.0f/16.0f, 4/16.0f, 4/16.0f ) ) )
-		.SetEvenRowBackgroundHoveredBrush( BOX_BRUSH( "PropertyView/DetailCategoryTop", FMargin( 4/16.0f, 8.0f/16.0f, 4/16.0f, 4/16.0f ), FLinearColor(0.5f,0.5f,0.5f,1.0f) ) )
-		.SetOddRowBackgroundBrush( BOX_BRUSH( "PropertyView/DetailCategoryTop", FMargin( 4/16.0f, 8.0f/16.0f, 4/16.0f, 4/16.0f ) ) )
-		.SetOddRowBackgroundHoveredBrush( BOX_BRUSH( "PropertyView/DetailCategoryTop", FMargin( 4/16.0f, 8.0f/16.0f, 4/16.0f, 4/16.0f ), FLinearColor(0.5f,0.5f,0.5f,1.0f) ) )
+		.SetEvenRowBackgroundBrush( BOX_BRUSH( "PropertyView/DetailCategoryMiddle", FMargin( 4/16.0f, 8.0f/16.0f, 4/16.0f, 4/16.0f ), FLinearColor(0.6f,0.6f,0.6f,1.0f) ) )
+		.SetEvenRowBackgroundHoveredBrush( BOX_BRUSH( "PropertyView/DetailCategoryMiddle", FMargin( 4/16.0f, 8.0f/16.0f, 4/16.0f, 4/16.0f ), FLinearColor(0.3f,0.3f,0.3f,1.0f) ) )
+		.SetOddRowBackgroundBrush( BOX_BRUSH( "PropertyView/DetailCategoryMiddle", FMargin( 4/16.0f, 8.0f/16.0f, 4/16.0f, 4/16.0f ), FLinearColor(0.6f,0.6f,0.6f,1.0f) ) )
+		.SetOddRowBackgroundHoveredBrush( BOX_BRUSH( "PropertyView/DetailCategoryMiddle", FMargin( 4/16.0f, 8.0f/16.0f, 4/16.0f, 4/16.0f ), FLinearColor(0.3f,0.3f,0.3f,1.0f) ) )
 		.SetSelectorFocusedBrush( BORDER_BRUSH( "Common/Selector", FMargin(4.f/16.f), SelectorColor ) )
 		.SetActiveBrush( IMAGE_BRUSH( "Common/Selection", Icon8x8, SelectionColor ) )
 		.SetActiveHoveredBrush( IMAGE_BRUSH( "Common/Selection", Icon8x8, SelectionColor ) )
@@ -5912,6 +5912,11 @@ void FSlateEditorStyle::FStyle::SetupUMGEditorStyles()
 		.SetSelectedTextColor( InvertedForeground )
 		);
 
+	Set("UMGEditor.PaletteItem", FTableRowStyle(NormalTableRowStyle)
+		.SetEvenRowBackgroundBrush(BOX_BRUSH("PropertyView/DetailCategoryMiddle", FMargin(4 / 16.0f, 8.0f / 16.0f, 4 / 16.0f, 4 / 16.0f)))
+		.SetOddRowBackgroundBrush(BOX_BRUSH("PropertyView/DetailCategoryMiddle", FMargin(4 / 16.0f, 8.0f / 16.0f, 4 / 16.0f, 4 / 16.0f)))
+		);
+
 	Set("HorizontalAlignment_Left", new IMAGE_BRUSH("Icons/UMG/Alignment/Horizontal_Left", Icon20x20));
 	Set("HorizontalAlignment_Center", new IMAGE_BRUSH("Icons/UMG/Alignment/Horizontal_Center", Icon20x20));
 	Set("HorizontalAlignment_Right", new IMAGE_BRUSH("Icons/UMG/Alignment/Horizontal_Right", Icon20x20));
diff --git a/Engine/Source/Editor/UMGEditor/Private/Palette/SPaletteView.cpp b/Engine/Source/Editor/UMGEditor/Private/Palette/SPaletteView.cpp
--- a/Engine/Source/Editor/UMGEditor/Private/Palette/SPaletteView.cpp
+++ b/Engine/Source/Editor/UMGEditor/Private/Palette/SPaletteView.cpp
@@ -15,6 +15,11 @@
 
 #include "WidgetTemplate.h"
 #include "WidgetTemplateClass.h"
+#include "WidgetTemplateBlueprintClass.h"
+
+#include "Developer/HotReload/Public/IHotReload.h"
+
+#include "AssetRegistryModule.h"
 
 #define LOCTEXT_NAMESPACE "UMG"
 
@@ -75,12 +80,12 @@ public:
 	{
 	}
 
-	virtual FText GetName() const
+	virtual FText GetName() const override
 	{
 		return Template->Name;
 	}
 
-	virtual FString GetFilterString() const
+	virtual FString GetFilterString() const override
 	{
 		return Template->Name.ToString();
 	}
@@ -89,6 +94,7 @@ public:
 	{
 		return SNew(STableRow< TSharedPtr<FWidgetViewModel> >, OwnerTable)
 			.Padding(2.0f)
+			.Style(FEditorStyle::Get(), "UMGEditor.PaletteItem")
 			.OnDragDetected(this, &FWidgetTemplateViewModel::OnDraggingWidgetTemplateItem)
 			[
 				SNew(SPaletteViewItem, Template)
@@ -112,12 +118,12 @@ public:
 	{
 	}
 
-	virtual FText GetName() const
+	virtual FText GetName() const override
 	{
 		return GroupName;
 	}
 
-	virtual FString GetFilterString() const
+	virtual FString GetFilterString() const override
 	{
 		// Headers should never be included in filtering to avoid showing a header with all of
 		// it's widgets filtered out, so return an empty filter string.
@@ -127,12 +133,14 @@ public:
 	virtual TSharedRef<ITableRow> BuildRow(const TSharedRef<STableViewBase>& OwnerTable) override
 	{
 		return SNew(STableRow< TSharedPtr<FWidgetViewModel> >, OwnerTable)
-		.Style( FEditorStyle::Get(), "UMGEditor.PaletteHeader" )
+			.Style( FEditorStyle::Get(), "UMGEditor.PaletteHeader" )
 			.Padding(2.0f)
 			.ShowSelection(false)
 			[
 				SNew(STextBlock)
 				.Text(GroupName)
+				.Font(FEditorStyle::GetFontStyle("DetailsView.CategoryFontStyle"))
+				.ShadowOffset(FVector2D(1.0f, 1.0f))
 			];
 	}
 
@@ -150,11 +158,16 @@ public:
 
 void SPaletteView::Construct(const FArguments& InArgs, TSharedPtr<FBlueprintEditor> InBlueprintEditor)
 {
-	BlueprintEditor = InBlueprintEditor;
-
+	// Register for events that can trigger a palette rebuild
+	FWidgetBlueprintCompiler::OnWidgetBlueprintCompiled.AddSP(this, &SPaletteView::HandleOnBlueprintCompiled);
+	FEditorDelegates::OnAssetsDeleted.AddSP(this, &SPaletteView::HandleOnAssetsDeleted);
+	IHotReloadModule::Get().OnHotReload().AddSP(this, &SPaletteView::HandleOnHotReload);
+	
 	// register for any objects replaced
 	GEditor->OnObjectsReplaced().AddRaw(this, &SPaletteView::OnObjectsReplaced);
 
+	BlueprintEditor = InBlueprintEditor;
+
 	UBlueprint* BP = InBlueprintEditor->GetBlueprintObj();
 
 	WidgetFilter = MakeShareable(new WidgetViewModelTextFilter(
@@ -166,10 +179,12 @@ void SPaletteView::Construct(const FArguments& InArgs, TSharedPtr<FBlueprintEdit
 	FilterHandler->SetGetChildrenDelegate(PaletteFilterHandler::FOnGetChildren::CreateRaw(this, &SPaletteView::OnGetChildren));
 
 	SAssignNew(WidgetTemplatesView, STreeView< TSharedPtr<FWidgetViewModel> >)
-	.SelectionMode(ESelectionMode::Single)
-	.OnGenerateRow(this, &SPaletteView::OnGenerateWidgetTemplateItem)
-	.OnGetChildren(FilterHandler.ToSharedRef(), &PaletteFilterHandler::OnGetFilteredChildren)
-	.TreeItemsSource(&TreeWidgetViewModels);
+		.ItemHeight(1.0f)
+		.SelectionMode(ESelectionMode::Single)
+		.OnGenerateRow(this, &SPaletteView::OnGenerateWidgetTemplateItem)
+		.OnGetChildren(FilterHandler.ToSharedRef(), &PaletteFilterHandler::OnGetFilteredChildren)
+		.TreeItemsSource(&TreeWidgetViewModels);
+		
 
 	FilterHandler->SetTreeView(WidgetTemplatesView.Get());
 
@@ -199,9 +214,9 @@ void SPaletteView::Construct(const FArguments& InArgs, TSharedPtr<FBlueprintEdit
 	bRefreshRequested = true;
 
 	BuildWidgetList();
-	LoadItemExpanssion();
+	LoadItemExpansion();
 
-	bRebuildRquested = false;
+	bRebuildRequested = false;
 }
 
 SPaletteView::~SPaletteView()
@@ -214,7 +229,11 @@ SPaletteView::~SPaletteView()
 		FilterHandler->RefreshAndFilterTree();
 	}
 
+	FWidgetBlueprintCompiler::OnWidgetBlueprintCompiled.RemoveAll(this);
+	FEditorDelegates::OnAssetsDeleted.RemoveAll(this);
+	IHotReloadModule::Get().OnHotReload().RemoveAll(this);
 	GEditor->OnObjectsReplaced().RemoveAll( this );
+	
 
 	SaveItemExpansion();
 }
@@ -232,7 +251,7 @@ FText SPaletteView::GetSearchText() const
 	return SearchText;
 }
 
-void SPaletteView::LoadItemExpanssion()
+void SPaletteView::LoadItemExpansion()
 {
 	// Restore the expansion state of the widget groups.
 	for ( TSharedPtr<FWidgetViewModel>& ViewModel : WidgetViewModels )
@@ -303,32 +322,76 @@ void SPaletteView::BuildClassWidgetList()
 {
 	static const FName DevelopmentStatusKey(TEXT("DevelopmentStatus"));
 
-	for ( TObjectIterator<UClass> ClassIt; ClassIt; ++ClassIt )
+	TMap<FName, TSubclassOf<UUserWidget>> LoadedWidgetBlueprintClassesByName;
+
+	auto ActiveWidgetBlueprintClass = GetBlueprint()->GeneratedClass;
+	FName ActiveWidgetBlueprintClassName = ActiveWidgetBlueprintClass->GetFName();
+
+	// Locate all UWidget classes from code and loaded widget BPs
+	for (TObjectIterator<UClass> ClassIt; ClassIt; ++ClassIt)
 	{
 		UClass* WidgetClass = *ClassIt;
+
 		if ( WidgetClass->IsChildOf(UWidget::StaticClass()) )
-		{
-			// We don't want to include experimental widget classes, so filter those out.
+		{	
+			// We aren't interested in classes that are experimental or cannot be instantiated
 			bool bIsExperimental, bIsEarlyAccess;
 			FObjectEditorUtils::GetClassDevelopmentStatus(WidgetClass, bIsExperimental, bIsEarlyAccess);
+			bool bIsInstantiable = !WidgetClass->HasAnyClassFlags(CLASS_Abstract | CLASS_Deprecated | CLASS_NewerVersionExists);
+			if ( !(bIsExperimental || bIsEarlyAccess || bIsInstantiable) )
+			{
+				continue;
+			}
 
-			// Don't include skeleton classes of widgets.
-			const bool bIsSkeletonClass = WidgetClass->HasAnyFlags(RF_Transient) && WidgetClass->HasAnyClassFlags(CLASS_CompiledFromBlueprint);
-			const bool bIsValidClass =
-				!( WidgetClass->HasAnyClassFlags(CLASS_Abstract | CLASS_Deprecated | CLASS_NewerVersionExists) || bIsSkeletonClass || bIsExperimental );
-			const bool bIsSameClass = WidgetClass->GetFName() == GetBlueprint()->GeneratedClass->GetFName();
-
-			//TODO UMG does not prevent deep nested circular references
+			if (WidgetClass->IsChildOf(UUserWidget::StaticClass()))
+			{
+				// Don't include skeleton classes or the same class as the widget being edited
+				const bool bIsSkeletonClass = WidgetClass->HasAnyFlags(RF_Transient) && WidgetClass->HasAnyClassFlags(CLASS_CompiledFromBlueprint);
+				const bool bIsSameClass = WidgetClass->GetFName() == ActiveWidgetBlueprintClassName;
 
-			// Only use non-(abstract/deprecated/old) classes and don't include the current widget blueprint class.
-			if ( bIsValidClass && !bIsSameClass )
+				// Check that the asset that generated this class is valid (necessary b/c of a larger issue wherein force delete does not wipe the generated class object)
+				if ( !(bIsSkeletonClass || bIsSameClass) && WidgetClass->ClassGeneratedBy )
+				{
+					// Track the widget blueprint classes that are already loaded
+					LoadedWidgetBlueprintClassesByName.Add(WidgetClass->ClassGeneratedBy->GetFName()) = WidgetClass;
+				}
+			}
+			else
 			{
 				TSharedPtr<FWidgetTemplateClass> Template = MakeShareable(new FWidgetTemplateClass(WidgetClass));
 
 				AddWidgetTemplate(Template);
 			}
+
+			//TODO UMG does not prevent deep nested circular references
 		}
 	}
+
+	// Locate all widget BP assets (include unloaded)
+	const FAssetRegistryModule& AssetRegistryModule = FModuleManager::LoadModuleChecked<FAssetRegistryModule>(TEXT("AssetRegistry"));
+	TArray<FAssetData> AllWidgetBPsAssetData;
+	AssetRegistryModule.Get().GetAssetsByClass(UWidgetBlueprint::StaticClass()->GetFName(), AllWidgetBPsAssetData, true);
+
+	FName ActiveWidgetBlueprintName = ActiveWidgetBlueprintClass->ClassGeneratedBy->GetFName();
+	for (auto& WidgetBPAssetData : AllWidgetBPsAssetData)
+	{
+		if (WidgetBPAssetData.AssetName == ActiveWidgetBlueprintName)
+		{
+			continue;
+		}
+
+		// If the blueprint generated class was found earlier, pass it to the template
+		TSubclassOf<UUserWidget> WidgetBPClass = nullptr;
+		auto LoadedWidgetBPClass = LoadedWidgetBlueprintClassesByName.Find(WidgetBPAssetData.AssetName);
+		if (LoadedWidgetBPClass)
+		{
+			WidgetBPClass = *LoadedWidgetBPClass;
+		}
+
+		auto Template = MakeShareable(new FWidgetTemplateBlueprintClass(WidgetBPAssetData, WidgetBPClass));
+
+		AddWidgetTemplate(Template);
+	}
 }
 
 void SPaletteView::BuildSpecialWidgetList()
@@ -358,8 +421,10 @@ TSharedRef<ITableRow> SPaletteView::OnGenerateWidgetTemplateItem(TSharedPtr<FWid
 
 void SPaletteView::Tick(const FGeometry& AllottedGeometry, const double InCurrentTime, const float InDeltaTime)
 {
-	if ( bRebuildRquested )
+	if ( bRebuildRequested )
 	{
+		bRebuildRequested = false;
+
 		// Save the old expanded items temporarily
 		TSet<TSharedPtr<FWidgetViewModel>> ExpandedItems;
 		WidgetTemplatesView->GetExpandedItems(ExpandedItems);
@@ -394,7 +459,35 @@ void SPaletteView::TransformWidgetViewModelToString(TSharedPtr<FWidgetViewModel>
 void SPaletteView::OnObjectsReplaced(const TMap<UObject*, UObject*>& ReplacementMap)
 {
 	//bRefreshRequested = true;
-	//bRebuildRquested = true;
+	//bRebuildRequested = true;
+}
+
+void SPaletteView::HandleOnBlueprintCompiled(UBlueprint* Blueprint)
+{
+	if (Blueprint->IsA<UWidgetBlueprint>())
+	{
+		bRebuildRequested = true;
+		bRefreshRequested = true;
+	}
+}
+
+void SPaletteView::HandleOnHotReload(bool bWasTriggeredAutomatically)
+{
+	bRebuildRequested = true;
+	bRefreshRequested = true;
+}
+
+void SPaletteView::HandleOnAssetsDeleted(const TArray<UClass*>& DeletedAssetClasses)
+{
+	for (auto DeletedAssetClass : DeletedAssetClasses)
+	{
+		if (DeletedAssetClass->IsChildOf(UWidgetBlueprint::StaticClass()))
+		{
+			bRebuildRequested = true;
+			bRefreshRequested = true;
+		}
+	}
+	
 }
 
 #undef LOCTEXT_NAMESPACE
diff --git a/Engine/Source/Editor/UMGEditor/Private/Palette/SPaletteView.h b/Engine/Source/Editor/UMGEditor/Private/Palette/SPaletteView.h
--- a/Engine/Source/Editor/UMGEditor/Private/Palette/SPaletteView.h
+++ b/Engine/Source/Editor/UMGEditor/Private/Palette/SPaletteView.h
@@ -56,7 +56,7 @@ private:
 	void OnSearchChanged(const FText& InFilterText);
 
 private:
-	void LoadItemExpanssion();
+	void LoadItemExpansion();
 	void SaveItemExpansion();
 
 	/** Called when a Blueprint is recompiled and live objects are swapped out for replacements */
@@ -67,6 +67,15 @@ private:
 	/** Transforms the widget view model into a searchable string. */
 	void TransformWidgetViewModelToString(TSharedPtr<FWidgetViewModel> WidgetViewModel, OUT TArray< FString >& Array);
 
+	/** Requests a rebuild of the widget list if a widget blueprint was compiled */
+	void HandleOnBlueprintCompiled(UBlueprint* Blueprint);
+
+	/** Requests a rebuild of the widget list */
+	void HandleOnHotReload(bool bWasTriggeredAutomatically);
+
+	/** Requests a rebuild of the widget list if a widget blueprint was deleted */
+	void HandleOnAssetsDeleted(const TArray<UClass*>& DeletedAssetClasses);
+
 	TWeakPtr<class FBlueprintEditor> BlueprintEditor;
 
 	/** Handles filtering the palette based on an IFilter. */
@@ -93,5 +102,5 @@ private:
 	FText SearchText;
 
 	/** Controls rebuilding the list of spawnable widgets */
-	bool bRebuildRquested;
+	bool bRebuildRequested;
 };
diff --git a/Engine/Source/Editor/UMGEditor/Private/Templates/WidgetTemplateBlueprintClass.cpp b/Engine/Source/Editor/UMGEditor/Private/Templates/WidgetTemplateBlueprintClass.cpp
new file mode 100644
--- /dev/null
+++ b/Engine/Source/Editor/UMGEditor/Private/Templates/WidgetTemplateBlueprintClass.cpp
@@ -0,0 +1,63 @@
+// Copyright 1998-2014 Epic Games, Inc. All Rights Reserved.
+
+#include "UMGEditorPrivatePCH.h"
+
+#include "WidgetTemplateBlueprintClass.h"
+#include "IDocumentation.h"
+
+#define LOCTEXT_NAMESPACE "UMGEditor"
+
+FWidgetTemplateBlueprintClass::FWidgetTemplateBlueprintClass(const FAssetData& InWidgetAssetData, TSubclassOf<UUserWidget> InUserWidgetClass)
+	: FWidgetTemplateClass(), WidgetAssetData(InWidgetAssetData)
+{
+	if (InUserWidgetClass)
+	{
+		WidgetClass = CastChecked<UClass>(InUserWidgetClass);
+		Name = WidgetClass->GetDisplayNameText();
+	}
+	else
+	{
+		Name = FText::FromString(FName::NameToDisplayString(WidgetAssetData.AssetName.ToString(), false));
+	}
+}
+
+FWidgetTemplateBlueprintClass::~FWidgetTemplateBlueprintClass()
+{
+}
+
+FText FWidgetTemplateBlueprintClass::GetCategory() const
+{
+	auto DefaultUserWidget = UUserWidget::StaticClass()->GetDefaultObject<UUserWidget>();
+	return DefaultUserWidget->GetPaletteCategory();
+}
+
+UWidget* FWidgetTemplateBlueprintClass::Create(UWidgetTree* Tree)
+{
+	// Load the blueprint asset if needed
+	if (!WidgetClass.Get())
+	{
+		FString AssetPath = WidgetAssetData.ObjectPath.ToString();
+		UWidgetBlueprint* LoadedWidget = LoadObject<UWidgetBlueprint>(nullptr, *AssetPath);
+		WidgetClass = CastChecked<UClass>(LoadedWidget->GeneratedClass);
+	}
+
+	return FWidgetTemplateClass::Create(Tree);
+}
+
+const FSlateBrush* FWidgetTemplateBlueprintClass::GetIcon() const
+{
+	auto DefaultUserWidget = UUserWidget::StaticClass()->GetDefaultObject<UUserWidget>();
+	return DefaultUserWidget->GetEditorIcon();
+}
+
+TSharedRef<IToolTip> FWidgetTemplateBlueprintClass::GetToolTip() const
+{
+	if (WidgetClass.Get())
+	{
+		return FWidgetTemplateClass::GetToolTip();
+	}
+
+	return IDocumentation::Get()->CreateToolTip(Name, nullptr, FString(TEXT("Shared/Types/")) + Name.ToString(), TEXT("Class"));
+}
+
+#undef LOCTEXT_NAMESPACE
diff --git a/Engine/Source/Editor/UMGEditor/Private/Templates/WidgetTemplateBlueprintClass.h b/Engine/Source/Editor/UMGEditor/Private/Templates/WidgetTemplateBlueprintClass.h
new file mode 100644
--- /dev/null
+++ b/Engine/Source/Editor/UMGEditor/Private/Templates/WidgetTemplateBlueprintClass.h
@@ -0,0 +1,39 @@
+// Copyright 1998-2014 Epic Games, Inc. All Rights Reserved.
+
+#pragma once
+
+#include "WidgetTemplateClass.h"
+
+/**
+* A template for classes generated by UWidgetBlueprint assets
+* Accounts for the possibility that a the blueprint may not be loaded when the palette is created
+*/
+class FWidgetTemplateBlueprintClass : public FWidgetTemplateClass
+{
+public:
+	/** 
+	* Constructor
+	* @param InWidgetAssetData The FAssetData for the widget blueprint asset
+	* @param InUserWidgetClass The user-created widget class that will be created by this template
+	*/
+	FWidgetTemplateBlueprintClass(const FAssetData& InWidgetAssetData, TSubclassOf<UUserWidget> InUserWidgetClass = nullptr);
+
+	/** Destructor */
+	virtual ~FWidgetTemplateBlueprintClass();
+
+	/** Gets the category for the widget */
+	virtual FText GetCategory() const override;
+
+	/** Creates an instance of the widget for the tree. */
+	virtual UWidget* Create(UWidgetTree* Tree) override;
+
+	/** The icon coming from the default object of the class */
+	virtual const FSlateBrush* GetIcon() const override;
+
+	/** Gets the tooltip widget for this palette item. */
+	virtual TSharedRef<IToolTip> GetToolTip() const override;
+
+protected:
+	/** The asset data for the widget blueprint */
+	FAssetData WidgetAssetData;
+};
diff --git a/Engine/Source/Editor/UMGEditor/Private/Templates/WidgetTemplateClass.cpp b/Engine/Source/Editor/UMGEditor/Private/Templates/WidgetTemplateClass.cpp
--- a/Engine/Source/Editor/UMGEditor/Private/Templates/WidgetTemplateClass.cpp
+++ b/Engine/Source/Editor/UMGEditor/Private/Templates/WidgetTemplateClass.cpp
@@ -7,10 +7,25 @@
 
 #define LOCTEXT_NAMESPACE "UMGEditor"
 
+FWidgetTemplateClass::FWidgetTemplateClass()
+	: WidgetClass(nullptr)
+{
+	// register for any objects replaced
+	GEditor->OnObjectsReplaced().AddRaw(this, &FWidgetTemplateClass::OnObjectsReplaced);
+}
+
 FWidgetTemplateClass::FWidgetTemplateClass(TSubclassOf<UWidget> InWidgetClass)
 	: WidgetClass(InWidgetClass)
 {
 	Name = WidgetClass->GetDisplayNameText();
+
+	// register for any objects replaced
+	GEditor->OnObjectsReplaced().AddRaw(this, &FWidgetTemplateClass::OnObjectsReplaced);
+}
+
+FWidgetTemplateClass::~FWidgetTemplateClass()
+{
+	GEditor->OnObjectsReplaced().RemoveAll(this);
 }
 
 FText FWidgetTemplateClass::GetCategory() const
@@ -21,7 +36,7 @@ FText FWidgetTemplateClass::GetCategory() const
 
 UWidget* FWidgetTemplateClass::Create(UWidgetTree* Tree)
 {
-	return Tree->ConstructWidget<UWidget>(WidgetClass);
+	return Tree->ConstructWidget<UWidget>(WidgetClass.Get());
 }
 
 const FSlateBrush* FWidgetTemplateClass::GetIcon() const
@@ -35,4 +50,13 @@ TSharedRef<IToolTip> FWidgetTemplateClass::GetToolTip() const
 	return IDocumentation::Get()->CreateToolTip(FText::FromString(WidgetClass->GetDescription()), nullptr, FString(TEXT("Shared/Types/")) + WidgetClass->GetName(), TEXT("Class"));
 }
 
+void FWidgetTemplateClass::OnObjectsReplaced(const TMap<UObject*, UObject*>& ReplacementMap)
+{
+	UObject* const* NewObject = ReplacementMap.Find(WidgetClass.Get());
+	if (NewObject)
+	{
+		WidgetClass = CastChecked<UClass>(*NewObject);
+	}
+}
+
 #undef LOCTEXT_NAMESPACE
diff --git a/Engine/Source/Editor/UMGEditor/Private/Templates/WidgetTemplateClass.h b/Engine/Source/Editor/UMGEditor/Private/Templates/WidgetTemplateClass.h
--- a/Engine/Source/Editor/UMGEditor/Private/Templates/WidgetTemplateClass.h
+++ b/Engine/Source/Editor/UMGEditor/Private/Templates/WidgetTemplateClass.h
@@ -11,10 +11,10 @@ class FWidgetTemplateClass : public FWidgetTemplate
 {
 public:
 	/** Constructor */
-	FWidgetTemplateClass(TSubclassOf<UWidget> InWidgetClass);
+	explicit FWidgetTemplateClass(TSubclassOf<UWidget> InWidgetClass);
 
 	/** Destructor */
-	virtual ~FWidgetTemplateClass() {}
+	virtual ~FWidgetTemplateClass();
 
 	/** Gets the category for the widget */
 	virtual FText GetCategory() const override;
@@ -29,6 +29,12 @@ public:
 	virtual TSharedRef<IToolTip> GetToolTip() const override;
 
 protected:
+	/** Creates a widget template class without any class reference */
+	FWidgetTemplateClass();
+
+	/** Called when objects need to be swapped out for new versions, like after a blueprint recompile. */
+	void OnObjectsReplaced(const TMap<UObject*, UObject*>& ReplacementMap);
+
 	/** The widget class that will be created by this template */
-	TSubclassOf<UWidget> WidgetClass;
+	TWeakObjectPtr<UClass> WidgetClass;
 };
diff --git a/Engine/Source/Editor/UnrealEd/Private/Editor.cpp b/Engine/Source/Editor/UnrealEd/Private/Editor.cpp
--- a/Engine/Source/Editor/UnrealEd/Private/Editor.cpp
+++ b/Engine/Source/Editor/UnrealEd/Private/Editor.cpp
@@ -139,6 +139,7 @@ FEditorDelegates::FOnDollyPerspectiveCamera				FEditorDelegates::OnDollyPerspect
 FEditorDelegates::FOnBlueprintContextMenuCreated		FEditorDelegates::OnBlueprintContextMenuCreated;
 FSimpleMulticastDelegate								FEditorDelegates::OnShutdownPostPackagesSaved;
 FEditorDelegates::FOnAssetsPreDelete					FEditorDelegates::OnAssetsPreDelete;
+FEditorDelegates::FOnAssetsDeleted						FEditorDelegates::OnAssetsDeleted;
 
 /*-----------------------------------------------------------------------------
 	Globals.
diff --git a/Engine/Source/Editor/UnrealEd/Private/ObjectTools.cpp b/Engine/Source/Editor/UnrealEd/Private/ObjectTools.cpp
--- a/Engine/Source/Editor/UnrealEd/Private/ObjectTools.cpp
+++ b/Engine/Source/Editor/UnrealEd/Private/ObjectTools.cpp
@@ -1651,11 +1651,15 @@ namespace ObjectTools
 		// Update the browser if something was actually deleted.
 		if ( bSawSuccessfulDelete )
 		{
+			TArray<UClass*> DeletedObjectClasses;
 			TArray<UPackage*> PotentialPackagesToDelete;
 			for ( int32 ObjIdx = 0; ObjIdx < ObjectsDeletedSuccessfully.Num(); ++ObjIdx )
 			{
+				DeletedObjectClasses.AddUnique(ObjectsDeletedSuccessfully[ObjIdx]->GetClass());
 				PotentialPackagesToDelete.AddUnique( ObjectsDeletedSuccessfully[ObjIdx]->GetOutermost() );
 			}
+			// Broadcast the classes of the successfully deleted objects (before cleanup)
+			FEditorDelegates::OnAssetsDeleted.Broadcast(DeletedObjectClasses);
 
 			bool bPerformReferenceCheck = false;
 			CleanupAfterSuccessfulDelete( PotentialPackagesToDelete, bPerformReferenceCheck );
diff --git a/Engine/Source/Editor/UnrealEd/Public/Editor.h b/Engine/Source/Editor/UnrealEd/Public/Editor.h
--- a/Engine/Source/Editor/UnrealEd/Public/Editor.h
+++ b/Engine/Source/Editor/UnrealEd/Public/Editor.h
@@ -90,7 +90,8 @@ struct UNREALED_API FEditorDelegates
 	DECLARE_MULTICAST_DELEGATE_OneParam(FOnEditorModeTransitioned, FEdMode* /*Mode*/);
 	/** delegate type for when a user requests to delete certain assets... DOES NOT mean the asset(s) will be deleted (the user could cancel) */
 	DECLARE_MULTICAST_DELEGATE_OneParam(FOnAssetsPreDelete, const TArray<UObject*>&);
-
+	/** delegate type for when one or more assets have been deleted */
+	DECLARE_MULTICAST_DELEGATE_OneParam(FOnAssetsDeleted, const TArray<UClass*>& /*DeletedAssetClasses*/);
 
 	/** Called when the CurrentLevel is switched to a new level.  Note that this event won't be fired for temporary
 		changes to the current level, such as when copying/pasting actors. */
@@ -182,6 +183,8 @@ struct UNREALED_API FEditorDelegates
 	static FSimpleMulticastDelegate OnShutdownPostPackagesSaved;
 	/** Called when the user requests certain assets be deleted (DOES NOT imply that the asset will be deleted... the user could cancel) */
 	static FOnAssetsPreDelete OnAssetsPreDelete;
+	/** Called when one or more assets have been deleted */
+	static FOnAssetsDeleted OnAssetsDeleted;
 };
 
 /**
