Fixed aliased selection outline border on UMG selection highlights

diff --git a/Engine/Source/Editor/EditorStyle/Private/SlateEditorStyle.cpp b/Engine/Source/Editor/EditorStyle/Private/SlateEditorStyle.cpp
--- a/Engine/Source/Editor/EditorStyle/Private/SlateEditorStyle.cpp
+++ b/Engine/Source/Editor/EditorStyle/Private/SlateEditorStyle.cpp
@@ -5919,7 +5919,7 @@ void FSlateEditorStyle::FStyle::SetupUMGEditorStyles()
 
 	const FTextBlockStyle NoAnimationFont =
 		FTextBlockStyle(NormalText)
-		.SetFont(TTF_CORE_FONT("Fonts/Roboto-Regular", 20))
+		.SetFont(TTF_CORE_FONT("Fonts/Roboto-Regular", 18))
 		.SetColorAndOpacity(FSlateColor::UseForeground())
 		.SetShadowOffset(FVector2D(1, 1))
 		.SetShadowColorAndOpacity(FLinearColor::Black);
@@ -5927,6 +5927,8 @@ void FSlateEditorStyle::FStyle::SetupUMGEditorStyles()
 
 	Set("UMGEditor.AddAnimationIcon", new IMAGE_BRUSH("Icons/PlusSymbol_12x", Icon12x12, FLinearColor(.05,.05,.05) ) );
 	Set("UMGEditor.NoAnimationFont", NoAnimationFont);
+
+	Set("UMGEditor.SelectionOutline", new BORDER_BRUSH("UMG/SelectionBorder", FMargin(.2f) ));
 }
 
 void FSlateEditorStyle::FStyle::SetupTranslationEditorStyles()
diff --git a/Engine/Source/Editor/UMGEditor/Private/Designer/SDesignerView.cpp b/Engine/Source/Editor/UMGEditor/Private/Designer/SDesignerView.cpp
--- a/Engine/Source/Editor/UMGEditor/Private/Designer/SDesignerView.cpp
+++ b/Engine/Source/Editor/UMGEditor/Private/Designer/SDesignerView.cpp
@@ -969,6 +969,7 @@ int32 SDesignerView::OnPaint(const FPaintArgs& Args, const FGeometry& AllottedGe
 		Ext->Paint(Selected, AllottedGeometry, MyClippingRect, OutDrawElements, LayerId, InWidgetStyle, bParentEnabled);
 	}
 
+	static const FName SelectionOutlineName("UMGEditor.SelectionOutline");
 	// Don't draw the hovered effect if it's also the selected widget
 	if ( HoveredSlateWidget.IsValid() && HoveredSlateWidget != SelectedSlateWidget )
 	{
@@ -981,16 +982,13 @@ int32 SDesignerView::OnPaint(const FPaintArgs& Args, const FGeometry& AllottedGe
 		// Azure = 0x007FFF
 		const FLinearColor HoveredTint(0, 0.5, 1, FMath::Clamp(HoverTime / HoveredAnimationTime, 0.0f, 1.0f));
 
-		FPaintGeometry HoveredGeometry(
-			ArrangedWidget.Geometry.AbsolutePosition,
-			ArrangedWidget.Geometry.Size * ArrangedWidget.Geometry.Scale,
-			ArrangedWidget.Geometry.Scale);
+		FPaintGeometry HoveredGeometry = ArrangedWidget.Geometry.ToPaintGeometry();
 
 		FSlateDrawElement::MakeBox(
 			OutDrawElements,
 			LayerId,
 			HoveredGeometry,
-			FCoreStyle::Get().GetBrush(TEXT("Debug.Border")),
+			FEditorStyle::Get().GetBrush(SelectionOutlineName),
 			MyClippingRect,
 			ESlateDrawEffect::None,
 			HoveredTint
@@ -1014,7 +1012,7 @@ int32 SDesignerView::OnPaint(const FPaintArgs& Args, const FGeometry& AllottedGe
 			OutDrawElements,
 			LayerId,
 			SelectionGeometry,
-			FCoreStyle::Get().GetBrush(TEXT("Debug.Border")),
+			FEditorStyle::Get().GetBrush(SelectionOutlineName),
 			MyClippingRect,
 			ESlateDrawEffect::None,
 			Tint
