UMG - Fixed regression in the palette view showing widgets we're not interested in exposing at this time.

diff --git a/Engine/Source/Editor/UMGEditor/Private/Palette/SPaletteView.cpp b/Engine/Source/Editor/UMGEditor/Private/Palette/SPaletteView.cpp
--- a/Engine/Source/Editor/UMGEditor/Private/Palette/SPaletteView.cpp
+++ b/Engine/Source/Editor/UMGEditor/Private/Palette/SPaletteView.cpp
@@ -337,8 +337,8 @@ void SPaletteView::BuildClassWidgetList()
 			// We aren't interested in classes that are experimental or cannot be instantiated
 			bool bIsExperimental, bIsEarlyAccess;
 			FObjectEditorUtils::GetClassDevelopmentStatus(WidgetClass, bIsExperimental, bIsEarlyAccess);
-			bool bIsInstantiable = !WidgetClass->HasAnyClassFlags(CLASS_Abstract | CLASS_Deprecated | CLASS_NewerVersionExists);
-			if ( !(bIsExperimental || bIsEarlyAccess || bIsInstantiable) )
+			const bool bIsInvalid = WidgetClass->HasAnyClassFlags(CLASS_Abstract | CLASS_Deprecated | CLASS_NewerVersionExists);
+			if ( bIsExperimental || bIsEarlyAccess || bIsInvalid )
 			{
 				continue;
 			}
