UMG - Moving the disappearing info bar to the bottom of the designer view.

diff --git a/Engine/Source/Editor/UMGEditor/Private/Designer/SDesignerView.cpp b/Engine/Source/Editor/UMGEditor/Private/Designer/SDesignerView.cpp
--- a/Engine/Source/Editor/UMGEditor/Private/Designer/SDesignerView.cpp
+++ b/Engine/Source/Editor/UMGEditor/Private/Designer/SDesignerView.cpp
@@ -247,17 +247,33 @@ void SDesignerView::Construct(const FArguments& InArgs, TSharedPtr<FWidgetBluepr
 				]
 			]
 
+			// Bottom bar to show current resolution & AR
+			+ SOverlay::Slot()
+			.HAlign(HAlign_Fill)
+			.VAlign(VAlign_Bottom)
+			[
+				SNew(SHorizontalBox)
+				+ SHorizontalBox::Slot()
+				.AutoWidth()
+				.Padding(6, 0, 0, 2)
+				[
+					SNew(STextBlock)
+					.TextStyle(FEditorStyle::Get(), "Graph.ZoomText")
+					.Text(this, &SDesignerView::GetCurrentResolutionText)
+					.ColorAndOpacity(this, &SDesignerView::GetResolutionTextColorAndOpacity)
+				]
+			]
+
 			// Info Bar, displays heads up information about some actions.
 			+ SOverlay::Slot()
 			.HAlign(HAlign_Fill)
 			.VAlign(VAlign_Bottom)
-			.Padding(FMargin(0, 0, 0, 200))
 			[
 				SNew(SDisappearingBar)
 				[
 					SNew(SBorder)
-					.BorderImage(FEditorStyle::GetBrush("ToolPanel.GroupBorder"))
-					.BorderBackgroundColor(FLinearColor(1, 1, 1, 0.75))
+					.BorderImage(FEditorStyle::GetBrush("WhiteBrush"))
+					.BorderBackgroundColor(FLinearColor(0.10, 0.10, 0.10, 0.75))
 					.HAlign(HAlign_Center)
 					.VAlign(VAlign_Center)
 					.Padding(FMargin(0, 5))
@@ -269,23 +285,6 @@ void SDesignerView::Construct(const FArguments& InArgs, TSharedPtr<FWidgetBluepr
 					]
 				]
 			]
-
-			// Bottom bar to show current resolution & AR
-			+ SOverlay::Slot()
-			.HAlign(HAlign_Fill)
-			.VAlign(VAlign_Bottom)
-			[
-				SNew(SHorizontalBox)
-				+ SHorizontalBox::Slot()
-				.AutoWidth()
-				.Padding(6, 0, 0, 2)
-				[
-					SNew(STextBlock)
-					.TextStyle(FEditorStyle::Get(), "Graph.ZoomText")
-					.Text(this, &SDesignerView::GetCurrentResolutionText)
-					.ColorAndOpacity(this, &SDesignerView::GetResolutionTextColorAndOpacity)
-				]
-			]
 		]
 	);
 
