UMG - Fixing a crash bug with styles still being referenced after blueprint compiles.  We now ditch all widget previews in all designers after any widget blueprint compile to prevent referencing UWidgets, but still presenting the slate previews.

diff --git a/Engine/Source/Editor/UMGEditor/Private/Designer/SDesignerView.cpp b/Engine/Source/Editor/UMGEditor/Private/Designer/SDesignerView.cpp
--- a/Engine/Source/Editor/UMGEditor/Private/Designer/SDesignerView.cpp
+++ b/Engine/Source/Editor/UMGEditor/Private/Designer/SDesignerView.cpp
@@ -654,18 +654,13 @@ void SDesignerView::Register(TSharedRef<FDesignerExtension> Extension)
 
 void SDesignerView::OnBlueprintCompiled(UBlueprint* InBlueprint)
 {
-	UBlueprint* BP = BlueprintEditor.Pin()->GetBlueprintObj();
-	if( BP == InBlueprint )
-	{
-		CachedWidgetGeometry.Reset();
-
-		PreviewWidget = NULL;
-		PreviewSurface->SetContent(SNullWidget::NullWidget);
-		//PreviewWidget.ReleaseNativeWidget()
+	// Because widget blueprints can contain other widget blueprints, the safe thing to do is to have all
+	// designers jettison their previews on the compilation of any widget blueprint.  We do this to prevent
+	// having slate widgets that still may reference data in their owner UWidget that has been garbage collected.
+	CachedWidgetGeometry.Reset();
 
-		//const bool bForceUpdate = true;
-		//UpdatePreviewWidget(bForceUpdate);
-	}
+	PreviewWidget = NULL;
+	PreviewSurface->SetContent(SNullWidget::NullWidget);
 }
 
 FWidgetReference SDesignerView::GetWidgetAtCursor(const FGeometry& MyGeometry, const FPointerEvent& MouseEvent, FArrangedWidget& ArrangedWidget)
