UMG - Introducing a Wrap Box that can wrap widgets onto a new line if it runs out of space.  Introducing named slot support, still under development, but allows user widgets to expose slots that content can be placed into in different instances of the user widget.  The designer now allows custom design time resolutions to be setup on a per widget basis, because sometimes you don't want to model a small widget as if it were a fullscreen widget.  Fixing a bug with how the designer draws effects by having a shim widget that paints at a specific layer in the designers overlay.  Introducing a modification to the progress bar so that the default progress bar brush is white, and the color is a tint so that progress bars can more easliy styed in UMG out of the box, fixes the 'my progress bar can't be blue' bug.  Adding a small grid to the anchor dialog to give a sense of resize instead of just pulsing.

diff --git a/Engine/Source/Editor/EditorStyle/Private/SlateEditorStyle.cpp b/Engine/Source/Editor/EditorStyle/Private/SlateEditorStyle.cpp
--- a/Engine/Source/Editor/EditorStyle/Private/SlateEditorStyle.cpp
+++ b/Engine/Source/Editor/EditorStyle/Private/SlateEditorStyle.cpp
@@ -486,7 +486,7 @@ void FSlateEditorStyle::FStyle::SetupGeneralStyles()
 	{
 		Set( "ProgressBar", FProgressBarStyle()
 			.SetBackgroundImage( BOX_BRUSH( "Common/ProgressBar_Background", FMargin(5.f/12.f) ) )
-			.SetFillImage( BOX_BRUSH( "Common/ProgressBar_Fill", FMargin(5.f/12.f) ) )
+			.SetFillImage( BOX_BRUSH( "Common/ProgressBar_Fill", FMargin(5.f/12.f), FLinearColor( 1.0f, 0.22f, 0.0f ) ) )
 			.SetMarqueeImage( IMAGE_BRUSH( "Common/ProgressBar_Marquee", FVector2D(20,12), FLinearColor::White, ESlateBrushTileType::Horizontal ) )
 			);
 
@@ -501,7 +501,7 @@ void FSlateEditorStyle::FStyle::SetupGeneralStyles()
 	{
 		Set( "WorkingBar", FProgressBarStyle()
 			.SetBackgroundImage( FSlateNoResource() )
-			.SetFillImage( BOX_BRUSH( "Common/ProgressBar_Fill", FMargin(5.f/12.f) ) )
+			.SetFillImage( BOX_BRUSH( "Common/ProgressBar_Fill", FMargin(5.f/12.f), FLinearColor( 1.0f, 0.22f, 0.0f ) ) )
 			.SetMarqueeImage( IMAGE_BRUSH( "Common/WorkingBar_Marquee", FVector2D(20,2), FLinearColor::White, ESlateBrushTileType::Horizontal ) )
 			);
 	}
@@ -1428,10 +1428,10 @@ void FSlateEditorStyle::FStyle::SetupGeneralStyles()
 
 		Set("HardwareTargets.Normal", TargetSettingsNormal);
 
- 		Set("HardwareTargets.Strong", FTextBlockStyle(TargetSettingsNormal)
- 			.SetFont(TTF_CORE_FONT("Fonts/Roboto-Bold", 8))
+		Set("HardwareTargets.Strong", FTextBlockStyle(TargetSettingsNormal)
+			.SetFont(TTF_CORE_FONT("Fonts/Roboto-Bold", 8))
 			.SetColorAndOpacity(EditorOrange)
- 			.SetShadowOffset(FVector2D::UnitVector));
+			.SetShadowOffset(FVector2D::UnitVector));
 	}
 #endif
 
@@ -5894,26 +5894,6 @@ void FSlateEditorStyle::FStyle::SetupUMGEditorStyles()
 
 	Set("UMGEditor.AnchoredWidget", new BOX_BRUSH("Common/Button", FVector2D(32, 32), 8.0f / 32.0f));
 	Set("UMGEditor.AnchoredWidgetAlignment", new IMAGE_BRUSH("Icons/icon_tab_DeviceManager_16x", Icon8x8));
-
-	Set("UMGEditor.TopLeft", new IMAGE_BRUSH("Icons/UMG/umg_anchors_top_left", Icon64x64));
-	Set("UMGEditor.TopCenter", new IMAGE_BRUSH("Icons/UMG/umg_anchors_top_center", Icon64x64));
-	Set("UMGEditor.TopRight", new IMAGE_BRUSH("Icons/UMG/umg_anchors_top_right", Icon64x64));
-	Set("UMGEditor.TopFill", new IMAGE_BRUSH("Icons/UMG/umg_anchors_top_fill", Icon64x64));
-
-	Set("UMGEditor.CenterLeft", new IMAGE_BRUSH("Icons/UMG/umg_anchors_center_left", Icon64x64));
-	Set("UMGEditor.CenterCenter", new IMAGE_BRUSH("Icons/UMG/umg_anchors_center_center", Icon64x64));
-	Set("UMGEditor.CenterRight", new IMAGE_BRUSH("Icons/UMG/umg_anchors_center_right", Icon64x64));
-	Set("UMGEditor.CenterFill", new IMAGE_BRUSH("Icons/UMG/umg_anchors_center_fill", Icon64x64));
-
-	Set("UMGEditor.BottomLeft", new IMAGE_BRUSH("Icons/UMG/umg_anchors_bottom_left", Icon64x64));
-	Set("UMGEditor.BottomCenter", new IMAGE_BRUSH("Icons/UMG/umg_anchors_bottom_center", Icon64x64));
-	Set("UMGEditor.BottomRight", new IMAGE_BRUSH("Icons/UMG/umg_anchors_bottom_right", Icon64x64));
-	Set("UMGEditor.BottomFill", new IMAGE_BRUSH("Icons/UMG/umg_anchors_bottom_fill", Icon64x64));
-
-	Set("UMGEditor.FillLeft", new IMAGE_BRUSH("Icons/UMG/umg_anchors_fill_left", Icon64x64));
-	Set("UMGEditor.FillCenter", new IMAGE_BRUSH("Icons/UMG/umg_anchors_fill_center", Icon64x64));
-	Set("UMGEditor.FillRight", new IMAGE_BRUSH("Icons/UMG/umg_anchors_fill_right", Icon64x64));
-	Set("UMGEditor.FillFill", new IMAGE_BRUSH("Icons/UMG/umg_anchors_fill_fill", Icon64x64));
 	
 	Set("UMGEditor.PalettePreviewWidget", new IMAGE_BRUSH("Icons/UMG/umg_anchors_fill_fill", Icon16x16));
 
@@ -5957,7 +5937,12 @@ void FSlateEditorStyle::FStyle::SetupUMGEditorStyles()
 	Set("UMGEditor.AddAnimationIcon", new IMAGE_BRUSH("Icons/PlusSymbol_12x", Icon12x12, FLinearColor(.05,.05,.05) ) );
 	Set("UMGEditor.NoAnimationFont", NoAnimationFont);
 
+	Set("UMGEditor.SwitchToDesigner", new IMAGE_BRUSH("UMG/Designer_40x", Icon40x40));
+	Set("UMGEditor.SwitchToDesigner.Small", new IMAGE_BRUSH("UMG/Designer_16x", Icon16x16));
+
 	Set("UMGEditor.SelectionOutline", new BORDER_BRUSH("UMG/SelectionBorder", FMargin(.2f) ));
+
+	Set("UMGEditor.AnchorGrid", new IMAGE_BRUSH("Icons/UMG/AnchorGrid", Icon10x10, FLinearColor::White, ESlateBrushTileType::Both ));
 }
 
 void FSlateEditorStyle::FStyle::SetupTranslationEditorStyles()
diff --git a/Engine/Source/Editor/UMGEditor/Private/Designer/SDesignerView.cpp b/Engine/Source/Editor/UMGEditor/Private/Designer/SDesignerView.cpp
--- a/Engine/Source/Editor/UMGEditor/Private/Designer/SDesignerView.cpp
+++ b/Engine/Source/Editor/UMGEditor/Private/Designer/SDesignerView.cpp
@@ -159,20 +159,43 @@ void SDesignerView::Construct(const FArguments& InArgs, TSharedPtr<FWidgetBluepr
 				.ZoomAmount(this, &SDesignerView::GetZoomAmount)
 				.ViewOffset(this, &SDesignerView::GetViewOffset)
 				[
-					SNew(SBox)
-					.WidthOverride(this, &SDesignerView::GetPreviewWidth)
-					.HeightOverride(this, &SDesignerView::GetPreviewHeight)
-					.HAlign(HAlign_Fill)
-					.VAlign(VAlign_Fill)
-					.Visibility(EVisibility::SelfHitTestInvisible)
+					SNew(SOverlay)
+
+					+ SOverlay::Slot()
+					[
+						SNew(SBorder)
+						[
+							SNew(SSpacer)
+							.Size(FVector2D(1, 1))
+						]
+					]
+
+					+ SOverlay::Slot()
 					[
-						SAssignNew(PreviewSurface, SDPIScaler)
-						.DPIScale(this, &SDesignerView::GetPreviewDPIScale)
+						SNew(SBox)
+						.WidthOverride(this, &SDesignerView::GetPreviewWidth)
+						.HeightOverride(this, &SDesignerView::GetPreviewHeight)
+						.HAlign(HAlign_Fill)
+						.VAlign(VAlign_Fill)
 						.Visibility(EVisibility::SelfHitTestInvisible)
+						[
+							SAssignNew(PreviewSurface, SDPIScaler)
+							.DPIScale(this, &SDesignerView::GetPreviewDPIScale)
+							.Visibility(EVisibility::SelfHitTestInvisible)
+						]
 					]
 				]
 			]
 
+			// A layer in the overlay where we put all the user intractable widgets, like the reorder widgets.
+			+ SOverlay::Slot()
+			.HAlign(HAlign_Fill)
+			.VAlign(VAlign_Fill)
+			[
+				SAssignNew(EffectsLayer, SPaintSurface)
+				.OnPaintHandler(this, &SDesignerView::HandleEffectsPainting)
+			]
+
 			// A layer in the overlay where we put all the user intractable widgets, like the reorder widgets.
 			+ SOverlay::Slot()
 			.HAlign(HAlign_Fill)
@@ -245,6 +268,50 @@ void SDesignerView::Construct(const FArguments& InArgs, TSharedPtr<FWidgetBluepr
 						.TextStyle(FEditorStyle::Get(), "ViewportMenu.Label")
 					]
 				]
+
+				+ SHorizontalBox::Slot()
+				.AutoWidth()
+				.Padding(5.0f)
+				.VAlign(VAlign_Center)
+				[
+					SNew(SNumericEntryBox<int32>)
+					.AllowSpin(true)
+					.Delta(5)
+					.MinSliderValue(1)
+					.MinValue(1)
+					.MaxSliderValue(TOptional<int32>(1000))
+					.Value(this, &SDesignerView::GetCustomResolutionWidth)
+					.OnValueChanged(this, &SDesignerView::OnCustomResolutionWidthChanged)
+					.Visibility(this, &SDesignerView::GetCustomResolutionEntryVisibility)
+					.MinDesiredValueWidth(50)
+					.LabelPadding(0)
+					.Label()
+					[
+						SNumericEntryBox<int32>::BuildLabel(LOCTEXT("Width", "Width"), FLinearColor::White, SNumericEntryBox<int32>::RedLabelBackgroundColor)
+					]
+				]
+
+				+ SHorizontalBox::Slot()
+				.AutoWidth()
+				.Padding(5.0f)
+				.VAlign(VAlign_Center)
+				[
+					SNew(SNumericEntryBox<int32>)
+					.AllowSpin(true)
+					.Delta(5)
+					.MinSliderValue(1)
+					.MaxSliderValue(TOptional<int32>(1000))
+					.MinValue(1)
+					.Value(this, &SDesignerView::GetCustomResolutionHeight)
+					.OnValueChanged(this, &SDesignerView::OnCustomResolutionHeightChanged)
+					.Visibility(this, &SDesignerView::GetCustomResolutionEntryVisibility)
+					.MinDesiredValueWidth(50)
+					.LabelPadding(0)
+					.Label()
+					[
+						SNumericEntryBox<int32>::BuildLabel(LOCTEXT("Height", "Height"), FLinearColor::White, SNumericEntryBox<int32>::GreenLabelBackgroundColor)
+					]
+				]
 			]
 
 			// Bottom bar to show current resolution & AR
@@ -328,6 +395,14 @@ void SDesignerView::BindCommands()
 		);
 }
 
+void SDesignerView::AddReferencedObjects(FReferenceCollector& Collector)
+{
+	if ( PreviewWidget )
+	{
+		Collector.AddReferencedObject(PreviewWidget);
+	}
+}
+
 void SDesignerView::SetTransformMode(ETransformMode::Type InTransformMode)
 {
 	if ( !InTransaction() )
@@ -386,11 +461,27 @@ ETransformMode::Type SDesignerView::GetTransformMode() const
 
 FOptionalSize SDesignerView::GetPreviewWidth() const
 {
+	if ( UUserWidget* DefaultWidget = GetDefaultWidget() )
+	{
+		if ( DefaultWidget->bUseDesignTimeSize )
+		{
+			return DefaultWidget->DesignTimeSize.X;
+		}
+	}
+
 	return (float)PreviewWidth;
 }
 
 FOptionalSize SDesignerView::GetPreviewHeight() const
 {
+	if ( UUserWidget* DefaultWidget = GetDefaultWidget() )
+	{
+		if ( DefaultWidget->bUseDesignTimeSize )
+		{
+			return DefaultWidget->DesignTimeSize.Y;
+		}
+	}
+
 	return (float)PreviewHeight;
 }
 
@@ -401,7 +492,7 @@ float SDesignerView::GetPreviewDPIScale() const
 
 FSlateRect SDesignerView::ComputeAreaBounds() const
 {
-	return FSlateRect(0, 0, PreviewWidth, PreviewHeight);
+	return FSlateRect(0, 0, GetPreviewWidth().Get(), GetPreviewHeight().Get());
 }
 
 EVisibility SDesignerView::GetInfoBarVisibility() const
@@ -470,6 +561,18 @@ FGeometry SDesignerView::GetDesignerGeometry() const
 	return CachedDesignerGeometry;
 }
 
+void SDesignerView::MarkDesignModifed(bool bRequiresRecompile)
+{
+	if ( bRequiresRecompile )
+	{
+		FBlueprintEditorUtils::MarkBlueprintAsStructurallyModified(GetBlueprint());
+	}
+	else
+	{
+		FBlueprintEditorUtils::MarkBlueprintAsModified(GetBlueprint());
+	}
+}
+
 bool SDesignerView::GetWidgetParentGeometry(const FWidgetReference& Widget, FGeometry& Geometry) const
 {
 	if ( UWidget* PreviewWidget = Widget.GetPreview() )
@@ -561,7 +664,7 @@ void SDesignerView::CreateExtensionWidgetsForSelection()
 FVector2D SDesignerView::GetExtensionPosition(TSharedRef<FDesignerSurfaceElement> ExtensionElement) const
 {
 	const FVector2D TopLeft = CachedDesignerWidgetLocation;
-	const FVector2D Size = CachedDesignerWidgetSize * GetZoomAmount() * GetPreviewDPIScale();
+	const FVector2D Size = CachedDesignerWidgetSize * GetPreviewScale();
 
 	// Calculate the parent position and size.  We use this information for calculating offsets.
 	FVector2D ParentPosition, ParentSize;
@@ -579,7 +682,7 @@ FVector2D SDesignerView::GetExtensionPosition(TSharedRef<FDesignerSurfaceElement
 			FDesignTimeUtils::GetArrangedWidgetRelativeToParent(WidgetPath, PreviewSlateWidget.ToSharedRef(), AsShared(), ArrangedWidget);
 
 			ParentPosition = ArrangedWidget.Geometry.AbsolutePosition;
-			ParentSize = ArrangedWidget.Geometry.Size * GetZoomAmount() * GetPreviewDPIScale();
+			ParentSize = ArrangedWidget.Geometry.Size * GetPreviewScale();
 		}
 	}
 
@@ -934,22 +1037,21 @@ void SDesignerView::CacheSelectedWidgetGeometry()
 	}
 }
 
-int32 SDesignerView::OnPaint(const FPaintArgs& Args, const FGeometry& AllottedGeometry, const FSlateRect& MyClippingRect, FSlateWindowElementList& OutDrawElements, int32 LayerId, const FWidgetStyle& InWidgetStyle, bool bParentEnabled) const
+int32 SDesignerView::HandleEffectsPainting(const FOnPaintHandlerParams& PaintArgs)
 {
-	SDesignSurface::OnPaint(Args, AllottedGeometry, MyClippingRect, OutDrawElements, LayerId, InWidgetStyle, bParentEnabled);
-
-	LayerId += 1000;
-
 	TSet<FWidgetReference> Selected;
 	Selected.Add(SelectedWidget);
 
 	// Allow the extensions to paint anything they want.
 	for ( const TSharedRef<FDesignerExtension>& Ext : DesignerExtensions )
 	{
-		Ext->Paint(Selected, AllottedGeometry, MyClippingRect, OutDrawElements, LayerId, InWidgetStyle, bParentEnabled);
+		Ext->Paint(Selected, PaintArgs.Geometry, PaintArgs.ClippingRect, PaintArgs.OutDrawElements, PaintArgs.Layer);
 	}
 
 	static const FName SelectionOutlineName("UMGEditor.SelectionOutline");
+	const FSlateBrush* SelectionOutlineBrush = FEditorStyle::Get().GetBrush(SelectionOutlineName);
+	FVector2D SelectionBrushInflationAmount = FVector2D(16, 16) * FVector2D(SelectionOutlineBrush->Margin.Left, SelectionOutlineBrush->Margin.Top) * ( 1.0f / GetPreviewScale() );
+
 	// Don't draw the hovered effect if it's also the selected widget
 	if ( HoveredSlateWidget.IsValid() && HoveredSlateWidget != SelectedSlateWidget )
 	{
@@ -962,14 +1064,14 @@ int32 SDesignerView::OnPaint(const FPaintArgs& Args, const FGeometry& AllottedGe
 		// Azure = 0x007FFF
 		const FLinearColor HoveredTint(0, 0.5, 1, FMath::Clamp(HoverTime / HoveredAnimationTime, 0.0f, 1.0f));
 
-		FPaintGeometry HoveredGeometry = ArrangedWidget.Geometry.ToPaintGeometry();
+		FPaintGeometry HoveredGeometry = ArrangedWidget.Geometry.ToInflatedPaintGeometry(SelectionBrushInflationAmount);
 
 		FSlateDrawElement::MakeBox(
-			OutDrawElements,
-			LayerId,
+			PaintArgs.OutDrawElements,
+			PaintArgs.Layer,
 			HoveredGeometry,
-			FEditorStyle::Get().GetBrush(SelectionOutlineName),
-			MyClippingRect,
+			SelectionOutlineBrush,
+			PaintArgs.ClippingRect,
 			ESlateDrawEffect::None,
 			HoveredTint
 			);
@@ -985,21 +1087,21 @@ int32 SDesignerView::OnPaint(const FPaintArgs& Args, const FGeometry& AllottedGe
 		const FLinearColor Tint(0, 1, 0);
 
 		// Draw selection effect
-		
-		FPaintGeometry SelectionGeometry = ArrangedWidget.Geometry.ToPaintGeometry();
+
+		FPaintGeometry SelectionGeometry = ArrangedWidget.Geometry.ToInflatedPaintGeometry(SelectionBrushInflationAmount);
 
 		FSlateDrawElement::MakeBox(
-			OutDrawElements,
-			LayerId,
+			PaintArgs.OutDrawElements,
+			PaintArgs.Layer,
 			SelectionGeometry,
-			FEditorStyle::Get().GetBrush(SelectionOutlineName),
-			MyClippingRect,
+			SelectionOutlineBrush,
+			PaintArgs.ClippingRect,
 			ESlateDrawEffect::None,
 			Tint
-		);
+			);
 	}
 
-	return LayerId;
+	return PaintArgs.Layer + 1;
 }
 
 void SDesignerView::UpdatePreviewWidget(bool bForceUpdate)
@@ -1059,7 +1161,7 @@ void SDesignerView::Tick(const FGeometry& AllottedGeometry, const double InCurre
 		if ( SelectedWidget.IsValid() )
 		{
 			// Set the selected widget so that we can draw the highlight
-			SelectedSlateWidget = PreviewWidget->GetWidgetFromName(SelectedWidget.GetTemplate()->GetName());
+			SelectedSlateWidget = PreviewWidget->GetWidgetFromName(SelectedWidget.GetTemplate()->GetFName());
 		}
 		else
 		{
@@ -1068,7 +1170,7 @@ void SDesignerView::Tick(const FGeometry& AllottedGeometry, const double InCurre
 
 		if ( HoveredWidget.IsValid() )
 		{
-			HoveredSlateWidget = PreviewWidget->GetWidgetFromName(HoveredWidget.GetTemplate()->GetName());
+			HoveredSlateWidget = PreviewWidget->GetWidgetFromName(HoveredWidget.GetTemplate()->GetFName());
 		}
 		else
 		{
@@ -1503,12 +1605,28 @@ void SDesignerView::HandleOnCommonResolutionSelected(int32 Width, int32 Height,
 	GConfig->SetInt(*ConfigSectionName, TEXT("PreviewHeight"), Height, GEditorUserSettingsIni);
 	GConfig->SetString(*ConfigSectionName, TEXT("PreviewAspectRatio"), *AspectRatio, GEditorUserSettingsIni);
 
+	// We're no longer using a custom design time size.
+	if ( UUserWidget* DefaultWidget = GetDefaultWidget() )
+	{
+		DefaultWidget->bUseDesignTimeSize = false;
+		MarkDesignModifed(/*bRequiresRecompile*/ false);
+	}
+
 	ResolutionTextFade.Play();
 }
 
 bool SDesignerView::HandleIsCommonResolutionSelected(int32 Width, int32 Height) const
 {
-	return (Width == PreviewWidth) && (Height == PreviewHeight);
+	// If we're using a custom design time size, none of the other resolutions should appear selected, even if they match.
+	if ( UUserWidget* DefaultWidget = GetDefaultWidget() )
+	{
+		if ( DefaultWidget->bUseDesignTimeSize )
+		{
+			return false;
+		}
+	}
+	
+	return ( Width == PreviewWidth ) && ( Height == PreviewHeight );
 }
 
 void SDesignerView::AddScreenResolutionSection(FMenuBuilder& MenuBuilder, const TArray<FPlayScreenResolution>& Resolutions, const FText& SectionName)
@@ -1519,9 +1637,8 @@ void SDesignerView::AddScreenResolutionSection(FMenuBuilder& MenuBuilder, const
 		{
 			// Actions for the resolution menu entry
 			FExecuteAction OnResolutionSelected = FExecuteAction::CreateRaw(this, &SDesignerView::HandleOnCommonResolutionSelected, Iter->Width, Iter->Height, Iter->AspectRatio);
-			FCanExecuteAction OnCanResolutionBeSelected = FCanExecuteAction::CreateRaw(&FSlateApplication::Get(), &FSlateApplication::IsNormalExecution);
 			FIsActionChecked OnIsResolutionSelected = FIsActionChecked::CreateRaw(this, &SDesignerView::HandleIsCommonResolutionSelected, Iter->Width, Iter->Height);
-			FUIAction Action(OnResolutionSelected, OnCanResolutionBeSelected, OnIsResolutionSelected);
+			FUIAction Action(OnResolutionSelected, FCanExecuteAction(), OnIsResolutionSelected);
 
 			MenuBuilder.AddMenuEntry(FText::FromString(Iter->Description), GetResolutionText(Iter->Width, Iter->Height, Iter->AspectRatio), FSlateIcon(), Action, NAME_None, EUserInterfaceActionType::Check);
 		}
@@ -1529,11 +1646,97 @@ void SDesignerView::AddScreenResolutionSection(FMenuBuilder& MenuBuilder, const
 	MenuBuilder.EndSection();
 }
 
+bool SDesignerView::HandleIsCustomResolutionSelected() const
+{
+	if ( UUserWidget* DefaultWidget = GetDefaultWidget() )
+	{
+		return DefaultWidget->bUseDesignTimeSize;
+	}
+
+	return false;
+}
+
+void SDesignerView::HandleOnCustomResolutionSelected()
+{
+	if ( UUserWidget* DefaultWidget = GetDefaultWidget() )
+	{
+		DefaultWidget->bUseDesignTimeSize = true;
+		MarkDesignModifed(/*bRequiresRecompile*/ false);
+	}
+}
+
+TOptional<int32> SDesignerView::GetCustomResolutionWidth() const
+{
+	if ( UUserWidget* DefaultWidget = GetDefaultWidget() )
+	{
+		return DefaultWidget->DesignTimeSize.X;
+	}
+
+	return 1;
+}
+
+TOptional<int32> SDesignerView::GetCustomResolutionHeight() const
+{
+	if ( UUserWidget* DefaultWidget = GetDefaultWidget() )
+	{
+		return DefaultWidget->DesignTimeSize.Y;
+	}
+
+	return 1;
+}
+
+void SDesignerView::OnCustomResolutionWidthChanged(int32 InValue)
+{
+	if ( UUserWidget* DefaultWidget = GetDefaultWidget() )
+	{
+		DefaultWidget->DesignTimeSize.X = InValue;
+		MarkDesignModifed(/*bRequiresRecompile*/ false);
+	}
+}
+
+void SDesignerView::OnCustomResolutionHeightChanged(int32 InValue)
+{
+	if ( UUserWidget* DefaultWidget = GetDefaultWidget() )
+	{
+		DefaultWidget->DesignTimeSize.Y = InValue;
+		MarkDesignModifed(/*bRequiresRecompile*/ false);
+	}
+}
+
+EVisibility SDesignerView::GetCustomResolutionEntryVisibility() const
+{
+	if ( UUserWidget* DefaultWidget = GetDefaultWidget() )
+	{
+		return DefaultWidget->bUseDesignTimeSize ? EVisibility::Visible : EVisibility::Collapsed;
+	}
+
+	return EVisibility::Collapsed;
+}
+
+UUserWidget* SDesignerView::GetDefaultWidget() const
+{
+	TSharedPtr<FWidgetBlueprintEditor> BPEd = BlueprintEditor.Pin();
+	if ( UUserWidget* Default = BPEd->GetWidgetBlueprintObj()->GeneratedClass->GetDefaultObject<UUserWidget>() )
+	{
+		return Default;
+	}
+
+	return nullptr;
+}
+
 TSharedRef<SWidget> SDesignerView::GetAspectMenu()
 {
 	const ULevelEditorPlaySettings* PlaySettings = GetDefault<ULevelEditorPlaySettings>();
 	FMenuBuilder MenuBuilder(true, NULL);
 
+	// Add custom option
+	FExecuteAction OnResolutionSelected = FExecuteAction::CreateRaw(this, &SDesignerView::HandleOnCustomResolutionSelected);
+	FIsActionChecked OnIsResolutionSelected = FIsActionChecked::CreateRaw(this, &SDesignerView::HandleIsCustomResolutionSelected);
+	FUIAction Action(OnResolutionSelected, FCanExecuteAction(), OnIsResolutionSelected);
+
+	MenuBuilder.AddMenuEntry(LOCTEXT("Custom", "Custom"), LOCTEXT("Custom", "Custom"), FSlateIcon(), Action, NAME_None, EUserInterfaceActionType::Check);
+
+	// Add the normal set of resultion options.
 	AddScreenResolutionSection(MenuBuilder, PlaySettings->PhoneScreenResolutions, LOCTEXT("CommonPhonesSectionHeader", "Phones"));
 	AddScreenResolutionSection(MenuBuilder, PlaySettings->TabletScreenResolutions, LOCTEXT("CommonTabletsSectionHeader", "Tablets"));
 	AddScreenResolutionSection(MenuBuilder, PlaySettings->LaptopScreenResolutions, LOCTEXT("CommonLaptopsSectionHeader", "Laptops"));
diff --git a/Engine/Source/Editor/UMGEditor/Private/Designer/SDesignerView.h b/Engine/Source/Editor/UMGEditor/Private/Designer/SDesignerView.h
--- a/Engine/Source/Editor/UMGEditor/Private/Designer/SDesignerView.h
+++ b/Engine/Source/Editor/UMGEditor/Private/Designer/SDesignerView.h
@@ -8,6 +8,8 @@
 #include "DesignerExtension.h"
 #include "IUMGDesigner.h"
 
+#include "SPaintSurface.h"
+
 namespace EDesignerMessage
 {
 	enum Type
@@ -22,7 +24,7 @@ class FDesignerExtension;
 /**
  * The designer for widgets.  Allows for laying out widgets in a drag and drop environment.
  */
-class SDesignerView : public SDesignSurface, public IUMGDesigner
+class SDesignerView : public SDesignSurface, public FGCObject, public IUMGDesigner
 {
 public:
 
@@ -41,9 +43,6 @@ public:
 
 	virtual FReply OnKeyDown(const FGeometry& MyGeometry, const FKeyboardEvent& InKeyboardEvent) override;
 
-	//virtual void OnArrangeChildren(const FGeometry& AllottedGeometry, FArrangedChildren& ArrangedChildren) const;
-
-	virtual int32 OnPaint(const FPaintArgs& Args, const FGeometry& AllottedGeometry, const FSlateRect& MyClippingRect, FSlateWindowElementList& OutDrawElements, int32 LayerId, const FWidgetStyle& InWidgetStyle, bool bParentEnabled) const override;
 	virtual void Tick(const FGeometry& AllottedGeometry, const double InCurrentTime, const float InDeltaTime) override;
 
 	virtual FReply OnDragDetected(const FGeometry& MyGeometry, const FPointerEvent& MouseEvent) override;
@@ -62,8 +61,13 @@ public:
 	virtual FGeometry GetDesignerGeometry() const override;
 	virtual bool GetWidgetGeometry(const FWidgetReference& Widget, FGeometry& Geometry) const override;
 	virtual bool GetWidgetParentGeometry(const FWidgetReference& Widget, FGeometry& Geometry) const override;
+	virtual void MarkDesignModifed(bool bRequiresRecompile) override;
 	// End of IUMGDesigner interface
 
+	// FGCObject interface
+	virtual void AddReferencedObjects(FReferenceCollector& Collector) override;
+	// End of FGCObject interface
+
 private:
 	/** Establishes the resolution and aspect ratio to use on construction from config settings */
 	void SetStartupResolution();
@@ -114,13 +118,26 @@ private:
 
 	FText GetCurrentResolutionText() const;
 	FSlateColor GetResolutionTextColorAndOpacity() const;
+
+	TOptional<int32> GetCustomResolutionWidth() const;
+	TOptional<int32> GetCustomResolutionHeight() const;
+	void OnCustomResolutionWidthChanged(int32 InValue);
+	void OnCustomResolutionHeightChanged(int32 InValue);
+	EVisibility GetCustomResolutionEntryVisibility() const;
 	
 	// Handles selecting a common screen resolution.
 	void HandleOnCommonResolutionSelected(int32 Width, int32 Height, FString AspectRatio);
 	bool HandleIsCommonResolutionSelected(int32 Width, int32 Height) const;
 	void AddScreenResolutionSection(FMenuBuilder& MenuBuilder, const TArray<FPlayScreenResolution>& Resolutions, const FText& SectionName);
+	bool HandleIsCustomResolutionSelected() const;
+	void HandleOnCustomResolutionSelected();
 	TSharedRef<SWidget> GetAspectMenu();
 
+	// Handles drawing selection and other effects a SPaintSurface widget injected into the hierarchy.
+	int32 HandleEffectsPainting(const FOnPaintHandlerParams& PaintArgs);
+
+	UUserWidget* GetDefaultWidget() const;
+
 	void BeginTransaction(const FText& SessionName);
 	bool InTransaction() const;
 	void EndTransaction(bool bCancel);
@@ -172,6 +189,7 @@ private:
 	TSharedPtr<class SZoomPan> PreviewHitTestRoot;
 	TSharedPtr<SDPIScaler> PreviewSurface;
 	TSharedPtr<SCanvas> ExtensionWidgetCanvas;
+	TSharedPtr<SPaintSurface> EffectsLayer;
 
 	FVector2D CachedDesignerWidgetLocation;
 	FVector2D CachedDesignerWidgetSize;
diff --git a/Engine/Source/Editor/UMGEditor/Private/Designer/SPaintSurface.h b/Engine/Source/Editor/UMGEditor/Private/Designer/SPaintSurface.h
new file mode 100644
--- /dev/null
+++ b/Engine/Source/Editor/UMGEditor/Private/Designer/SPaintSurface.h
@@ -0,0 +1,85 @@
+// Copyright 1998-2014 Epic Games, Inc. All Rights Reserved.
+
+#pragma once
+
+#include "SCompoundWidget.h"
+
+struct FOnPaintHandlerParams
+{
+	const FGeometry& Geometry;
+	const FSlateRect& ClippingRect;
+	FSlateWindowElementList& OutDrawElements;
+	const int32 Layer;
+	const bool bEnabled;
+
+	FOnPaintHandlerParams( const FGeometry& InGeometry, const FSlateRect& InClippingRect, FSlateWindowElementList& InOutDrawElements, int32 InLayer, bool bInEnabled )
+		: Geometry( InGeometry )
+		, ClippingRect( InClippingRect )
+		, OutDrawElements( InOutDrawElements )
+		, Layer( InLayer )
+		, bEnabled( bInEnabled )
+	{
+	}
+
+};
+
+/** Delegate type for allowing custom OnPaint handlers */
+DECLARE_DELEGATE_RetVal_OneParam( 
+	int32,
+	FOnPaintHandler,
+	const FOnPaintHandlerParams& );
+
+/**
+ * Widget with a handler for OnPaint, allows the designer to insert painting on different layers of the overlay when drawing widgets and effects
+ * intermixed.
+ */
+class SPaintSurface : public SCompoundWidget
+{
+public:
+	SLATE_BEGIN_ARGS( SPaintSurface )
+		: _OnPaintHandler()
+	{
+		_Visibility = EVisibility::SelfHitTestInvisible;
+	}
+
+		SLATE_EVENT( FOnPaintHandler, OnPaintHandler )
+	SLATE_END_ARGS()
+
+	/**
+	 * Construct this widget
+	 *
+	 * @param	InArgs	The declaration data for this widget
+	 */
+	void Construct(const FArguments& InArgs)
+	{
+		OnPaintHandler = InArgs._OnPaintHandler;
+	}
+
+	virtual FVector2D ComputeDesiredSize() const override
+	{
+		return FVector2D(128, 128);
+	}
+
+	virtual int32 OnPaint( const FPaintArgs& Args, const FGeometry& AllottedGeometry, const FSlateRect& MyClippingRect, FSlateWindowElementList& OutDrawElements, int32 LayerId, const FWidgetStyle& InWidgetStyle, bool bParentEnabled ) const
+	{
+		if( OnPaintHandler.IsBound() )
+		{
+			FOnPaintHandlerParams Params( AllottedGeometry, MyClippingRect, OutDrawElements, LayerId, bParentEnabled && IsEnabled() ); 
+			OnPaintHandler.Execute( Params );
+		}
+		else
+		{
+			FSlateDrawElement::MakeDebugQuad(
+				OutDrawElements,
+				LayerId,
+				AllottedGeometry.ToPaintGeometry(),
+				MyClippingRect
+			);
+		}
+
+		return SCompoundWidget::OnPaint( Args, AllottedGeometry, MyClippingRect, OutDrawElements, LayerId, InWidgetStyle, bParentEnabled && IsEnabled() );
+	}
+
+private:
+	FOnPaintHandler OnPaintHandler;
+};
\ No newline at end of file
diff --git a/Engine/Source/Editor/UMGEditor/Private/Designer/STransformHandle.cpp b/Engine/Source/Editor/UMGEditor/Private/Designer/STransformHandle.cpp
--- a/Engine/Source/Editor/UMGEditor/Private/Designer/STransformHandle.cpp
+++ b/Engine/Source/Editor/UMGEditor/Private/Designer/STransformHandle.cpp
@@ -74,6 +74,9 @@ FReply STransformHandle::OnMouseButtonUp(const FGeometry& MyGeometry, const FPoi
 {
 	if ( HasMouseCapture() && MouseEvent.GetEffectingButton() == EKeys::LeftMouseButton )
 	{
+		const bool bRequiresRecompile = false;
+		Designer->MarkDesignModifed(bRequiresRecompile);
+
 		Action = ETransformAction::None;
 		return FReply::Handled().ReleaseMouseCapture();
 	}
diff --git a/Engine/Source/Editor/UMGEditor/Private/Details/SWidgetDetailsView.cpp b/Engine/Source/Editor/UMGEditor/Private/Details/SWidgetDetailsView.cpp
--- a/Engine/Source/Editor/UMGEditor/Private/Details/SWidgetDetailsView.cpp
+++ b/Engine/Source/Editor/UMGEditor/Private/Details/SWidgetDetailsView.cpp
@@ -268,7 +268,7 @@ bool SWidgetDetailsView::HandleVerifyNameTextChanged(const FText& InText, FText&
 		FString NewName = InText.ToString();
 
 		UWidgetBlueprint* Blueprint = BlueprintEditor.Pin()->GetWidgetBlueprintObj();
-		UWidget* TemplateWidget = Blueprint->WidgetTree->FindWidget(NewName);
+		UWidget* TemplateWidget = Blueprint->WidgetTree->FindWidget( FName(*NewName) );
 
 		bool bIsSameWidget = false;
 		if ( TemplateWidget != NULL )
@@ -345,7 +345,7 @@ void SWidgetDetailsView::HandleIsVariableChanged(ESlateCheckBoxState::Type Check
 		UWidget* Widget = Cast<UWidget>(SelectedObjects[0].Get());
 		UWidgetBlueprint* Blueprint = BlueprintEditor.Pin()->GetWidgetBlueprintObj();
 		
-		FWidgetReference WidgetRef = BPEditor->GetReferenceFromTemplate(Blueprint->WidgetTree->FindWidget(Widget->GetName()));
+		FWidgetReference WidgetRef = BPEditor->GetReferenceFromTemplate(Blueprint->WidgetTree->FindWidget(Widget->GetFName()));
 		if ( WidgetRef.IsValid() )
 		{
 			UWidget* Template = WidgetRef.GetTemplate();
diff --git a/Engine/Source/Editor/UMGEditor/Private/Extensions/CanvasSlotExtension.cpp b/Engine/Source/Editor/UMGEditor/Private/Extensions/CanvasSlotExtension.cpp
--- a/Engine/Source/Editor/UMGEditor/Private/Extensions/CanvasSlotExtension.cpp
+++ b/Engine/Source/Editor/UMGEditor/Private/Extensions/CanvasSlotExtension.cpp
@@ -240,9 +240,9 @@ void FCanvasSlotExtension::Tick(const FGeometry& AllottedGeometry, const double
 
 }
 
-void FCanvasSlotExtension::Paint(const TSet< FWidgetReference >& Selection, const FGeometry& AllottedGeometry, const FSlateRect& MyClippingRect, FSlateWindowElementList& OutDrawElements, int32 LayerId, const FWidgetStyle& InWidgetStyle, bool bParentEnabled) const
+void FCanvasSlotExtension::Paint(const TSet< FWidgetReference >& Selection, const FGeometry& AllottedGeometry, const FSlateRect& MyClippingRect, FSlateWindowElementList& OutDrawElements, int32 LayerId) const
 {
-	PaintCollisionLines(Selection, AllottedGeometry, MyClippingRect, OutDrawElements, LayerId, InWidgetStyle, bParentEnabled);
+	PaintCollisionLines(Selection, AllottedGeometry, MyClippingRect, OutDrawElements, LayerId);
 }
 
 FReply FCanvasSlotExtension::HandleAnchorBeginDrag(const FGeometry& Geometry, const FPointerEvent& Event, EAnchorWidget::Type AnchorType)
@@ -416,7 +416,7 @@ FReply FCanvasSlotExtension::HandleAnchorDragging(const FGeometry& Geometry, con
 	return FReply::Unhandled();
 }
 
-void FCanvasSlotExtension::PaintCollisionLines(const TSet< FWidgetReference >& Selection, const FGeometry& AllottedGeometry, const FSlateRect& MyClippingRect, FSlateWindowElementList& OutDrawElements, int32 LayerId, const FWidgetStyle& InWidgetStyle, bool bParentEnabled) const
+void FCanvasSlotExtension::PaintCollisionLines(const TSet< FWidgetReference >& Selection, const FGeometry& AllottedGeometry, const FSlateRect& MyClippingRect, FSlateWindowElementList& OutDrawElements, int32 LayerId) const
 {
 	for ( const FWidgetReference& WidgetRef : Selection )
 	{
diff --git a/Engine/Source/Editor/UMGEditor/Private/Extensions/CanvasSlotExtension.h b/Engine/Source/Editor/UMGEditor/Private/Extensions/CanvasSlotExtension.h
--- a/Engine/Source/Editor/UMGEditor/Private/Extensions/CanvasSlotExtension.h
+++ b/Engine/Source/Editor/UMGEditor/Private/Extensions/CanvasSlotExtension.h
@@ -37,7 +37,7 @@ public:
 	
 	virtual void ExtendSelection(const TArray< FWidgetReference >& Selection, TArray< TSharedRef<FDesignerSurfaceElement> >& SurfaceElements) override;
 	virtual void Tick(const FGeometry& AllottedGeometry, const double InCurrentTime, const float InDeltaTime);
-	virtual void Paint(const TSet< FWidgetReference >& Selection, const FGeometry& AllottedGeometry, const FSlateRect& MyClippingRect, FSlateWindowElementList& OutDrawElements, int32 LayerId, const FWidgetStyle& InWidgetStyle, bool bParentEnabled) const override;
+	virtual void Paint(const TSet< FWidgetReference >& Selection, const FGeometry& AllottedGeometry, const FSlateRect& MyClippingRect, FSlateWindowElementList& OutDrawElements, int32 LayerId) const override;
 
 private:
 
@@ -55,7 +55,7 @@ private:
 	static bool GetCollisionSegmentsForSlot(UCanvasPanel* Canvas, UCanvasPanelSlot* Slot, TArray<FVector2D>& Segments);
 	static void GetCollisionSegmentsFromGeometry(FGeometry ArrangedGeometry, TArray<FVector2D>& Segments);
 
-	void PaintCollisionLines(const TSet< FWidgetReference >& Selection, const FGeometry& AllottedGeometry, const FSlateRect& MyClippingRect, FSlateWindowElementList& OutDrawElements, int32 LayerId, const FWidgetStyle& InWidgetStyle, bool bParentEnabled) const;
+	void PaintCollisionLines(const TSet< FWidgetReference >& Selection, const FGeometry& AllottedGeometry, const FSlateRect& MyClippingRect, FSlateWindowElementList& OutDrawElements, int32 LayerId) const;
 
 	void ProximitySnapValue(float SnapFrequency, float SnapProximity, float& Value);
 
diff --git a/Engine/Source/Editor/UMGEditor/Private/Hierarchy/SHierarchyViewItem.cpp b/Engine/Source/Editor/UMGEditor/Private/Hierarchy/SHierarchyViewItem.cpp
--- a/Engine/Source/Editor/UMGEditor/Private/Hierarchy/SHierarchyViewItem.cpp
+++ b/Engine/Source/Editor/UMGEditor/Private/Hierarchy/SHierarchyViewItem.cpp
@@ -295,7 +295,7 @@ void FHierarchyRoot::GetChildren(TArray< TSharedPtr<FHierarchyModel> >& Children
 void FHierarchyRoot::OnSelection()
 {
 	TSharedPtr<FWidgetBlueprintEditor> BPEd = BlueprintEditor.Pin();
-	if ( UWidget* Default =BPEd->GetWidgetBlueprintObj()->GeneratedClass->GetDefaultObject<UWidget>() )
+	if ( UWidget* Default = BPEd->GetWidgetBlueprintObj()->GeneratedClass->GetDefaultObject<UWidget>() )
 	{
 		TSet<UObject*> SelectedObjects;
 		SelectedObjects.Add(Default);
@@ -319,6 +319,112 @@ FReply FHierarchyRoot::HandleDrop(FDragDropEvent const& DragDropEvent)
 
 //////////////////////////////////////////////////////////////////////////
 
+FNamedSlotModel::FNamedSlotModel(FWidgetReference InItem, FName InSlotName, TSharedPtr<FWidgetBlueprintEditor> InBlueprintEditor)
+	: Item(InItem)
+	, SlotName(InSlotName)
+	, BlueprintEditor(InBlueprintEditor)
+{
+}
+
+FName FNamedSlotModel::GetUniqueName() const
+{
+	UWidget* WidgetTemplate = Item.GetTemplate();
+	if ( WidgetTemplate )
+	{
+		FString UniqueSlot = WidgetTemplate->GetName() + TEXT(".") + SlotName.ToString();
+		return FName(*UniqueSlot);
+	}
+
+	return NAME_None;
+}
+
+FText FNamedSlotModel::GetText() const
+{
+	if ( INamedSlotInterface* NamedSlotHost = InterfaceCast<INamedSlotInterface>(Item.GetPreview()) )
+	{
+		TSet<FWidgetReference> SelectedWidgets;
+		if ( UWidget* SlotContent = NamedSlotHost->GetContentForSlot(SlotName) )
+		{
+			return FText::Format(LOCTEXT("NamedSlotTextFormat", "{0} ({1})"), FText::FromName(SlotName), FText::FromName(SlotContent->GetFName()));
+		}
+	}
+
+	return FText::FromName(SlotName);
+}
+
+const FSlateBrush* FNamedSlotModel::GetImage() const
+{
+	return NULL;
+}
+
+FSlateFontInfo FNamedSlotModel::GetFont() const
+{
+	return FSlateFontInfo(FPaths::EngineContentDir() / TEXT("Slate/Fonts/Roboto-Bold.ttf"), 10);
+}
+
+void FNamedSlotModel::GetChildren(TArray< TSharedPtr<FHierarchyModel> >& Children)
+{
+	//TSharedPtr<FWidgetBlueprintEditor> BPEd = BlueprintEditor.Pin();
+	//UWidgetBlueprint* Blueprint = BPEd->GetWidgetBlueprintObj();
+
+	//if ( Blueprint->WidgetTree->RootWidget )
+	//{
+	//	TSharedPtr<FHierarchyWidget> RootChild = MakeShareable(new FHierarchyWidget(BPEd->GetReferenceFromTemplate(Blueprint->WidgetTree->RootWidget), BPEd));
+	//	Children.Add(RootChild);
+	//}
+}
+
+void FNamedSlotModel::OnSelection()
+{
+	TSharedPtr<FWidgetBlueprintEditor> BPEd = BlueprintEditor.Pin();
+	if ( INamedSlotInterface* NamedSlotHost = InterfaceCast<INamedSlotInterface>(Item.GetPreview()) )
+	{
+		TSet<FWidgetReference> SelectedWidgets;
+		if ( UWidget* SlotContent = NamedSlotHost->GetContentForSlot(SlotName) )
+		{
+			SelectedWidgets.Add( BPEd->GetReferenceFromPreview(SlotContent) );
+		}
+
+		BPEd->SelectWidgets(SelectedWidgets);
+	}
+}
+
+FReply FNamedSlotModel::OnDragOver(const FGeometry& MyGeometry, const FDragDropEvent& DragDropEvent)
+{
+	//UWidgetBlueprint* Blueprint = BlueprintEditor.Pin()->GetWidgetBlueprintObj();
+	//bool bIsDrop = false;
+	//return ProcessHierarchyDragDrop(DragDropEvent, bIsDrop, Blueprint, FWidgetReference());
+
+	return FReply::Unhandled();
+}
+
+FReply FNamedSlotModel::HandleDrop(FDragDropEvent const& DragDropEvent)
+{
+	UWidgetBlueprint* Blueprint = BlueprintEditor.Pin()->GetWidgetBlueprintObj();
+
+	bool bIsDrop = true;
+
+	// Is this a drag/drop op to create a new widget in the tree?
+	TSharedPtr<FWidgetTemplateDragDropOp> TemplateDragDropOp = DragDropEvent.GetOperationAs<FWidgetTemplateDragDropOp>();
+	if ( TemplateDragDropOp.IsValid() )
+	{
+		if ( INamedSlotInterface* NamedSlotHost = InterfaceCast<INamedSlotInterface>(Item.GetTemplate()) )
+		{
+			UWidget* Widget = TemplateDragDropOp->Template->Create(Blueprint->WidgetTree);
+
+			NamedSlotHost->SetContentForSlot(SlotName, Widget);
+
+			FBlueprintEditorUtils::MarkBlueprintAsStructurallyModified(Blueprint);
+
+			return FReply::Handled();
+		}
+	}
+
+	return FReply::Unhandled();
+}
+
+//////////////////////////////////////////////////////////////////////////
+
 FHierarchyWidget::FHierarchyWidget(FWidgetReference InItem, TSharedPtr<FWidgetBlueprintEditor> InBlueprintEditor)
 	: Item(InItem)
 	, BlueprintEditor(InBlueprintEditor)
@@ -422,7 +528,7 @@ bool FHierarchyWidget::OnVerifyNameTextChanged(const FText& InText, FText& OutEr
 	FString NewName = InText.ToString();
 
 	UWidgetBlueprint* Blueprint = BlueprintEditor.Pin()->GetWidgetBlueprintObj();
-	UWidget* ExistingTemplate = Blueprint->WidgetTree->FindWidget(NewName);
+	UWidget* ExistingTemplate = Blueprint->WidgetTree->FindWidget( FName(*NewName) );
 
 	bool bIsSameWidget = false;
 	if ( ExistingTemplate != NULL )
@@ -459,13 +565,26 @@ void FHierarchyWidget::OnNameTextCommited(const FText& InText, ETextCommit::Type
 void FHierarchyWidget::GetChildren(TArray< TSharedPtr<FHierarchyModel> >& Children)
 {
 	TSharedPtr<FWidgetBlueprintEditor> BPEd = BlueprintEditor.Pin();
+
+	// Check for named slots
+	if ( INamedSlotInterface* NamedSlotHost = InterfaceCast<INamedSlotInterface>(Item.GetPreview()) )
+	{
+		TArray<FName> SlotNames;
+		NamedSlotHost->GetSlotNames(SlotNames);
+
+		for ( FName& SlotName : SlotNames )
+		{
+			TSharedPtr<FNamedSlotModel> ChildItem = MakeShareable(new FNamedSlotModel(Item, SlotName, BPEd));
+			Children.Add(ChildItem);
+		}
+	}
 	
-	UPanelWidget* Widget = Cast<UPanelWidget>(Item.GetTemplate());
-	if ( Widget )
+	// Check if it's a panel widget that can support children
+	if ( UPanelWidget* PanelWidget = Cast<UPanelWidget>(Item.GetTemplate()) )
 	{
-		for ( int32 i = 0; i < Widget->GetChildrenCount(); i++ )
+		for ( int32 i = 0; i < PanelWidget->GetChildrenCount(); i++ )
 		{
-			UWidget* Child = Widget->GetChildAt(i);
+			UWidget* Child = PanelWidget->GetChildAt(i);
 			if ( Child )
 			{
 				TSharedPtr<FHierarchyWidget> ChildItem = MakeShareable(new FHierarchyWidget(BPEd->GetReferenceFromTemplate(Child), BPEd));
diff --git a/Engine/Source/Editor/UMGEditor/Private/Hierarchy/SHierarchyViewItem.h b/Engine/Source/Editor/UMGEditor/Private/Hierarchy/SHierarchyViewItem.h
--- a/Engine/Source/Editor/UMGEditor/Private/Hierarchy/SHierarchyViewItem.h
+++ b/Engine/Source/Editor/UMGEditor/Private/Hierarchy/SHierarchyViewItem.h
@@ -76,6 +76,39 @@ private:
 	TWeakPtr<FWidgetBlueprintEditor> BlueprintEditor;
 };
 
+class FNamedSlotModel : public FHierarchyModel
+{
+public:
+	FNamedSlotModel(FWidgetReference InItem, FName InSlotName, TSharedPtr<FWidgetBlueprintEditor> InBlueprintEditor);
+
+	virtual ~FNamedSlotModel() {}
+
+	virtual FName GetUniqueName() const override;
+
+	/* @returns the widget name to use for the tree item */
+	virtual FText GetText() const override;
+
+	virtual const FSlateBrush* GetImage() const override;
+
+	virtual FSlateFontInfo GetFont() const override;
+
+	virtual void OnSelection();
+
+	virtual FReply OnDragOver(const FGeometry& MyGeometry, const FDragDropEvent& DragDropEvent);
+
+	virtual FReply HandleDrop(FDragDropEvent const& DragDropEvent);
+
+protected:
+	virtual void GetChildren(TArray< TSharedPtr<FHierarchyModel> >& Children) override;
+
+private:
+
+	FWidgetReference Item;
+	FName SlotName;
+
+	TWeakPtr<FWidgetBlueprintEditor> BlueprintEditor;
+};
+
 class FHierarchyWidget : public FHierarchyModel
 {
 public:
diff --git a/Engine/Source/Editor/UMGEditor/Private/Templates/WidgetTemplateClass.cpp b/Engine/Source/Editor/UMGEditor/Private/Templates/WidgetTemplateClass.cpp
--- a/Engine/Source/Editor/UMGEditor/Private/Templates/WidgetTemplateClass.cpp
+++ b/Engine/Source/Editor/UMGEditor/Private/Templates/WidgetTemplateClass.cpp
@@ -36,7 +36,10 @@ FText FWidgetTemplateClass::GetCategory() const
 
 UWidget* FWidgetTemplateClass::Create(UWidgetTree* Tree)
 {
-	return Tree->ConstructWidget<UWidget>(WidgetClass.Get());
+	UWidget* NewWidget = Tree->ConstructWidget<UWidget>(WidgetClass.Get());
+	NewWidget->OnCreationFromPalette();
+
+	return NewWidget;
 }
 
 const FSlateBrush* FWidgetTemplateClass::GetIcon() const
diff --git a/Engine/Source/Editor/UMGEditor/Private/WidgetBlueprintEditor.cpp b/Engine/Source/Editor/UMGEditor/Private/WidgetBlueprintEditor.cpp
--- a/Engine/Source/Editor/UMGEditor/Private/WidgetBlueprintEditor.cpp
+++ b/Engine/Source/Editor/UMGEditor/Private/WidgetBlueprintEditor.cpp
@@ -419,7 +419,7 @@ void FWidgetBlueprintEditor::MigrateFromChain(FEditPropertyChain* PropertyThatCh
 
 			if ( PreviewWidget )
 			{
-				FString PreviewWidgetName = PreviewWidget->GetName();
+				FName PreviewWidgetName = PreviewWidget->GetFName();
 				UWidget* TemplateWidget = Blueprint->WidgetTree->FindWidget(PreviewWidgetName);
 
 				if ( TemplateWidget )
@@ -495,7 +495,7 @@ FWidgetReference FWidgetBlueprintEditor::GetReferenceFromPreview(UWidget* Previe
 
 		if ( PreviewWidget )
 		{
-			FString Name = PreviewWidget->GetName();
+			FName Name = PreviewWidget->GetFName();
 			return GetReferenceFromTemplate(Blueprint->WidgetTree->FindWidget(Name));
 		}
 	}
diff --git a/Engine/Source/Editor/UMGEditor/Private/WidgetBlueprintEditorToolbar.cpp b/Engine/Source/Editor/UMGEditor/Private/WidgetBlueprintEditorToolbar.cpp
--- a/Engine/Source/Editor/UMGEditor/Private/WidgetBlueprintEditorToolbar.cpp
+++ b/Engine/Source/Editor/UMGEditor/Private/WidgetBlueprintEditorToolbar.cpp
@@ -81,8 +81,8 @@ void FWidgetBlueprintEditorToolbar::FillWidgetBlueprintEditorModesToolbar(FToolB
 				NULL,
 				TEXT("Shared/Editors/BlueprintEditor"),
 				TEXT("DesignerMode")))
-			.IconImage(FEditorStyle::GetBrush("FullBlueprintEditor.SwitchToBlueprintDefaultsMode"))
-			.SmallIconImage(FEditorStyle::GetBrush("FullBlueprintEditor.SwitchToBlueprintDefaultsMode.Small"))
+			.IconImage(FEditorStyle::GetBrush("UMGEditor.SwitchToDesigner"))
+			.SmallIconImage(FEditorStyle::GetBrush("UMGEditor.SwitchToDesigner.Small"))
 			.AddMetaData<FTagMetaData>(FTagMetaData(TEXT("DesignerMode")))
 		);
 
diff --git a/Engine/Source/Editor/UMGEditor/Private/WidgetBlueprintEditorUtils.cpp b/Engine/Source/Editor/UMGEditor/Private/WidgetBlueprintEditorUtils.cpp
--- a/Engine/Source/Editor/UMGEditor/Private/WidgetBlueprintEditorUtils.cpp
+++ b/Engine/Source/Editor/UMGEditor/Private/WidgetBlueprintEditorUtils.cpp
@@ -76,7 +76,7 @@ bool FWidgetBlueprintEditorUtils::RenameWidget(TSharedRef<FWidgetBlueprintEditor
 	const FString NewNameStr = NewName.ToString();
 	const FString OldNameStr = OldName.ToString();
 
-	UWidget* Widget = Blueprint->WidgetTree->FindWidget(OldNameStr);
+	UWidget* Widget = Blueprint->WidgetTree->FindWidget(OldName);
 	check(Widget);
 
 	if ( Widget )
@@ -158,6 +158,9 @@ bool FWidgetBlueprintEditorUtils::RenameWidget(TSharedRef<FWidgetBlueprintEditor
 			}
 		}
 
+		// Update named slot bindings?
+		// TODO...
+
 		// Validate child blueprints and adjust variable names to avoid a potential name collision
 		FBlueprintEditorUtils::ValidateBlueprintChildVariables(Blueprint, NewName);
 
@@ -225,15 +228,22 @@ void FWidgetBlueprintEditorUtils::DeleteWidgets(UWidgetBlueprint* BP, TSet<FWidg
 
 			bRemoved = BP->WidgetTree->RemoveWidget(WidgetTemplate);
 
+			// If we failed to remove the widget from the tree, it may be a named slot widget that
+			// we need to search for in the named slot bindings.
+			if ( !bRemoved && WidgetTemplate->GetParent() == nullptr )
+			{
+				bRemoved = FindAndRemoveNamedSlotContent(WidgetTemplate, BP->WidgetTree);
+			}
+
 			// Rename the removed widget to the transient package so that it doesn't conflict with future widgets sharing the same name.
-			WidgetTemplate->Rename(NULL, NULL);
+			WidgetTemplate->Rename(nullptr, nullptr);
 
 			// Rename all child widgets as well, to the transient package so that they don't conflict with future widgets sharing the same name.
 			TArray<UWidget*> ChildWidgets;
 			BP->WidgetTree->GetChildWidgets(WidgetTemplate, ChildWidgets);
 			for ( UWidget* Widget : ChildWidgets )
 			{
-				Widget->Rename(NULL, NULL);
+				Widget->Rename(nullptr, nullptr);
 			}
 		}
 
@@ -246,6 +256,35 @@ void FWidgetBlueprintEditorUtils::DeleteWidgets(UWidgetBlueprint* BP, TSet<FWidg
 	}
 }
 
+bool FWidgetBlueprintEditorUtils::FindAndRemoveNamedSlotContent(UWidget* WidgetTemplate, UWidgetTree* WidgetTree)
+{
+	TArray<UWidget*> AllWidgets;
+	WidgetTree->GetAllWidgets(AllWidgets);
+
+	for ( UWidget* Widget : AllWidgets )
+	{
+		if ( INamedSlotInterface* NamedSlotHost = InterfaceCast<INamedSlotInterface>(Widget) )
+		{
+			TArray<FName> SlotNames;
+			NamedSlotHost->GetSlotNames(SlotNames);
+
+			for ( FName SlotName : SlotNames )
+			{
+				if ( UWidget* SlotContent = NamedSlotHost->GetContentForSlot(SlotName) )
+				{
+					if ( SlotContent == WidgetTemplate )
+					{
+						NamedSlotHost->SetContentForSlot(SlotName, nullptr);
+						return true;
+					}
+				}
+			}
+		}
+	}
+
+	return false;
+}
+
 void FWidgetBlueprintEditorUtils::BuildWrapWithMenu(FMenuBuilder& Menu, UWidgetBlueprint* BP, TSet<FWidgetReference> Widgets)
 {
 	Menu.BeginSection("WrapWith", LOCTEXT("WidgetTree_WrapWith", "Wrap With..."));
diff --git a/Engine/Source/Editor/UMGEditor/Private/WidgetBlueprintEditorUtils.h b/Engine/Source/Editor/UMGEditor/Private/WidgetBlueprintEditorUtils.h
--- a/Engine/Source/Editor/UMGEditor/Private/WidgetBlueprintEditorUtils.h
+++ b/Engine/Source/Editor/UMGEditor/Private/WidgetBlueprintEditorUtils.h
@@ -32,6 +32,8 @@ public:
 	static void ImportPropertiesFromText(UObject* Object, const TMap<FName, FString>& ExportedProperties);
 
 private:
+	static bool FindAndRemoveNamedSlotContent(UWidget* WidgetTemplate, UWidgetTree* WidgetTree);
+
 	static void BuildWrapWithMenu(FMenuBuilder& Menu, UWidgetBlueprint* BP, TSet<FWidgetReference> Widgets);
 
 	static void WrapWidgets(UWidgetBlueprint* BP, TSet<FWidgetReference> Widgets, UClass* WidgetClass);
diff --git a/Engine/Source/Editor/UMGEditor/Public/DesignerExtension.h b/Engine/Source/Editor/UMGEditor/Public/DesignerExtension.h
--- a/Engine/Source/Editor/UMGEditor/Public/DesignerExtension.h
+++ b/Engine/Source/Editor/UMGEditor/Public/DesignerExtension.h
@@ -120,7 +120,7 @@ public:
 	{
 	}
 
-	virtual void Paint(const TSet< FWidgetReference >& Selection, const FGeometry& AllottedGeometry, const FSlateRect& MyClippingRect, FSlateWindowElementList& OutDrawElements, int32 LayerId, const FWidgetStyle& InWidgetStyle, bool bParentEnabled) const
+	virtual void Paint(const TSet< FWidgetReference >& Selection, const FGeometry& AllottedGeometry, const FSlateRect& MyClippingRect, FSlateWindowElementList& OutDrawElements, int32 LayerId) const
 	{
 	}
 
diff --git a/Engine/Source/Editor/UMGEditor/Public/IUMGDesigner.h b/Engine/Source/Editor/UMGEditor/Public/IUMGDesigner.h
--- a/Engine/Source/Editor/UMGEditor/Public/IUMGDesigner.h
+++ b/Engine/Source/Editor/UMGEditor/Public/IUMGDesigner.h
@@ -44,4 +44,9 @@ public:
 	 * Gets the previous frames widget geometry of the parent of the provided widget.
 	 */
 	virtual bool GetWidgetParentGeometry(const FWidgetReference& Widget, FGeometry& Geometry) const = 0;
+
+	/**
+	 * Marks the designer content as being modified.
+	 */
+	virtual void MarkDesignModifed(bool bRequiresRecompile) = 0;
 };
diff --git a/Engine/Source/Runtime/Slate/Private/Widgets/Layout/SWrapBox.cpp b/Engine/Source/Runtime/Slate/Private/Widgets/Layout/SWrapBox.cpp
--- a/Engine/Source/Runtime/Slate/Private/Widgets/Layout/SWrapBox.cpp
+++ b/Engine/Source/Runtime/Slate/Private/Widgets/Layout/SWrapBox.cpp
@@ -8,14 +8,14 @@ SWrapBox::SWrapBox()
 {
 }
 
-FWrapSlot& SWrapBox::Slot()
+SWrapBox::FSlot& SWrapBox::Slot()
 {
-	return *( new FWrapSlot() );
+	return *( new SWrapBox::FSlot() );
 }
 
-FWrapSlot& SWrapBox::AddSlot()
+SWrapBox::FSlot& SWrapBox::AddSlot()
 {
-	FWrapSlot* NewSlot = new FWrapSlot();
+	SWrapBox::FSlot* NewSlot = new SWrapBox::FSlot();
 	Slots.Add(NewSlot);
 	return *NewSlot;
 }
@@ -223,3 +223,8 @@ FChildren* SWrapBox::GetChildren()
 {
 	return &Slots;	
 }
+
+void SWrapBox::SetInnerSlotPadding(FVector2D InInnerSlotPadding)
+{
+	InnerSlotPadding = InInnerSlotPadding;
+}
diff --git a/Engine/Source/Runtime/Slate/Public/Widgets/Layout/SWrapBox.h b/Engine/Source/Runtime/Slate/Public/Widgets/Layout/SWrapBox.h
--- a/Engine/Source/Runtime/Slate/Public/Widgets/Layout/SWrapBox.h
+++ b/Engine/Source/Runtime/Slate/Public/Widgets/Layout/SWrapBox.h
@@ -2,35 +2,6 @@
 
 #pragma once
 
-/** A slot that support alignment of content and padding */
-class SLATE_API FWrapSlot : public TSlotBase<FWrapSlot>, public TSupportsContentAlignmentMixin<FWrapSlot>, public TSupportsContentPaddingMixin<FWrapSlot>
-{
-public:
-	FWrapSlot()
-		: TSlotBase<FWrapSlot>()
-		, TSupportsContentAlignmentMixin<FWrapSlot>( HAlign_Fill, VAlign_Fill )
-		, bSlotFillEmptySpace( false )
-	{
-	}
-
-	/** If the total available space in the wrap panel drops below this threshold, this slot will attempt to fill an entire line. */
-	FWrapSlot& FillLineWhenWidthLessThan( TOptional<float> InFillLineWhenWidthLessThan )
-	{
-		SlotFillLineWhenWidthLessThan = InFillLineWhenWidthLessThan;
-		return *( static_cast<FWrapSlot*>( this ) );
-	}
-
-	/** Should this slot fill the remaining space on the line? */
-	FWrapSlot& FillEmptySpace( bool bInFillEmptySpace )
-	{
-		bSlotFillEmptySpace = bInFillEmptySpace;
-		return *( static_cast<FWrapSlot*>( this ) );
-	}
-
-	TOptional<float> SlotFillLineWhenWidthLessThan;
-	bool bSlotFillEmptySpace;
-};
-
 /**
  * Arranges widgets left-to-right.
  * When the widgets exceed the PreferredWidth
@@ -46,9 +17,39 @@ public:
  */
 class SLATE_API SWrapBox : public SPanel
 {
-
 public:
 
+	/** A slot that support alignment of content and padding */
+	class FSlot : public TSlotBase<FSlot>, public TSupportsContentAlignmentMixin<FSlot>, public TSupportsContentPaddingMixin<FSlot>
+	{
+	public:
+		FSlot()
+			: TSlotBase<FSlot>()
+			, TSupportsContentAlignmentMixin<FSlot>(HAlign_Fill, VAlign_Fill)
+			, SlotFillLineWhenWidthLessThan()
+			, bSlotFillEmptySpace(false)
+		{
+		}
+
+		/** If the total available space in the wrap panel drops below this threshold, this slot will attempt to fill an entire line. */
+		FSlot& FillLineWhenWidthLessThan(TOptional<float> InFillLineWhenWidthLessThan)
+		{
+			SlotFillLineWhenWidthLessThan = InFillLineWhenWidthLessThan;
+			return *( static_cast<FSlot*>( this ) );
+		}
+
+		/** Should this slot fill the remaining space on the line? */
+		FSlot& FillEmptySpace(bool bInFillEmptySpace)
+		{
+			bSlotFillEmptySpace = bInFillEmptySpace;
+			return *( static_cast<FSlot*>( this ) );
+		}
+
+		TOptional<float> SlotFillLineWhenWidthLessThan;
+		bool bSlotFillEmptySpace;
+	};
+
+
 	SLATE_BEGIN_ARGS(SWrapBox)
 		: _PreferredWidth( 100.f )
 		, _InnerSlotPadding( FVector2D::ZeroVector )
@@ -58,7 +59,7 @@ public:
 		}
 
 		/** The slot supported by this panel */
-		SLATE_SUPPORTS_SLOT( FWrapSlot )
+		SLATE_SUPPORTS_SLOT( FSlot )
 
 		/** The preferred width, if not set will fill the space */
 		SLATE_ATTRIBUTE( float, PreferredWidth )
@@ -72,9 +73,9 @@ public:
 
 	SWrapBox();
 
-	static FWrapSlot& Slot();
+	static FSlot& Slot();
 
-	FWrapSlot& AddSlot();
+	FSlot& AddSlot();
 
 	/** Removes a slot from this box panel which contains the specified SWidget
 	 *
@@ -95,13 +96,16 @@ public:
 
 	virtual FChildren* GetChildren() override;
 
+	/** See InnerSlotPadding Attribute */
+	void SetInnerSlotPadding(FVector2D InInnerSlotPadding);
+
 private:
 
 	/** How wide this panel should appear to be. Any widgets past this line will be wrapped onto the next line. */
 	TAttribute<float> PreferredWidth;
 
 	/** The slots that contain this panel's children. */
-	TPanelChildren<FWrapSlot> Slots;
+	TPanelChildren<FSlot> Slots;
 
 	/** When two slots end up sharing a border, this will put that much padding between then, but otherwise wont. */
 	FVector2D InnerSlotPadding;
diff --git a/Engine/Source/Runtime/SlateCore/Private/Styling/CoreStyle.cpp b/Engine/Source/Runtime/SlateCore/Private/Styling/CoreStyle.cpp
--- a/Engine/Source/Runtime/SlateCore/Private/Styling/CoreStyle.cpp
+++ b/Engine/Source/Runtime/SlateCore/Private/Styling/CoreStyle.cpp
@@ -501,7 +501,7 @@ TSharedRef<ISlateStyle> FCoreStyle::Create( const FName& InStyleSetName )
 	{
 		Style->Set( "ProgressBar", FProgressBarStyle()
 			.SetBackgroundImage( BOX_BRUSH( "Common/ProgressBar_Background", FMargin(5.f/12.f) ) )
-			.SetFillImage( BOX_BRUSH( "Common/ProgressBar_Fill", FMargin(5.f/12.f) ) )
+			.SetFillImage( BOX_BRUSH( "Common/ProgressBar_Fill", FMargin(5.f/12.f), FLinearColor( 1.0f, 0.22f, 0.0f )  ) )
 			.SetMarqueeImage( IMAGE_BRUSH( "Common/ProgressBar_Marquee", FVector2D(20,12), FLinearColor::White, ESlateBrushTileType::Horizontal ) )
 			);
 	}
diff --git a/Engine/Source/Runtime/UMG/Private/Components/ContentWidget.cpp b/Engine/Source/Runtime/UMG/Private/Components/ContentWidget.cpp
--- a/Engine/Source/Runtime/UMG/Private/Components/ContentWidget.cpp
+++ b/Engine/Source/Runtime/UMG/Private/Components/ContentWidget.cpp
@@ -20,7 +20,7 @@ UContentWidget::UContentWidget(const FPostConstructInitializeProperties& PCIP)
 
 UPanelSlot* UContentWidget::GetContentSlot() const
 {
-	return Slots.Num() > 0 ? Slots[0] : NULL;
+	return Slots.Num() > 0 ? Slots[0] : nullptr;
 }
 
 UClass* UContentWidget::GetSlotClass() const
diff --git a/Engine/Source/Runtime/UMG/Private/Components/NamedSlot.cpp b/Engine/Source/Runtime/UMG/Private/Components/NamedSlot.cpp
new file mode 100644
--- /dev/null
+++ b/Engine/Source/Runtime/UMG/Private/Components/NamedSlot.cpp
@@ -0,0 +1,105 @@
+// Copyright 1998-2014 Epic Games, Inc. All Rights Reserved.
+
+#include "UMGPrivatePCH.h"
+
+#define LOCTEXT_NAMESPACE "UMG"
+
+/////////////////////////////////////////////////////
+// UNamedSlot
+
+UNamedSlot::UNamedSlot(const FPostConstructInitializeProperties& PCIP)
+	: Super(PCIP)
+{
+	bIsVariable = true;
+}
+
+void UNamedSlot::ReleaseSlateResources(bool bReleaseChildren)
+{
+	Super::ReleaseSlateResources(bReleaseChildren);
+
+	MyBox.Reset();
+}
+
+TSharedRef<SWidget> UNamedSlot::RebuildWidget()
+{
+	MyBox = SNew(SBox);
+
+	if ( IsDesignTime() )
+	{
+		MyBox->SetContent(
+			SNew(SBox)
+			.HAlign(HAlign_Center)
+			.VAlign(VAlign_Center)
+			[
+				SNew(STextBlock)
+				.Text(FText::FromName(GetFName()))
+			]
+		);
+	}
+
+	// Add any existing content to the new slate box
+	if ( GetChildrenCount() > 0 )
+	{
+		UPanelSlot* ContentSlot = GetContentSlot();
+		if ( ContentSlot->Content )
+		{
+			MyBox->SetContent(ContentSlot->Content->TakeWidget());
+		}
+	}
+
+	return MyBox.ToSharedRef();
+}
+
+void UNamedSlot::SynchronizeProperties()
+{
+	Super::SynchronizeProperties();
+}
+
+void UNamedSlot::OnSlotAdded(UPanelSlot* Slot)
+{
+	// Add the child to the live slot if it already exists
+	if ( MyBox.IsValid() && Slot->Content )
+	{
+		MyBox->SetContent(Slot->Content->TakeWidget());
+	}
+}
+
+void UNamedSlot::OnSlotRemoved(UPanelSlot* Slot)
+{
+	// Remove the widget from the live slot if it exists.
+	if ( MyBox.IsValid() )
+	{
+		MyBox->SetContent(SNullWidget::NullWidget);
+
+		if ( IsDesignTime() )
+		{
+			MyBox->SetContent(
+				SNew(SBox)
+				.HAlign(HAlign_Center)
+				.VAlign(VAlign_Center)
+				[
+					SNew(STextBlock)
+					.Text(FText::FromName(GetFName()))
+				]
+			);
+		}
+	}
+}
+
+#if WITH_EDITOR
+
+const FSlateBrush* UNamedSlot::GetEditorIcon()
+{
+	return FUMGStyle::Get().GetBrush("Widget.Spacer");
+}
+
+const FText UNamedSlot::GetPaletteCategory()
+{
+	return LOCTEXT("Common", "Common");
+}
+
+#endif
+
+/////////////////////////////////////////////////////
+
+#undef LOCTEXT_NAMESPACE
diff --git a/Engine/Source/Runtime/UMG/Private/Components/ProgressBar.cpp b/Engine/Source/Runtime/UMG/Private/Components/ProgressBar.cpp
--- a/Engine/Source/Runtime/UMG/Private/Components/ProgressBar.cpp
+++ b/Engine/Source/Runtime/UMG/Private/Components/ProgressBar.cpp
@@ -12,6 +12,7 @@ UProgressBar::UProgressBar(const FPostConstructInitializeProperties& PCIP)
 {
 	SProgressBar::FArguments SlateDefaults;
 	WidgetStyle = *SlateDefaults._Style;
+	WidgetStyle.FillImage.TintColor = FLinearColor::White;
 
 	BarFillType = EProgressBarFillType::LeftToRight;
 	bIsMarquee = false;
@@ -113,6 +114,11 @@ const FText UProgressBar::GetPaletteCategory()
 	return LOCTEXT("Common", "Common");
 }
 
+void UProgressBar::OnCreationFromPalette()
+{
+	FillColorAndOpacity = FLinearColor(0, 0.5f, 1.0f);
+}
+
 #endif
 
 /////////////////////////////////////////////////////
diff --git a/Engine/Source/Runtime/UMG/Private/Components/Widget.cpp b/Engine/Source/Runtime/UMG/Private/Components/Widget.cpp
--- a/Engine/Source/Runtime/UMG/Private/Components/Widget.cpp
+++ b/Engine/Source/Runtime/UMG/Private/Components/Widget.cpp
@@ -204,12 +204,9 @@ TSharedRef<SWidget> UWidget::TakeWidget()
 	TSharedPtr<SWidget> SafeWidget;
 	bool bNewlyCreated = false;
 
+	// If the underlying widget doesn't exist we need to construct and cache the widget for the first run.
 	if ( !MyWidget.IsValid() )
 	{
-		// We lie a bit about this being a const function.  If this is the first call it's not really const
-		// and we need to construct and cache the widget for the first run.  But instead of forcing everyone
-		// downstream to make RebuildWidget const and force every implementation to make things mutable, we
-		// just blow away the const here.
 		SafeWidget = RebuildWidget();
 		MyWidget = SafeWidget;
 
@@ -223,10 +220,12 @@ TSharedRef<SWidget> UWidget::TakeWidget()
 	// If it is a user widget wrap it in a SObjectWidget to keep the instance from being GC'ed
 	if ( IsA(UUserWidget::StaticClass()) )
 	{
+		// If the GC Widget is still valid we still exist in the slate hierarchy, so just return the GC Widget.
 		if ( MyGCWidget.IsValid() )
 		{
 			return MyGCWidget.Pin().ToSharedRef();
 		}
+		// Otherwise we need to recreate the wrapper widget
 		else
 		{
 			TSharedPtr<SWidget> SafeGCWidget = SNew(SObjectWidget, Cast<UUserWidget>(this))
diff --git a/Engine/Source/Runtime/UMG/Private/Components/WrapBox.cpp b/Engine/Source/Runtime/UMG/Private/Components/WrapBox.cpp
new file mode 100644
--- /dev/null
+++ b/Engine/Source/Runtime/UMG/Private/Components/WrapBox.cpp
@@ -0,0 +1,98 @@
+// Copyright 1998-2014 Epic Games, Inc. All Rights Reserved.
+
+#include "UMGPrivatePCH.h"
+
+#define LOCTEXT_NAMESPACE "UMG"
+
+/////////////////////////////////////////////////////
+// UWrapBox
+
+UWrapBox::UWrapBox(const FPostConstructInitializeProperties& PCIP)
+	: Super(PCIP)
+{
+	bIsVariable = false;
+
+	SWrapBox::FArguments Defaults;
+	Visiblity = UWidget::ConvertRuntimeToSerializedVisiblity( Defaults._Visibility.Get() );
+}
+
+void UWrapBox::ReleaseSlateResources(bool bReleaseChildren)
+{
+	Super::ReleaseSlateResources(bReleaseChildren);
+
+	MyWrapBox.Reset();
+}
+
+UClass* UWrapBox::GetSlotClass() const
+{
+	return UWrapBoxSlot::StaticClass();
+}
+
+void UWrapBox::OnSlotAdded(UPanelSlot* Slot)
+{
+	// Add the child to the live canvas if it already exists
+	if ( MyWrapBox.IsValid() )
+	{
+		Cast<UWrapBoxSlot>(Slot)->BuildSlot(MyWrapBox.ToSharedRef());
+	}
+}
+
+void UWrapBox::OnSlotRemoved(UPanelSlot* Slot)
+{
+	// Remove the widget from the live slot if it exists.
+	if ( MyWrapBox.IsValid() )
+	{
+		TSharedPtr<SWidget> Widget = Slot->Content->GetCachedWidget();
+		if ( Widget.IsValid() )
+		{
+			MyWrapBox->RemoveSlot(Widget.ToSharedRef());
+		}
+	}
+}
+
+UWrapBoxSlot* UWrapBox::AddChildWrapBox(UWidget* Content)
+{
+	return Cast<UWrapBoxSlot>(Super::AddChild(Content));
+}
+
+TSharedRef<SWidget> UWrapBox::RebuildWidget()
+{
+	MyWrapBox = SNew(SWrapBox)
+		.UseAllottedWidth(true);
+
+	for ( UPanelSlot* Slot : Slots )
+	{
+		if ( UWrapBoxSlot* TypedSlot = Cast<UWrapBoxSlot>(Slot) )
+		{
+			TypedSlot->Parent = this;
+			TypedSlot->BuildSlot(MyWrapBox.ToSharedRef());
+		}
+	}
+
+	return BuildDesignTimeWidget( MyWrapBox.ToSharedRef() );
+}
+
+void UWrapBox::SynchronizeProperties()
+{
+	Super::SynchronizeProperties();
+
+	MyWrapBox->SetInnerSlotPadding(InnerSlotPadding);
+}
+
+#if WITH_EDITOR
+
+const FSlateBrush* UWrapBox::GetEditorIcon()
+{
+	return FUMGStyle::Get().GetBrush("Widget.WrapBox");
+}
+
+const FText UWrapBox::GetPaletteCategory()
+{
+	return LOCTEXT("Panel", "Panel");
+}
+
+#endif
+
+/////////////////////////////////////////////////////
+
+#undef LOCTEXT_NAMESPACE
\ No newline at end of file
diff --git a/Engine/Source/Runtime/UMG/Private/Components/WrapBoxSlot.cpp b/Engine/Source/Runtime/UMG/Private/Components/WrapBoxSlot.cpp
new file mode 100644
--- /dev/null
+++ b/Engine/Source/Runtime/UMG/Private/Components/WrapBoxSlot.cpp
@@ -0,0 +1,91 @@
+// Copyright 1998-2014 Epic Games, Inc. All Rights Reserved.
+
+#include "UMGPrivatePCH.h"
+
+/////////////////////////////////////////////////////
+// UWrapBoxSlot
+
+UWrapBoxSlot::UWrapBoxSlot(const FPostConstructInitializeProperties& PCIP)
+	: Super(PCIP)
+	, Slot(NULL)
+{
+	HorizontalAlignment = HAlign_Fill;
+	VerticalAlignment = VAlign_Fill;
+
+	bFillEmptySpace = false;
+	FillSpanWhenLessThan = 0;
+}
+
+void UWrapBoxSlot::ReleaseSlateResources(bool bReleaseChildren)
+{
+	Super::ReleaseSlateResources(bReleaseChildren);
+
+	Slot = NULL;
+}
+
+void UWrapBoxSlot::BuildSlot(TSharedRef<SWrapBox> WrapBox)
+{
+	Slot = &WrapBox->AddSlot()
+		.Padding(Padding)
+		.HAlign(HorizontalAlignment)
+		.VAlign(VerticalAlignment)
+		.FillEmptySpace(bFillEmptySpace)
+		.FillLineWhenWidthLessThan(FillSpanWhenLessThan == 0 ? TOptional<float>() : TOptional<float>(FillSpanWhenLessThan))
+		[
+			Content == NULL ? SNullWidget::NullWidget : Content->TakeWidget()
+		];
+}
+
+void UWrapBoxSlot::SetPadding(FMargin InPadding)
+{
+	Padding = InPadding;
+	if ( Slot )
+	{
+		Slot->Padding(InPadding);
+	}
+}
+
+void UWrapBoxSlot::SetFillEmptySpace(bool InbFillEmptySpace)
+{
+	bFillEmptySpace = InbFillEmptySpace;
+	if ( Slot )
+	{
+		Slot->FillEmptySpace(InbFillEmptySpace);
+	}
+}
+
+void UWrapBoxSlot::SetFillSpanWhenLessThan(float InFillSpanWhenLessThan)
+{
+	FillSpanWhenLessThan = InFillSpanWhenLessThan;
+	if ( Slot )
+	{
+		Slot->FillLineWhenWidthLessThan(InFillSpanWhenLessThan == 0 ? TOptional<float>() : TOptional<float>(InFillSpanWhenLessThan));
+	}
+}
+
+void UWrapBoxSlot::SetHorizontalAlignment(EHorizontalAlignment InHorizontalAlignment)
+{
+	HorizontalAlignment = InHorizontalAlignment;
+	if ( Slot )
+	{
+		Slot->HAlign(InHorizontalAlignment);
+	}
+}
+
+void UWrapBoxSlot::SetVerticalAlignment(EVerticalAlignment InVerticalAlignment)
+{
+	VerticalAlignment = InVerticalAlignment;
+	if ( Slot )
+	{
+		Slot->VAlign(InVerticalAlignment);
+	}
+}
+
+void UWrapBoxSlot::SynchronizeProperties()
+{
+	SetPadding(Padding);
+	SetFillEmptySpace(bFillEmptySpace);
+	SetFillSpanWhenLessThan(FillSpanWhenLessThan);
+	SetHorizontalAlignment(HorizontalAlignment);
+	SetVerticalAlignment(VerticalAlignment);
+}
diff --git a/Engine/Source/Runtime/UMG/Private/Slate/SConstraintCanvas.cpp b/Engine/Source/Runtime/UMG/Private/Slate/SConstraintCanvas.cpp
--- a/Engine/Source/Runtime/UMG/Private/Slate/SConstraintCanvas.cpp
+++ b/Engine/Source/Runtime/UMG/Private/Slate/SConstraintCanvas.cpp
@@ -22,13 +22,11 @@ void SConstraintCanvas::Construct( const SConstraintCanvas::FArguments& InArgs )
 	}
 }
 
-
 void SConstraintCanvas::ClearChildren( )
 {
 	Children.Empty();
 }
 
-
 int32 SConstraintCanvas::RemoveSlot( const TSharedRef<SWidget>& SlotWidget )
 {
 	for (int32 SlotIdx = 0; SlotIdx < Children.Num(); ++SlotIdx)
@@ -149,7 +147,6 @@ void SConstraintCanvas::OnArrangeChildren( const FGeometry& AllottedGeometry, FA
 	}
 }
 
-
 int32 SConstraintCanvas::OnPaint( const FPaintArgs& Args, const FGeometry& AllottedGeometry, const FSlateRect& MyClippingRect, FSlateWindowElementList& OutDrawElements, int32 LayerId, const FWidgetStyle& InWidgetStyle, bool bParentEnabled ) const
 {
 	FArrangedChildren ArrangedChildren(EVisibility::Visible);
@@ -171,13 +168,35 @@ int32 SConstraintCanvas::OnPaint( const FPaintArgs& Args, const FGeometry& Allot
 	return MaxLayerId;
 }
 
-
 FVector2D SConstraintCanvas::ComputeDesiredSize() const
 {
-	// ConstraintCanvas widgets have no desired size -- their size is always determined by their container
-	return FVector2D::ZeroVector;
-}
+	FVector2D DesiredSize(0,0);
+
+	// Arrange the children now in their proper z-order.
+	for ( int32 ChildIndex = 0; ChildIndex < Children.Num(); ++ChildIndex )
+	{
+		const SConstraintCanvas::FSlot& CurChild = Children[ChildIndex];
+
+		const FMargin Offset = CurChild.OffsetAttr.Get();
+		const FVector2D Alignment = CurChild.AlignmentAttr.Get();
+		const FAnchors Anchors = CurChild.AnchorsAttr.Get();
 
+		const FVector2D SlotSize = FVector2D(Offset.Right, Offset.Bottom);
+		const FVector2D WidgetDesiredSize = CurChild.GetWidget()->GetDesiredSize();
+
+		const bool AutoSize = CurChild.AutoSizeAttr.Get();
+
+		const FVector2D Size = AutoSize ? WidgetDesiredSize : SlotSize;
+
+		const bool bIsDockedHorizontally = ( Anchors.Minimum.X == Anchors.Maximum.X ) && ( Anchors.Minimum.X == 0 || Anchors.Minimum.X == 1 );
+		const bool bIsDockedVertically = ( Anchors.Minimum.Y == Anchors.Maximum.Y ) && ( Anchors.Minimum.Y == 0 || Anchors.Minimum.Y == 1 );
+
+		DesiredSize.X = FMath::Max(DesiredSize.X, Size.X + ( bIsDockedHorizontally ? FMath::Abs(Offset.Left) : 0.0f ));
+		DesiredSize.Y = FMath::Max(DesiredSize.Y, Size.Y + ( bIsDockedVertically ? FMath::Abs(Offset.Top) : 0.0f ));
+	}
+
+	return DesiredSize;
+}
 
 FChildren* SConstraintCanvas::GetChildren()
 {
diff --git a/Engine/Source/Runtime/UMG/Private/UMGStyle.cpp b/Engine/Source/Runtime/UMG/Private/UMGStyle.cpp
--- a/Engine/Source/Runtime/UMG/Private/UMGStyle.cpp
+++ b/Engine/Source/Runtime/UMG/Private/UMGStyle.cpp
@@ -79,6 +79,7 @@ TSharedRef< FSlateStyleSet > FUMGStyle::Create()
 	Style->Set("Widget.SizeBox", new IMAGE_BRUSH(TEXT("SizeBox"), Icon16x16));
 	Style->Set("Widget.SpinBox", new IMAGE_BRUSH(TEXT("SpinBox"), Icon16x16));
 	Style->Set("Widget.Grid", new IMAGE_BRUSH(TEXT("Grid"), Icon16x16));
+	Style->Set("Widget.WrapBox", new IMAGE_BRUSH(TEXT("WrapBox"), Icon16x16));
 
 	Style->Set("Widget.UserWidget", new IMAGE_BRUSH(TEXT("UserWidget"), Icon16x16));
 
diff --git a/Engine/Source/Runtime/UMG/Private/UserWidget.cpp b/Engine/Source/Runtime/UMG/Private/UserWidget.cpp
--- a/Engine/Source/Runtime/UMG/Private/UserWidget.cpp
+++ b/Engine/Source/Runtime/UMG/Private/UserWidget.cpp
@@ -24,6 +24,9 @@ UUserWidget::UUserWidget(const FPostConstructInitializeProperties& PCIP)
 	CachedWorld = NULL;
 
 	bSupportsKeyboardFocus = true;
+
+	bUseDesignTimeSize = false;
+	DesignTimeSize = FVector2D(100, 100);
 }
 
 void UUserWidget::Initialize()
@@ -39,6 +42,22 @@ void UUserWidget::Initialize()
 		{
 			BGClass->InitializeWidget(this);
 		}
+
+		for ( UWidget* Widget : Components )
+		{
+			if ( UNamedSlot* NamedWidet = Cast<UNamedSlot>(Widget) )
+			{
+				for ( FNamedSlotBinding& Binding : NamedSlotBindings )
+				{
+					if ( Binding.Content && Binding.Name == NamedWidet->GetFName() )
+					{
+						NamedWidet->ClearChildren();
+						NamedWidet->AddChild(Binding.Content);
+						continue;
+					}
+				}
+			}
+		}
 	}
 }
 
@@ -352,7 +371,7 @@ TSharedRef<SWidget> UUserWidget::RebuildWidget()
 	return UserRootWidget.ToSharedRef();
 }
 
-TSharedPtr<SWidget> UUserWidget::GetWidgetFromName(const FString& Name) const
+TSharedPtr<SWidget> UUserWidget::GetWidgetFromName(const FName& Name) const
 {
 	UWidget* WidgetObject = WidgetTree->FindWidget(Name);
 	if ( WidgetObject )
@@ -376,6 +395,73 @@ UWidget* UUserWidget::GetHandleFromName(const FString& Name) const
 	return nullptr;
 }
 
+void UUserWidget::GetSlotNames(TArray<FName>& SlotNames) const
+{
+	//TODO UMG This should be static data computed and stored on the class at compile time.
+
+	// Add all the names of the named slot widgets to the slot names structure.
+	for ( UWidget* Widget : Components )
+	{
+		if ( Widget && Widget->IsA<UNamedSlot>() )
+		{
+			SlotNames.Add( Widget->GetFName() );
+		}
+	}
+
+	// Also add any existing bindings uniquely, templates don't have any components
+	// yet because they're never initialized because of their CDO status.
+	for ( const FNamedSlotBinding& Binding : NamedSlotBindings )
+	{
+		SlotNames.AddUnique( Binding.Name );
+	}
+}
+
+UWidget* UUserWidget::GetContentForSlot(FName SlotName) const
+{
+	for ( const FNamedSlotBinding& Binding : NamedSlotBindings )
+	{
+		if ( Binding.Name == SlotName )
+		{
+			return Binding.Content;
+		}
+	}
+
+	return nullptr;
+}
+
+void UUserWidget::SetContentForSlot(FName SlotName, UWidget* Content)
+{
+	// Find the binding in the existing set and replace the content for that binding.
+	for ( int32 BindingIndex = 0; BindingIndex < NamedSlotBindings.Num(); BindingIndex++)
+	{
+		FNamedSlotBinding& Binding = NamedSlotBindings[BindingIndex];
+
+		if ( Binding.Name == SlotName )
+		{
+			if ( Content )
+			{
+				Binding.Content = Content;
+			}
+			else
+			{
+				NamedSlotBindings.RemoveAt(BindingIndex);
+			}
+
+			return;
+		}
+	}
+
+	if ( Content )
+	{
+		// Add the new binding to the list of bindings.
+		FNamedSlotBinding NewBinding;
+		NewBinding.Name = SlotName;
+		NewBinding.Content = Content;
+
+		NamedSlotBindings.Add(NewBinding);
+	}
+}
+
 void UUserWidget::NativeTick(const FGeometry& MyGeometry, float InDeltaTime )
 {
 	// Update active movie scenes
@@ -416,7 +502,6 @@ TSharedRef<SWidget> UUserWidget::MakeViewportWidget(TSharedPtr<SWidget>& UserSla
 		.Offset(BIND_UOBJECT_ATTRIBUTE(FMargin, GetFullScreenOffset))
 		.Anchors(BIND_UOBJECT_ATTRIBUTE(FAnchors, GetViewportAnchors))
 		.Alignment(BIND_UOBJECT_ATTRIBUTE(FVector2D, GetFullScreenAlignment))
-		.ZOrder(BIND_UOBJECT_ATTRIBUTE(int32, GetFullScreenZOrder))
 		[
 			UserSlateWidget.ToSharedRef()
 		];
@@ -429,7 +514,7 @@ UWidget* UUserWidget::GetRootWidgetComponent()
 		return Components[0];
 	}
 
-	return NULL;
+	return nullptr;
 }
 
 void UUserWidget::AddToViewport()
@@ -547,11 +632,6 @@ void UUserWidget::SetAlignmentInViewport(FVector2D Alignment)
 	ViewportAlignment = Alignment;
 }
 
-void UUserWidget::SetZOrderInViewport(int32 ZOrder)
-{
-	ViewportZOrder = ZOrder;
-}
-
 FMargin UUserWidget::GetFullScreenOffset() const
 {
 	// If the size is zero, and we're not stretched, then use the desired size.
@@ -578,11 +658,6 @@ FVector2D UUserWidget::GetFullScreenAlignment() const
 	return ViewportAlignment;
 }
 
-int32 UUserWidget::GetFullScreenZOrder() const
-{
-	return ViewportZOrder;
-}
-
 #if WITH_EDITOR
 
 const FSlateBrush* UUserWidget::GetEditorIcon()
diff --git a/Engine/Source/Runtime/UMG/Private/WidgetTree.cpp b/Engine/Source/Runtime/UMG/Private/WidgetTree.cpp
--- a/Engine/Source/Runtime/UMG/Private/WidgetTree.cpp
+++ b/Engine/Source/Runtime/UMG/Private/WidgetTree.cpp
@@ -10,18 +10,15 @@ UWidgetTree::UWidgetTree(const FPostConstructInitializeProperties& PCIP)
 {
 }
 
-UWidget* UWidgetTree::FindWidget(const FString& Name) const
+UWidget* UWidgetTree::FindWidget(const FName& Name) const
 {
-	FString ExistingName;
-
 	// TODO UMG Hacky, remove this find widget function, or make it faster.
 	TArray<UWidget*> Widgets;
 	GetAllWidgets(Widgets);
 
 	for ( UWidget* Widget : Widgets )
 	{
-		Widget->GetName(ExistingName);
-		if ( ExistingName.Equals(Name, ESearchCase::IgnoreCase) )
+		if ( Widget->GetFName() == Name )
 		{
 			return Widget;
 		}
@@ -46,13 +43,13 @@ UWidget* UWidgetTree::FindWidget(TSharedRef<SWidget> InWidget) const
 		}
 	}
 
-	return NULL;
+	return nullptr;
 }
 
 UPanelWidget* UWidgetTree::FindWidgetParent(UWidget* Widget, int32& OutChildIndex)
 {
 	UPanelWidget* Parent = Widget->GetParent();
-	if ( Parent != NULL )
+	if ( Parent != nullptr )
 	{
 		OutChildIndex = Parent->GetChildIndex(Widget);
 	}
@@ -97,6 +94,20 @@ void UWidgetTree::GetAllWidgets(TArray<UWidget*>& Widgets) const
 
 void UWidgetTree::GetChildWidgets(UWidget* Parent, TArray<UWidget*>& Widgets) const
 {
+	if ( INamedSlotInterface* NamedSlotHost = InterfaceCast<INamedSlotInterface>(Parent) )
+	{
+		TArray<FName> SlotNames;
+		NamedSlotHost->GetSlotNames(SlotNames);
+
+		for ( FName SlotName : SlotNames )
+		{
+			if ( UWidget* SlotContent = NamedSlotHost->GetContentForSlot(SlotName) )
+			{
+				Widgets.Add(SlotContent);
+			}
+		}
+	}
+
 	if ( UPanelWidget* PanelParent = Cast<UPanelWidget>(Parent) )
 	{
 		for ( int32 ChildIndex = 0; ChildIndex < PanelParent->GetChildrenCount(); ChildIndex++ )
diff --git a/Engine/Source/Runtime/UMG/Public/Blueprint/UserWidget.h b/Engine/Source/Runtime/UMG/Public/Blueprint/UserWidget.h
--- a/Engine/Source/Runtime/UMG/Public/Blueprint/UserWidget.h
+++ b/Engine/Source/Runtime/UMG/Public/Blueprint/UserWidget.h
@@ -4,6 +4,9 @@
 
 #include "Geometry.h"
 #include "Engine/GameInstance.h"
+
+#include "NamedSlotInterface.h"
+
 #include "UserWidget.generated.h"
 
 static FGeometry NullGeometry;
@@ -69,6 +72,29 @@ public:
 	int32 MaxLayer;
 };
 
+/**
+* The state passed into OnPaint that we can expose as a single painting structure to blueprints to
+* allow script code to override OnPaint behavior.
+*/
+USTRUCT()
+struct UMG_API FNamedSlotBinding
+{
+	GENERATED_USTRUCT_BODY()
+
+public:
+
+	FNamedSlotBinding()
+		: Name(NAME_None)
+		, Content(nullptr)
+	{ }
+
+	UPROPERTY()
+	FName Name;
+
+	UPROPERTY()
+	UWidget* Content;
+};
+
 class UUMGSequencePlayer;
 
 //TODO UMG If you want to host a widget that's full screen there may need to be a SWindow equivalent that you spawn it into.
@@ -80,7 +106,7 @@ DECLARE_DYNAMIC_MULTICAST_DELEGATE_OneParam(FOnVisibilityChangedEvent, ESlateVis
  * The user widget is extensible by users through the WidgetBlueprint.
  */
 UCLASS(Abstract, editinlinenew, BlueprintType, Blueprintable, meta=( Category="User Controls" ))
-class UMG_API UUserWidget : public UWidget
+class UMG_API UUserWidget : public UWidget, public INamedSlotInterface
 {
 	GENERATED_UCLASS_BODY()
 public:
@@ -97,6 +123,12 @@ public:
 	virtual void ReleaseSlateResources(bool bReleaseChildren) override;
 	// End of UVisual interface
 
+	// UNamedSlotInterface Begin
+	virtual void GetSlotNames(TArray<FName>& SlotNames) const override;
+	virtual UWidget* GetContentForSlot(FName SlotName) const override;
+	virtual void SetContentForSlot(FName SlotName, UWidget* Content) override;
+	// UNamedSlotInterface End
+
 	/** Sets that this widget is being designed sets it on all children as well. */
 	virtual void SetIsDesignTime(bool bInDesignTime) override;
 
@@ -131,10 +163,6 @@ public:
 	UFUNCTION(BlueprintCallable, Category="User Interface|Viewport")
 	void SetAlignmentInViewport(FVector2D Alignment);
 
-	/*  */
-	UFUNCTION(BlueprintCallable, Category="User Interface|Viewport")
-	void SetZOrderInViewport(int32 ZOrder);
-
 	/*  */
 	UFUNCTION(BlueprintPure, Category="Appearance")
 	bool GetIsVisible() const;
@@ -272,7 +300,7 @@ public:
 	UWidget* GetRootWidgetComponent();
 
 	/** @returns The slate widget corresponding to a given name */
-	TSharedPtr<SWidget> GetWidgetFromName(const FString& Name) const;
+	TSharedPtr<SWidget> GetWidgetFromName(const FName& Name) const;
 
 	/** @returns The uobject widget corresponding to a given name */
 	UWidget* GetHandleFromName(const FString& Name) const;
@@ -309,19 +337,33 @@ public:
 	UPROPERTY(Transient)
 	TArray<UUMGSequencePlayer*> StoppedSequencePlayers;
 
+	/** Stores the widgets being assigned to named slots */
+	UPROPERTY()
+	TArray<FNamedSlotBinding> NamedSlotBindings;
+
+#if WITH_EDITORONLY_DATA
+
+	/** Stores the design time desired size of the user widget */
+	UPROPERTY()
+	FVector2D DesignTimeSize;
+
+	/** Stores the design time desired size of the user widget */
+	UPROPERTY()
+	bool bUseDesignTimeSize;
+
+#endif
+
 protected:
 	virtual TSharedRef<SWidget> RebuildWidget() override;
 
 	FMargin GetFullScreenOffset() const;
 	FAnchors GetViewportAnchors() const;
 	FVector2D GetFullScreenAlignment() const;
-	int32 GetFullScreenZOrder() const;
 
 private:
 	FAnchors ViewportAnchors;
 	FMargin ViewportOffsets;
 	FVector2D ViewportAlignment;
-	int32 ViewportZOrder;
 
 	TWeakPtr<SWidget> FullScreenWidget;
 
diff --git a/Engine/Source/Runtime/UMG/Public/Blueprint/WidgetTree.h b/Engine/Source/Runtime/UMG/Public/Blueprint/WidgetTree.h
--- a/Engine/Source/Runtime/UMG/Public/Blueprint/WidgetTree.h
+++ b/Engine/Source/Runtime/UMG/Public/Blueprint/WidgetTree.h
@@ -15,7 +15,7 @@ class UMG_API UWidgetTree : public UObject
 public:
 
 	/** Finds the widget in the tree by name. */
-	UWidget* FindWidget(const FString& Name) const;
+	UWidget* FindWidget(const FName& Name) const;
 
 	/** Finds a widget in the tree using the native widget as the key. */
 	UWidget* FindWidget(TSharedRef<SWidget> InWidget) const;
diff --git a/Engine/Source/Runtime/UMG/Public/Components/NamedSlot.h b/Engine/Source/Runtime/UMG/Public/Components/NamedSlot.h
new file mode 100644
--- /dev/null
+++ b/Engine/Source/Runtime/UMG/Public/Components/NamedSlot.h
@@ -0,0 +1,43 @@
+// Copyright 1998-2014 Epic Games, Inc. All Rights Reserved.
+
+#pragma once
+
+#include "NamedSlot.generated.h"
+
+/**
+ * A representation for external socketing of widgets so that user controls can have widget insertion points 
+ * visible in the designer.
+ */
+UCLASS(ClassGroup=UserInterface)
+class UMG_API UNamedSlot : public UContentWidget
+{
+	GENERATED_UCLASS_BODY()
+
+public:
+
+	// UPanelWidget interface
+	virtual void OnSlotAdded(UPanelSlot* Slot) override;
+	virtual void OnSlotRemoved(UPanelSlot* Slot) override;
+	// End of UPanelWidget interface
+
+	// UWidget interface
+	virtual void SynchronizeProperties() override;
+	// End of UWidget interface
+
+	// UVisual interface
+	virtual void ReleaseSlateResources(bool bReleaseChildren) override;
+	// End of UVisual interface
+
+#if WITH_EDITOR
+	virtual const FSlateBrush* GetEditorIcon() override;
+	virtual const FText GetPaletteCategory() override;
+#endif
+
+protected:
+	// UWidget interface
+	virtual TSharedRef<SWidget> RebuildWidget() override;
+	// End of UWidget interface
+
+protected:
+	TSharedPtr<SBox> MyBox;
+};
diff --git a/Engine/Source/Runtime/UMG/Public/Components/NamedSlotInterface.h b/Engine/Source/Runtime/UMG/Public/Components/NamedSlotInterface.h
--- a/Engine/Source/Runtime/UMG/Public/Components/NamedSlotInterface.h
+++ b/Engine/Source/Runtime/UMG/Public/Components/NamedSlotInterface.h
@@ -17,4 +17,10 @@ class UMG_API INamedSlotInterface
 
 	/**  */
 	virtual void GetSlotNames(TArray<FName>& SlotNames) const = 0;
+
+	/**  */
+	virtual UWidget* GetContentForSlot(FName SlotName) const = 0;
+
+	/**  */
+	virtual void SetContentForSlot(FName SlotName, UWidget* Content) = 0;
 };
diff --git a/Engine/Source/Runtime/UMG/Public/Components/ProgressBar.h b/Engine/Source/Runtime/UMG/Public/Components/ProgressBar.h
--- a/Engine/Source/Runtime/UMG/Public/Components/ProgressBar.h
+++ b/Engine/Source/Runtime/UMG/Public/Components/ProgressBar.h
@@ -79,6 +79,7 @@ public:
 	// UWidget interface
 	virtual const FSlateBrush* GetEditorIcon() override;
 	virtual const FText GetPaletteCategory() override;
+	virtual void OnCreationFromPalette() override;
 	// End UWidget interface
 #endif
 
diff --git a/Engine/Source/Runtime/UMG/Public/Components/Widget.h b/Engine/Source/Runtime/UMG/Public/Components/Widget.h
--- a/Engine/Source/Runtime/UMG/Public/Components/Widget.h
+++ b/Engine/Source/Runtime/UMG/Public/Components/Widget.h
@@ -178,13 +178,12 @@ public:
 
 	/**
 	 * Checks to see if this widget is the current mouse captor
-	 *
 	 * @return  True if this widget has captured the mouse
 	 */
 	UFUNCTION(BlueprintCallable, Category="Widget")
 	bool HasMouseCapture() const;
 
-	/**  */
+	/** Sets the focus to this widget. */
 	UFUNCTION(BlueprintCallable, Category="Widget")
 	void SetKeyboardFocus() const;
 
@@ -194,7 +193,7 @@ public:
 
 	/**
 	 * Gets the widgets desired size.
-	 * NOTE: The underlying Slate widget must exist and be valid, also a at least one pre-pass must
+	 * NOTE: The underlying Slate widget must exist and be valid, also at least one pre-pass must
 	 *       have occurred before this value will be of any use.
 	 * 
 	 * @return The widget's desired size
@@ -255,6 +254,12 @@ public:
 	/** Gets the palette category of the widget */
 	virtual const FText GetPaletteCategory();
 
+	/**
+	 * Called by the palette after constructing a new widget, allows the widget to perform interesting 
+	 * default setup that we don't want to be UObject Defaults.
+	 */
+	virtual void OnCreationFromPalette() { }
+
 	/** Gets the editor icon */
 	virtual const FSlateBrush* GetEditorIcon();
 	
diff --git a/Engine/Source/Runtime/UMG/Public/Components/WrapBox.h b/Engine/Source/Runtime/UMG/Public/Components/WrapBox.h
new file mode 100644
--- /dev/null
+++ b/Engine/Source/Runtime/UMG/Public/Components/WrapBox.h
@@ -0,0 +1,58 @@
+// Copyright 1998-2014 Epic Games, Inc. All Rights Reserved.
+
+#pragma once
+
+#include "WrapBox.generated.h"
+
+/**
+* A vertical box widget is a layout panel allowing child widgets to be automatically laid out
+* vertically.
+*/
+UCLASS(ClassGroup=UserInterface)
+class UMG_API UWrapBox : public UPanelWidget
+{
+	GENERATED_UCLASS_BODY()
+
+public:
+	/** The inner slot padding goes between slots sharing borders */
+	UPROPERTY(EditDefaultsOnly, Category="Conetnt Layout")
+	FVector2D InnerSlotPadding;
+
+public:
+
+	/**  */
+	UFUNCTION(BlueprintCallable, Category="Panel")
+	UWrapBoxSlot* AddChildWrapBox(UWidget* Content);
+
+#if WITH_EDITOR
+	// UWidget interface
+	virtual const FSlateBrush* GetEditorIcon() override;
+	virtual const FText GetPaletteCategory() override;
+	// End UWidget interface
+#endif
+
+protected:
+
+	// UPanelWidget
+	virtual UClass* GetSlotClass() const override;
+	virtual void OnSlotAdded(UPanelSlot* Slot) override;
+	virtual void OnSlotRemoved(UPanelSlot* Slot) override;
+	// End UPanelWidget
+
+	// UWidget interface
+	virtual void SynchronizeProperties() override;
+	// End of UWidget interface
+
+	// UVisual interface
+	virtual void ReleaseSlateResources(bool bReleaseChildren) override;
+	// End of UVisual interface
+
+protected:
+
+	TSharedPtr<class SWrapBox> MyWrapBox;
+
+protected:
+	// UWidget interface
+	virtual TSharedRef<SWidget> RebuildWidget() override;
+	// End of UWidget interface
+};
diff --git a/Engine/Source/Runtime/UMG/Public/Components/WrapBoxSlot.h b/Engine/Source/Runtime/UMG/Public/Components/WrapBoxSlot.h
new file mode 100644
--- /dev/null
+++ b/Engine/Source/Runtime/UMG/Public/Components/WrapBoxSlot.h
@@ -0,0 +1,75 @@
+
+
+// Copyright 1998-2014 Epic Games, Inc. All Rights Reserved.
+
+#pragma once
+
+#include "SlateWrapperTypes.h"
+
+#include "WrapBoxSlot.generated.h"
+
+/** The Slot for the UWrapBox, contains the widget that is flowed vertically */
+UCLASS()
+class UMG_API UWrapBoxSlot : public UPanelSlot
+{
+	GENERATED_UCLASS_BODY()
+
+public:
+
+	/** The padding area between the slot and the content it contains. */
+	UPROPERTY(EditDefaultsOnly, Category="Layout (Wrap Box Slot)")
+	FMargin Padding;
+
+	/** Should this slot fill the remaining space on the line? */
+	UPROPERTY(EditDefaultsOnly, Category="Layout (Wrap Box Slot)")
+	bool bFillEmptySpace;
+
+	/**
+	 * If the total available space in the wrap panel drops below this threshold, this slot will attempt to fill an entire line.
+	 * NOTE: A value of 0, denotes no filling will occur.
+	 */
+	UPROPERTY(EditDefaultsOnly, Category="Layout (Wrap Box Slot)")
+	float FillSpanWhenLessThan;
+
+	/** The alignment of the object horizontally. */
+	UPROPERTY(EditDefaultsOnly, Category="Layout (Wrap Box Slot)")
+	TEnumAsByte<EHorizontalAlignment> HorizontalAlignment;
+
+	/** The alignment of the object vertically. */
+	UPROPERTY(EditDefaultsOnly, Category="Layout (Wrap Box Slot)")
+	TEnumAsByte<EVerticalAlignment> VerticalAlignment;
+
+public:
+
+	UFUNCTION(BlueprintCallable, Category="Layout (Wrap Box Slot)")
+	void SetPadding(FMargin InPadding);
+
+	UFUNCTION(BlueprintCallable, Category="Layout (Wrap Box Slot)")
+	void SetFillEmptySpace(bool InbFillEmptySpace);
+
+	UFUNCTION(BlueprintCallable, Category="Layout (Wrap Box Slot)")
+	void SetFillSpanWhenLessThan(float InFillSpanWhenLessThan);
+
+	UFUNCTION(BlueprintCallable, Category="Layout (Wrap Box Slot)")
+	void SetHorizontalAlignment(EHorizontalAlignment InHorizontalAlignment);
+
+	UFUNCTION(BlueprintCallable, Category="Layout (Wrap Box Slot)")
+	void SetVerticalAlignment(EVerticalAlignment InVerticalAlignment);
+
+public:
+
+	// UPanelSlot interface
+	virtual void SynchronizeProperties() override;
+	// End of UPanelSlot interface
+
+	virtual void ReleaseSlateResources(bool bReleaseChildren) override;
+
+	/** Builds the underlying FSlot for the Slate layout panel. */
+	void BuildSlot(TSharedRef<SWrapBox> InWrapBox);
+
+private:
+	//TODO UMG Slots should hold weak or shared refs to slots.
+
+	/** A raw pointer to the slot to allow us to adjust the size, padding...etc at runtime. */
+	SWrapBox::FSlot* Slot;
+};
diff --git a/Engine/Source/Runtime/UMG/Public/UMG.h b/Engine/Source/Runtime/UMG/Public/UMG.h
--- a/Engine/Source/Runtime/UMG/Public/UMG.h
+++ b/Engine/Source/Runtime/UMG/Public/UMG.h
@@ -32,11 +32,12 @@
 
 #include "Widget.h"
 
-#include "NamedSlotInterface.h"
-
 #include "PanelWidget.h"
 #include "ContentWidget.h"
 
+#include "NamedSlot.h"
+#include "NamedSlotInterface.h"
+
 #include "CanvasPanelSlot.h"
 #include "CanvasPanel.h"
 
@@ -46,6 +47,9 @@
 #include "VerticalBoxSlot.h"
 #include "VerticalBox.h"
 
+#include "WrapBoxSlot.h"
+#include "WrapBox.h"
+
 #include "UniformGridSlot.h"
 #include "UniformGridPanel.h"
 
