UMG - Cleaning up some things in the anchors menu.

diff --git a/Engine/Source/Editor/UMGEditor/Private/Customizations/CanvasSlotCustomization.cpp b/Engine/Source/Editor/UMGEditor/Private/Customizations/CanvasSlotCustomization.cpp
--- a/Engine/Source/Editor/UMGEditor/Private/Customizations/CanvasSlotCustomization.cpp
+++ b/Engine/Source/Editor/UMGEditor/Private/Customizations/CanvasSlotCustomization.cpp
@@ -34,8 +34,7 @@ public:
 		[
 			SNew(SButton)
 			.ButtonStyle(FEditorStyle::Get(), "SimpleSharpButton")
-			//.ButtonColorAndOpacity(FLinearColor(FColor(64, 64, 64)))
-			.ButtonColorAndOpacity(FLinearColor(FColor(48, 48, 48)))
+			.ButtonColorAndOpacity(FLinearColor(FColor(40, 40, 40)))
 			.OnClicked(this, &SAnchorPreviewWidget::OnAnchorClicked, AnchorsHandle, AlignmentHandle, OffsetsHandle, Anchors)
 			.ContentPadding(FMargin(2.0f, 2.0f))
 			[
@@ -44,52 +43,41 @@ public:
 				+ SVerticalBox::Slot()
 				.AutoHeight()
 				[
-					SNew(SBox)
-					.WidthOverride(64)
-					.HeightOverride(64)
-					.HAlign(HAlign_Center)
-					.VAlign(VAlign_Center)
+					SNew(SBorder)
+					.BorderImage(FEditorStyle::GetBrush("UMGEditor.AnchorGrid"))
+					.Padding(0)
 					[
 						SNew(SBox)
-						.WidthOverride(this, &SAnchorPreviewWidget::GetCurrentWidth)
-						.HeightOverride(this, &SAnchorPreviewWidget::GetCurrentHeight)
+						.WidthOverride(64)
+						.HeightOverride(64)
+						.HAlign(HAlign_Center)
+						.VAlign(VAlign_Center)
 						[
-							SNew(SBorder)
-							.BorderImage(FEditorStyle::GetBrush("ToolPanel.GroupBorder"))
-							.Padding(1)
+							SNew(SBox)
+							.WidthOverride(this, &SAnchorPreviewWidget::GetCurrentWidth)
+							.HeightOverride(this, &SAnchorPreviewWidget::GetCurrentHeight)
 							[
-								SNew(SConstraintCanvas)
-
-								+ SConstraintCanvas::Slot()
-								.Anchors(Anchors)
-								.Offset(FMargin(0, 0, Anchors.IsStretchedHorizontal() ? 0 : 15, Anchors.IsStretchedVertical() ? 0 : 15))
-								.Alignment(FVector2D(Anchors.IsStretchedHorizontal() ? 0 : Anchors.Minimum.X, Anchors.IsStretchedVertical() ? 0 : Anchors.Minimum.Y))
+								SNew(SBorder)
+								//.BorderImage(FEditorStyle::GetBrush("NoBrush"))
+								//.BorderImage(FEditorStyle::GetBrush("ToolPanel.GroupBorder"))
+								//.BorderBackgroundColor
+								.Padding(1)
 								[
-									SNew(SImage)
-									.Image(FEditorStyle::Get().GetBrush("UMGEditor.AnchoredWidget"))
+									SNew(SConstraintCanvas)
+
+									+ SConstraintCanvas::Slot()
+									.Anchors(Anchors)
+									.Offset(FMargin(0, 0, Anchors.IsStretchedHorizontal() ? 0 : 15, Anchors.IsStretchedVertical() ? 0 : 15))
+									.Alignment(FVector2D(Anchors.IsStretchedHorizontal() ? 0 : Anchors.Minimum.X, Anchors.IsStretchedVertical() ? 0 : Anchors.Minimum.Y))
+									[
+										SNew(SImage)
+										.Image(FEditorStyle::Get().GetBrush("UMGEditor.AnchoredWidget"))
+									]
 								]
-
-								//+ SConstraintCanvas::Slot()
-								//.Anchors(Anchors)
-								//.Offset(FMargin(0, 0, 8, 8))
-								//.Alignment(FVector2D(Anchors.IsStretchedHorizontal() ? 0 : Anchors.Minimum.X, Anchors.IsStretchedVertical() ? 0 : Anchors.Minimum.Y))
-								//[
-								//	SNew(SImage)
-								//	.Image(FEditorStyle::Get().GetBrush("UMGEditor.AnchoredWidgetAlignment"))
-								//]
 							]
 						]
 					]
 				]
-
-				//+ SVerticalBox::Slot()
-				//.AutoHeight()
-				//.HAlign(HAlign_Center)
-				//.Padding(0,2,0,0)
-				//[
-				//	SNew(STextBlock)
-				//	.Text(LabelText)
-				//]
 			]
 		];
 	}
