Fixing Warning

diff --git a/Engine/Source/Editor/UMGEditor/Private/WidgetBlueprintEditorUtils.cpp b/Engine/Source/Editor/UMGEditor/Private/WidgetBlueprintEditorUtils.cpp
--- a/Engine/Source/Editor/UMGEditor/Private/WidgetBlueprintEditorUtils.cpp
+++ b/Engine/Source/Editor/UMGEditor/Private/WidgetBlueprintEditorUtils.cpp
@@ -263,7 +263,7 @@ bool FWidgetBlueprintEditorUtils::FindAndRemoveNamedSlotContent(UWidget* WidgetT
 
 	for ( UWidget* Widget : AllWidgets )
 	{
-		if ( INamedSlotInterface* NamedSlotHost = InterfaceCast<INamedSlotInterface>(Widget) )
+		if ( INamedSlotInterface* NamedSlotHost = Cast<INamedSlotInterface>(Widget) )
 		{
 			TArray<FName> SlotNames;
 			NamedSlotHost->GetSlotNames(SlotNames);
