Fixing InterfaceCast warnings

diff --git a/Engine/Source/Editor/UMGEditor/Private/Hierarchy/SHierarchyViewItem.cpp b/Engine/Source/Editor/UMGEditor/Private/Hierarchy/SHierarchyViewItem.cpp
--- a/Engine/Source/Editor/UMGEditor/Private/Hierarchy/SHierarchyViewItem.cpp
+++ b/Engine/Source/Editor/UMGEditor/Private/Hierarchy/SHierarchyViewItem.cpp
@@ -340,7 +340,7 @@ FName FNamedSlotModel::GetUniqueName() const
 
 FText FNamedSlotModel::GetText() const
 {
-	if ( INamedSlotInterface* NamedSlotHost = InterfaceCast<INamedSlotInterface>(Item.GetPreview()) )
+	if ( INamedSlotInterface* NamedSlotHost = Cast<INamedSlotInterface>(Item.GetPreview()) )
 	{
 		TSet<FWidgetReference> SelectedWidgets;
 		if ( UWidget* SlotContent = NamedSlotHost->GetContentForSlot(SlotName) )
@@ -377,7 +377,7 @@ void FNamedSlotModel::GetChildren(TArray< TSharedPtr<FHierarchyModel> >& Childre
 void FNamedSlotModel::OnSelection()
 {
 	TSharedPtr<FWidgetBlueprintEditor> BPEd = BlueprintEditor.Pin();
-	if ( INamedSlotInterface* NamedSlotHost = InterfaceCast<INamedSlotInterface>(Item.GetPreview()) )
+	if ( INamedSlotInterface* NamedSlotHost = Cast<INamedSlotInterface>(Item.GetPreview()) )
 	{
 		TSet<FWidgetReference> SelectedWidgets;
 		if ( UWidget* SlotContent = NamedSlotHost->GetContentForSlot(SlotName) )
@@ -408,7 +408,7 @@ FReply FNamedSlotModel::HandleDrop(FDragDropEvent const& DragDropEvent)
 	TSharedPtr<FWidgetTemplateDragDropOp> TemplateDragDropOp = DragDropEvent.GetOperationAs<FWidgetTemplateDragDropOp>();
 	if ( TemplateDragDropOp.IsValid() )
 	{
-		if ( INamedSlotInterface* NamedSlotHost = InterfaceCast<INamedSlotInterface>(Item.GetTemplate()) )
+		if ( INamedSlotInterface* NamedSlotHost = Cast<INamedSlotInterface>(Item.GetTemplate()) )
 		{
 			UWidget* Widget = TemplateDragDropOp->Template->Create(Blueprint->WidgetTree);
 
@@ -567,7 +567,7 @@ void FHierarchyWidget::GetChildren(TArray< TSharedPtr<FHierarchyModel> >& Childr
 	TSharedPtr<FWidgetBlueprintEditor> BPEd = BlueprintEditor.Pin();
 
 	// Check for named slots
-	if ( INamedSlotInterface* NamedSlotHost = InterfaceCast<INamedSlotInterface>(Item.GetPreview()) )
+	if ( INamedSlotInterface* NamedSlotHost = Cast<INamedSlotInterface>(Item.GetPreview()) )
 	{
 		TArray<FName> SlotNames;
 		NamedSlotHost->GetSlotNames(SlotNames);
