UMG - RemoveFromParent now works for items added tot he viewport.  Deprecating RemoveFromViewport.  You can now adjust the visibility of a widget for design time only purposes.  Deprecating IsVisible, it wasn't a good description of the widget's actual visibility, created a new method that better represents what the old method did, IsInViewport.

diff --git a/Engine/Source/Editor/EditorStyle/Private/SlateEditorStyle.cpp b/Engine/Source/Editor/EditorStyle/Private/SlateEditorStyle.cpp
--- a/Engine/Source/Editor/EditorStyle/Private/SlateEditorStyle.cpp
+++ b/Engine/Source/Editor/EditorStyle/Private/SlateEditorStyle.cpp
@@ -5864,6 +5864,9 @@ void FSlateEditorStyle::FStyle::SetupUMGEditorStyles()
 	Set("WidgetDesigner.ZoomToFit", new IMAGE_BRUSH("Icons/UMG/Fit_16x", Icon16x16));
 	Set("WidgetDesigner.ZoomToFit.Small", new IMAGE_BRUSH("Icons/UMG/Fit_16x", Icon16x16));
 
+	Set("WidgetDesigner.WidgetVisible", new IMAGE_BRUSH("/Icons/icon_layer_visible", Icon16x16));
+	Set("WidgetDesigner.WidgetHidden", new IMAGE_BRUSH("/Icons/icon_layer_not_visible", Icon16x16));
+
 	Set("UMGEditor.ZoomToFit", new IMAGE_BRUSH("Icons/UMG/Fit_16x", Icon16x16, FLinearColor(.05f, .05f, .05f, 1.f)));
 
 	Set("UMGEditor.AnchorGizmo.Center", new IMAGE_BRUSH("Icons/UMG/AnchorGizmo/center", Icon16x16));
diff --git a/Engine/Source/Editor/UMGEditor/Private/Hierarchy/SHierarchyViewItem.cpp b/Engine/Source/Editor/UMGEditor/Private/Hierarchy/SHierarchyViewItem.cpp
--- a/Engine/Source/Editor/UMGEditor/Private/Hierarchy/SHierarchyViewItem.cpp
+++ b/Engine/Source/Editor/UMGEditor/Private/Hierarchy/SHierarchyViewItem.cpp
@@ -642,6 +642,26 @@ void SHierarchyViewItem::Construct(const FArguments& InArgs, const TSharedRef< S
 				.OnTextCommitted(this, &SHierarchyViewItem::OnNameTextCommited)
 				.IsSelected(this, &SHierarchyViewItem::IsSelectedExclusively)
 			]
+
+			// Visibility icon
+			+ SHorizontalBox::Slot()
+			.AutoWidth()
+			.VAlign(VAlign_Center)
+			[
+				SNew(SButton)
+				.ContentPadding(0)
+				.ButtonStyle(FEditorStyle::Get(), "HoverHintOnly")
+				.OnClicked(this, &SHierarchyViewItem::OnToggleVisibility)
+				.Visibility(Model->CanControlVisibility() ? EVisibility::Visible : EVisibility::Hidden)
+				.ToolTipText(LOCTEXT("WidgetVisibilityButtonToolTip", "Toggle Widget's Editor Visibility"))
+				.HAlign(HAlign_Center)
+				.VAlign(VAlign_Center)
+				.Content()
+				[
+					SNew(SImage)
+					.Image(this, &SHierarchyViewItem::GetVisibilityBrushForWidget)
+				]
+			]
 		],
 		InOwnerTableView);
 }
@@ -691,4 +711,18 @@ FReply SHierarchyViewItem::HandleDrop(const FDragDropEvent& DragDropEvent)
 	return Model->HandleDrop(DragDropEvent);
 }
 
+FReply SHierarchyViewItem::OnToggleVisibility()
+{
+	Model->SetIsVisible(!Model->IsVisible());
+
+	return FReply::Handled();
+}
+
+const FSlateBrush* SHierarchyViewItem::GetVisibilityBrushForWidget() const
+{
+	return Model->IsVisible() ?
+		FEditorStyle::GetBrush("WidgetDesigner.WidgetVisible") :
+		FEditorStyle::GetBrush("WidgetDesigner.WidgetHidden");
+}
+
 #undef LOCTEXT_NAMESPACE
diff --git a/Engine/Source/Editor/UMGEditor/Private/Hierarchy/SHierarchyViewItem.h b/Engine/Source/Editor/UMGEditor/Private/Hierarchy/SHierarchyViewItem.h
--- a/Engine/Source/Editor/UMGEditor/Private/Hierarchy/SHierarchyViewItem.h
+++ b/Engine/Source/Editor/UMGEditor/Private/Hierarchy/SHierarchyViewItem.h
@@ -37,6 +37,10 @@ public:
 
 	virtual bool IsModel(UObject* PossibleModel) const { return false; }
 
+	virtual bool IsVisible() const { return true; }
+	virtual bool CanControlVisibility() const { return false; }
+	virtual void SetIsVisible(bool IsVisible) { }
+
 protected:
 	virtual void GetChildren(TArray< TSharedPtr<FHierarchyModel> >& Children) = 0;
 
@@ -142,6 +146,26 @@ public:
 		return Item.GetTemplate() == PossibleModel;
 	}
 
+	virtual bool IsVisible() const override
+	{
+		return !Item.GetTemplate()->bHiddenInDesigner;
+	}
+
+	virtual bool CanControlVisibility() const override
+	{
+		return true;
+	}
+
+	virtual void SetIsVisible(bool IsVisible) override
+	{
+		Item.GetTemplate()->bHiddenInDesigner = !IsVisible;
+
+		if ( UWidget* PreviewWidget = Item.GetPreview() )
+		{
+			PreviewWidget->bHiddenInDesigner = !IsVisible;
+		}
+	}
+
 protected:
 	virtual void GetChildren(TArray< TSharedPtr<FHierarchyModel> >& Children) override;
 
@@ -179,15 +203,21 @@ private:
 	/** Called when text is being committed to check for validity */
 	bool OnVerifyNameTextChanged(const FText& InText, FText& OutErrorMessage);
 
-	/* Called when text is committed on the node */
+	/** Called when text is committed on the node */
 	void OnNameTextCommited(const FText& InText, ETextCommit::Type CommitInfo);
 
-	/* Gets the font to use for the text item, bold for customized named items */
+	/** Gets the font to use for the text item, bold for customized named items */
 	FSlateFontInfo GetItemFont() const;
 
-	/* @returns the widget name to use for the tree item */
+	/** @returns the widget name to use for the tree item */
 	FText GetItemText() const;
 
+	/** Handles clicking the visibility toggle */
+	FReply OnToggleVisibility();
+
+	/** Returns a brush representing the visibility item of the widget's visibility button */
+	const FSlateBrush* GetVisibilityBrushForWidget() const;
+
 	/* The mode that this tree item represents */
 	TSharedPtr<FHierarchyModel> Model;
 };
diff --git a/Engine/Source/Runtime/UMG/Private/Components/Widget.cpp b/Engine/Source/Runtime/UMG/Private/Components/Widget.cpp
--- a/Engine/Source/Runtime/UMG/Private/Components/Widget.cpp
+++ b/Engine/Source/Runtime/UMG/Private/Components/Widget.cpp
@@ -332,6 +332,11 @@ const FSlateBrush* UWidget::GetEditorIcon()
 	return FUMGStyle::Get().GetBrush("Widget");
 }
 
+EVisibility UWidget::GetVisibilityInDesigner() const
+{
+	return bHiddenInDesigner ? EVisibility::Collapsed : EVisibility::Visible;
+}
+
 void UWidget::PostEditChangeProperty(FPropertyChangedEvent& PropertyChangedEvent)
 {
 	Super::PostEditChangeProperty(PropertyChangedEvent);
@@ -409,13 +414,15 @@ void UWidget::SynchronizeProperties()
 	// visibility and enabled status are not stomping values setup in the root widget in the User Widget.
 	TSharedPtr<SWidget> SafeWidget = GetCachedWidget();
 
+#if WITH_EDITOR
 	// Always use an enabled and visible state in the designer.
 	if ( IsDesignTime() )
 	{
 		SafeWidget->SetEnabled(true);
-		SafeWidget->SetVisibility(EVisibility::Visible);
+		SafeWidget->SetVisibility(TAttribute<EVisibility>::Create(TAttribute<EVisibility>::FGetter::CreateUObject(this, &UWidget::GetVisibilityInDesigner)));
 	}
 	else
+#endif
 	{
 		SafeWidget->SetEnabled(OPTIONAL_BINDING(bool, bIsEnabled));
 		SafeWidget->SetVisibility(OPTIONAL_BINDING_CONVERT(ESlateVisibility::Type, Visiblity, EVisibility, ConvertVisibility));
diff --git a/Engine/Source/Runtime/UMG/Private/UserWidget.cpp b/Engine/Source/Runtime/UMG/Private/UserWidget.cpp
--- a/Engine/Source/Runtime/UMG/Private/UserWidget.cpp
+++ b/Engine/Source/Runtime/UMG/Private/UserWidget.cpp
@@ -541,6 +541,11 @@ void UUserWidget::AddToViewport()
 }
 
 void UUserWidget::RemoveFromViewport()
+{
+	RemoveFromParent();
+}
+
+void UUserWidget::RemoveFromParent()
 {
 	if ( FullScreenWidget.IsValid() )
 	{
@@ -556,6 +561,10 @@ void UUserWidget::RemoveFromViewport()
 			}
 		}
 	}
+	else
+	{
+		Super::RemoveFromParent();
+	}
 }
 
 bool UUserWidget::GetIsVisible() const
@@ -563,6 +572,11 @@ bool UUserWidget::GetIsVisible() const
 	return FullScreenWidget.IsValid();
 }
 
+bool UUserWidget::IsInViewport() const
+{
+	return FullScreenWidget.IsValid();
+}
+
 TEnumAsByte<ESlateVisibility::Type> UUserWidget::GetVisiblity() const
 {
 	if ( FullScreenWidget.IsValid() )
diff --git a/Engine/Source/Runtime/UMG/Public/Blueprint/UserWidget.h b/Engine/Source/Runtime/UMG/Public/Blueprint/UserWidget.h
--- a/Engine/Source/Runtime/UMG/Public/Blueprint/UserWidget.h
+++ b/Engine/Source/Runtime/UMG/Public/Blueprint/UserWidget.h
@@ -142,9 +142,12 @@ public:
 	/**
 	 * Removes the widget from the viewport.
 	 */
-	UFUNCTION(BlueprintCallable, Category="User Interface|Viewport")
+	UFUNCTION(BlueprintCallable, Category="User Interface|Viewport", meta=( DeprecatedFunction, DeprecationMessage="Use RemoveFromParent instead" ))
 	void RemoveFromViewport();
 
+	/** Removes the widget from it's parent widget, including the viewport if it was added to the viewport. */
+	virtual void RemoveFromParent() override;
+
 	/**
 	 * Sets the widgets position in the viewport.
 	 */
@@ -164,9 +167,13 @@ public:
 	void SetAlignmentInViewport(FVector2D Alignment);
 
 	/*  */
-	UFUNCTION(BlueprintPure, Category="Appearance")
+	UFUNCTION(BlueprintPure, Category="Appearance", meta=( DeprecatedFunction, DeprecationMessage="Use IsInViewport instead" ))
 	bool GetIsVisible() const;
 
+	/* @return true if the widget was added to the viewport using AddToViewport. */
+	UFUNCTION(BlueprintPure, Category="Appearance")
+	bool IsInViewport() const;
+
 	UFUNCTION(BlueprintPure, Category="Appearance")
 	TEnumAsByte<ESlateVisibility::Type> GetVisiblity() const;
 
diff --git a/Engine/Source/Runtime/UMG/Public/Components/Widget.h b/Engine/Source/Runtime/UMG/Public/Components/Widget.h
--- a/Engine/Source/Runtime/UMG/Public/Components/Widget.h
+++ b/Engine/Source/Runtime/UMG/Public/Components/Widget.h
@@ -106,14 +106,22 @@ public:
 	UPROPERTY()
 	FGetMouseCursor CursorDelegate;
 
-	/**  */
+	/** The render transform of the widget allows for arbitrary 2D transforms to be applied to the widget. */
 	UPROPERTY(EditDefaultsOnly, Category=RenderTransform, meta=(DisplayName="Transform"))
 	FWidgetTransform RenderTransform;
 
-	/**  */
+	/** The render transform pivot controls the location about which transforms are applied.  This value is a normalized coordinate about which things like rotations will occur. */
 	UPROPERTY(EditDefaultsOnly, Category=RenderTransform, meta=(DisplayName="Pivot"))
 	FVector2D RenderTransformPivot;
 
+#if WITH_EDITORONLY_DATA
+
+	/** Stores the design time flag setting if the widget is hidden inside the designer */
+	UPROPERTY()
+	bool bHiddenInDesigner;
+
+#endif
+
 	/** */
 	UFUNCTION(BlueprintCallable, Category="Widget|Transform")
 	void SetRenderTransform(FWidgetTransform InTransform);
@@ -207,7 +215,7 @@ public:
 
 	/** Removes the widget from it's parent widget */
 	UFUNCTION(BlueprintCallable, Category="Widget")
-	void RemoveFromParent();
+	virtual void RemoveFromParent();
 
 	/**
 	 * Gets the underlying slate widget or constructs it if it doesn't exist.  This function is
@@ -270,6 +278,9 @@ public:
 	virtual void PostEditChangeProperty(struct FPropertyChangedEvent& PropertyChangedEvent) override;
 	// End of UObject interface
 
+	/** Gets the visibility of the widget inside the designer. */
+	EVisibility GetVisibilityInDesigner() const;
+
 	// Begin Designer contextual events
 	void Select();
 	void Deselect();
