UMG - Changing the editor visibility now adjusts the modified state of the document (because that info is stored in the asset).

diff --git a/Engine/Source/Editor/UMGEditor/Private/Hierarchy/SHierarchyViewItem.h b/Engine/Source/Editor/UMGEditor/Private/Hierarchy/SHierarchyViewItem.h
--- a/Engine/Source/Editor/UMGEditor/Private/Hierarchy/SHierarchyViewItem.h
+++ b/Engine/Source/Editor/UMGEditor/Private/Hierarchy/SHierarchyViewItem.h
@@ -164,6 +164,9 @@ public:
 		{
 			PreviewWidget->bHiddenInDesigner = !IsVisible;
 		}
+
+		// Mark the blueprint as modified
+		FBlueprintEditorUtils::MarkBlueprintAsModified(BlueprintEditor.Pin()->GetBlueprintObj());
 	}
 
 protected:
