UMG - Updating the class documentation on UMG widgets to be better.  Now displaying the comment data correctly in the palette view.

diff --git a/Engine/Source/Editor/UMGEditor/Private/Palette/SPaletteView.cpp b/Engine/Source/Editor/UMGEditor/Private/Palette/SPaletteView.cpp
--- a/Engine/Source/Editor/UMGEditor/Private/Palette/SPaletteView.cpp
+++ b/Engine/Source/Editor/UMGEditor/Private/Palette/SPaletteView.cpp
@@ -46,6 +46,7 @@ public:
 		ChildSlot
 		[
 			SNew(SHorizontalBox)
+			.Visibility(EVisibility::Visible)
 			.ToolTip(Template->GetToolTip())
 
 			+ SHorizontalBox::Slot()
diff --git a/Engine/Source/Editor/UMGEditor/Private/Templates/WidgetTemplateClass.cpp b/Engine/Source/Editor/UMGEditor/Private/Templates/WidgetTemplateClass.cpp
--- a/Engine/Source/Editor/UMGEditor/Private/Templates/WidgetTemplateClass.cpp
+++ b/Engine/Source/Editor/UMGEditor/Private/Templates/WidgetTemplateClass.cpp
@@ -50,7 +50,7 @@ const FSlateBrush* FWidgetTemplateClass::GetIcon() const
 
 TSharedRef<IToolTip> FWidgetTemplateClass::GetToolTip() const
 {
-	return IDocumentation::Get()->CreateToolTip(FText::FromString(WidgetClass->GetDescription()), nullptr, FString(TEXT("Shared/Types/")) + WidgetClass->GetName(), TEXT("Class"));
+	return IDocumentation::Get()->CreateToolTip(WidgetClass->GetToolTipText(), nullptr, FString(TEXT("Shared/Types/")) + WidgetClass->GetName(), TEXT("Class"));
 }
 
 void FWidgetTemplateClass::OnObjectsReplaced(const TMap<UObject*, UObject*>& ReplacementMap)
diff --git a/Engine/Source/Runtime/UMG/Public/Components/CanvasPanel.h b/Engine/Source/Runtime/UMG/Public/Components/CanvasPanel.h
--- a/Engine/Source/Runtime/UMG/Public/Components/CanvasPanel.h
+++ b/Engine/Source/Runtime/UMG/Public/Components/CanvasPanel.h
@@ -4,6 +4,12 @@
 
 #include "CanvasPanel.generated.h"
 
+/**
+ * The canvas panel is a designer friendly panel that allows widgets to be laid out at arbitrary 
+ * locations, anchored and z-ordered with other children of the canvas.  The canvas is a great widget
+ * for manual layout, but bad when you want to procedurally just generate widgets and place them in a 
+ * container (unless you want absolute layout).
+ */
 UCLASS(ClassGroup=UserInterface)
 class UMG_API UCanvasPanel : public UPanelWidget
 {
diff --git a/Engine/Source/Runtime/UMG/Public/Components/CircularThrobber.h b/Engine/Source/Runtime/UMG/Public/Components/CircularThrobber.h
--- a/Engine/Source/Runtime/UMG/Public/Components/CircularThrobber.h
+++ b/Engine/Source/Runtime/UMG/Public/Components/CircularThrobber.h
@@ -6,7 +6,9 @@
 
 class USlateBrushAsset;
 
-/**  A throbber widget that orients images in a spinning circle. */
+/**
+ * A throbber widget that orients images in a spinning circle.
+ */
 UCLASS(ClassGroup=UserInterface)
 class UMG_API UCircularThrobber : public UWidget
 {
diff --git a/Engine/Source/Runtime/UMG/Public/Components/EditableText.h b/Engine/Source/Runtime/UMG/Public/Components/EditableText.h
--- a/Engine/Source/Runtime/UMG/Public/Components/EditableText.h
+++ b/Engine/Source/Runtime/UMG/Public/Components/EditableText.h
@@ -6,7 +6,9 @@
 
 #include "EditableText.generated.h"
 
-/** Editable text box widget */
+/**
+ * Editable text box widget
+ */
 UCLASS(meta=( Category="Primitive" ), ClassGroup=UserInterface)
 class UMG_API UEditableText : public UWidget
 {
diff --git a/Engine/Source/Runtime/UMG/Public/Components/GridPanel.h b/Engine/Source/Runtime/UMG/Public/Components/GridPanel.h
--- a/Engine/Source/Runtime/UMG/Public/Components/GridPanel.h
+++ b/Engine/Source/Runtime/UMG/Public/Components/GridPanel.h
@@ -4,7 +4,9 @@
 
 #include "GridPanel.generated.h"
 
-/** A panel that evenly divides up available space between all of its children. */
+/**
+ * A panel that evenly divides up available space between all of its children.
+ */
 UCLASS(ClassGroup=UserInterface)
 class UMG_API UGridPanel : public UPanelWidget
 {
diff --git a/Engine/Source/Runtime/UMG/Public/Components/HorizontalBox.h b/Engine/Source/Runtime/UMG/Public/Components/HorizontalBox.h
--- a/Engine/Source/Runtime/UMG/Public/Components/HorizontalBox.h
+++ b/Engine/Source/Runtime/UMG/Public/Components/HorizontalBox.h
@@ -4,6 +4,9 @@
 
 #include "HorizontalBox.generated.h"
 
+/**
+ * Allows widgets to be laid out in a flow horizontally.
+ */
 UCLASS(ClassGroup=UserInterface)
 class UMG_API UHorizontalBox : public UPanelWidget
 {
diff --git a/Engine/Source/Runtime/UMG/Public/Components/ListView.h b/Engine/Source/Runtime/UMG/Public/Components/ListView.h
--- a/Engine/Source/Runtime/UMG/Public/Components/ListView.h
+++ b/Engine/Source/Runtime/UMG/Public/Components/ListView.h
@@ -4,7 +4,9 @@
 
 #include "ListView.generated.h"
 
-/** Allows thousands of items to be displayed in a list.  Generates widgets dynamically for each item. */
+/**
+ * Allows thousands of items to be displayed in a list.  Generates widgets dynamically for each item.
+ */
 UCLASS(Experimental, ClassGroup=UserInterface)
 class UMG_API UListView : public UTableViewBase
 {
diff --git a/Engine/Source/Runtime/UMG/Public/Components/MultiLineEditableText.h b/Engine/Source/Runtime/UMG/Public/Components/MultiLineEditableText.h
--- a/Engine/Source/Runtime/UMG/Public/Components/MultiLineEditableText.h
+++ b/Engine/Source/Runtime/UMG/Public/Components/MultiLineEditableText.h
@@ -4,7 +4,9 @@
 
 #include "MultiLineEditableText.generated.h"
 
-/** Editable text box widget */
+/**
+ * Editable text box widget
+ */
 UCLASS(ClassGroup=UserInterface)
 class UMG_API UMultiLineEditableText : public UWidget
 {
diff --git a/Engine/Source/Runtime/UMG/Public/Components/NamedSlot.h b/Engine/Source/Runtime/UMG/Public/Components/NamedSlot.h
--- a/Engine/Source/Runtime/UMG/Public/Components/NamedSlot.h
+++ b/Engine/Source/Runtime/UMG/Public/Components/NamedSlot.h
@@ -5,8 +5,8 @@
 #include "NamedSlot.generated.h"
 
 /**
- * A representation for external socketing of widgets so that user controls can have widget insertion points 
- * visible in the designer.
+ * Allows you to expose an external slot for your user widget.  When others reuse your user control, 
+ * they can put whatever they want in this named slot.
  */
 UCLASS(ClassGroup=UserInterface)
 class UMG_API UNamedSlot : public UContentWidget
diff --git a/Engine/Source/Runtime/UMG/Public/Components/Overlay.h b/Engine/Source/Runtime/UMG/Public/Components/Overlay.h
--- a/Engine/Source/Runtime/UMG/Public/Components/Overlay.h
+++ b/Engine/Source/Runtime/UMG/Public/Components/Overlay.h
@@ -4,6 +4,9 @@
 
 #include "Overlay.generated.h"
 
+/**
+ * Allows widgets to be stacked on top of each other, uses simple flow layout for content on each layer.
+ */
 UCLASS(ClassGroup=UserInterface)
 class UMG_API UOverlay : public UPanelWidget
 {
diff --git a/Engine/Source/Runtime/UMG/Public/Components/ProgressBar.h b/Engine/Source/Runtime/UMG/Public/Components/ProgressBar.h
--- a/Engine/Source/Runtime/UMG/Public/Components/ProgressBar.h
+++ b/Engine/Source/Runtime/UMG/Public/Components/ProgressBar.h
@@ -6,7 +6,9 @@
 
 #include "ProgressBar.generated.h"
 
-/** ProgressBar widget */
+/**
+ * The progress bar widget is a simple bar that fills up that can be restyled to fit any number of uses.
+ */
 UCLASS(ClassGroup=UserInterface)
 class UMG_API UProgressBar : public UWidget
 {
diff --git a/Engine/Source/Runtime/UMG/Public/Components/ScrollBox.h b/Engine/Source/Runtime/UMG/Public/Components/ScrollBox.h
--- a/Engine/Source/Runtime/UMG/Public/Components/ScrollBox.h
+++ b/Engine/Source/Runtime/UMG/Public/Components/ScrollBox.h
@@ -4,7 +4,9 @@
 
 #include "ScrollBox.generated.h"
 
-/** An arbitrary scrollable collection of widgets.  Great for presenting 10-100 widgets in a list.  Doesn't support virtualization. */
+/**
+ * An arbitrary scrollable collection of widgets.  Great for presenting 10-100 widgets in a list.  Doesn't support virtualization.
+ */
 UCLASS(ClassGroup=UserInterface)
 class UMG_API UScrollBox : public UPanelWidget
 {
diff --git a/Engine/Source/Runtime/UMG/Public/Components/SizeBox.h b/Engine/Source/Runtime/UMG/Public/Components/SizeBox.h
--- a/Engine/Source/Runtime/UMG/Public/Components/SizeBox.h
+++ b/Engine/Source/Runtime/UMG/Public/Components/SizeBox.h
@@ -7,8 +7,8 @@
 class USlateBrushAsset;
 
 /**
- * A SizeBox is a container widget that can contain one child widget, providing an opportunity 
- * to surround it with a SizeBox image and adjustable padding.
+ * A widget that allows you to specify the size it reports to have and desire.  Not all widgets report a desired size
+ * that you actually desire.  Wrapping them in a SizeBox lets you have the Size Box force them to be a particular size.
  */
 UCLASS(ClassGroup=UserInterface)
 class UMG_API USizeBox : public UContentWidget
diff --git a/Engine/Source/Runtime/UMG/Public/Components/Slider.h b/Engine/Source/Runtime/UMG/Public/Components/Slider.h
--- a/Engine/Source/Runtime/UMG/Public/Components/Slider.h
+++ b/Engine/Source/Runtime/UMG/Public/Components/Slider.h
@@ -8,7 +8,9 @@ DECLARE_DYNAMIC_MULTICAST_DELEGATE(FOnMouseCaptureBeginEvent);
 DECLARE_DYNAMIC_MULTICAST_DELEGATE(FOnMouseCaptureEndEvent);
 DECLARE_DYNAMIC_MULTICAST_DELEGATE_OneParam(FOnFloatValueChangedEvent, float, Value);
 
-/** Slider widget */
+/**
+ * A simple widget that shows a sliding bar with a handle that allows you to control the value between 0..1.
+ */
 UCLASS(ClassGroup=UserInterface)
 class UMG_API USlider : public UWidget
 {
diff --git a/Engine/Source/Runtime/UMG/Public/Components/Spacer.h b/Engine/Source/Runtime/UMG/Public/Components/Spacer.h
--- a/Engine/Source/Runtime/UMG/Public/Components/Spacer.h
+++ b/Engine/Source/Runtime/UMG/Public/Components/Spacer.h
@@ -4,7 +4,9 @@
 
 #include "Spacer.generated.h"
 
-/** A spacer widget; it does not have a visual representation, and just provides padding between other widgets. */
+/**
+ * A spacer widget; it does not have a visual representation, and just provides padding between other widgets.
+ */
 UCLASS(ClassGroup=UserInterface)
 class UMG_API USpacer : public UWidget
 {
diff --git a/Engine/Source/Runtime/UMG/Public/Components/SpinBox.h b/Engine/Source/Runtime/UMG/Public/Components/SpinBox.h
--- a/Engine/Source/Runtime/UMG/Public/Components/SpinBox.h
+++ b/Engine/Source/Runtime/UMG/Public/Components/SpinBox.h
@@ -4,7 +4,9 @@
 
 #include "SpinBox.generated.h"
 
-/** Spin box widget */
+/**
+ * A numerical entry box that allows for direct entry of the number or allows the user to click and slide the number.
+ */
 UCLASS(ClassGroup = UserInterface)
 class UMG_API USpinBox : public UWidget
 {
diff --git a/Engine/Source/Runtime/UMG/Public/Components/TextBlock.h b/Engine/Source/Runtime/UMG/Public/Components/TextBlock.h
--- a/Engine/Source/Runtime/UMG/Public/Components/TextBlock.h
+++ b/Engine/Source/Runtime/UMG/Public/Components/TextBlock.h
@@ -4,8 +4,10 @@
 
 #include "TextBlock.generated.h"
 
-/** A simple static text widget */
-UCLASS(ClassGroup=UserInterface )
+/**
+ * A simple static text widget
+ */
+UCLASS(ClassGroup=UserInterface)
 class UMG_API UTextBlock : public UWidget
 {
 	GENERATED_UCLASS_BODY()
diff --git a/Engine/Source/Runtime/UMG/Public/Components/Throbber.h b/Engine/Source/Runtime/UMG/Public/Components/Throbber.h
--- a/Engine/Source/Runtime/UMG/Public/Components/Throbber.h
+++ b/Engine/Source/Runtime/UMG/Public/Components/Throbber.h
@@ -6,7 +6,9 @@
 
 class USlateBrushAsset;
 
-/** A Throbber widget that shows several zooming circles in a row. */
+/**
+ * A Throbber widget that shows several zooming circles in a row.
+ */
 UCLASS(ClassGroup=UserInterface)
 class UMG_API UThrobber : public UWidget
 {
diff --git a/Engine/Source/Runtime/UMG/Public/Components/TileView.h b/Engine/Source/Runtime/UMG/Public/Components/TileView.h
--- a/Engine/Source/Runtime/UMG/Public/Components/TileView.h
+++ b/Engine/Source/Runtime/UMG/Public/Components/TileView.h
@@ -4,7 +4,9 @@
 
 #include "TileView.generated.h"
 
-/** A flow panel that presents the contents as a set of tiles all uniformly sized. */
+/**
+ * A flow panel that presents the contents as a set of tiles all uniformly sized.
+ */
 UCLASS(Experimental, ClassGroup=UserInterface)
 class UMG_API UTileView : public UTableViewBase
 {
diff --git a/Engine/Source/Runtime/UMG/Public/Components/UniformGridPanel.h b/Engine/Source/Runtime/UMG/Public/Components/UniformGridPanel.h
--- a/Engine/Source/Runtime/UMG/Public/Components/UniformGridPanel.h
+++ b/Engine/Source/Runtime/UMG/Public/Components/UniformGridPanel.h
@@ -4,7 +4,9 @@
 
 #include "UniformGridPanel.generated.h"
 
-/** A panel that evenly divides up available space between all of its children. */
+/**
+ * A panel that evenly divides up available space between all of its children.
+ */
 UCLASS(ClassGroup=UserInterface)
 class UMG_API UUniformGridPanel : public UPanelWidget
 {
diff --git a/Engine/Source/Runtime/UMG/Public/Components/WidgetSwitcher.h b/Engine/Source/Runtime/UMG/Public/Components/WidgetSwitcher.h
--- a/Engine/Source/Runtime/UMG/Public/Components/WidgetSwitcher.h
+++ b/Engine/Source/Runtime/UMG/Public/Components/WidgetSwitcher.h
@@ -5,8 +5,7 @@
 #include "WidgetSwitcher.generated.h"
 
 /**
- * A vertical box widget is a layout panel allowing child widgets to be automatically laid out
- * vertically.
+ * A widget switcher is like a tab control, but without tabs. At most one widget is visible at time.
  */
 UCLASS(ClassGroup=UserInterface)
 class UMG_API UWidgetSwitcher : public UPanelWidget
diff --git a/Engine/Source/Runtime/UMG/Public/Components/WrapBox.h b/Engine/Source/Runtime/UMG/Public/Components/WrapBox.h
--- a/Engine/Source/Runtime/UMG/Public/Components/WrapBox.h
+++ b/Engine/Source/Runtime/UMG/Public/Components/WrapBox.h
@@ -5,9 +5,8 @@
 #include "WrapBox.generated.h"
 
 /**
-* A vertical box widget is a layout panel allowing child widgets to be automatically laid out
-* vertically.
-*/
+ * Arranges widgets left-to-right.  When the widgets exceed the Width it will place widgets on the next line.
+ */
 UCLASS(ClassGroup=UserInterface)
 class UMG_API UWrapBox : public UPanelWidget
 {
