UMG - Hiding designer gizmos for designer hidden widgets.

diff --git a/Engine/Source/Editor/UMGEditor/Private/Designer/STransformHandle.cpp b/Engine/Source/Editor/UMGEditor/Private/Designer/STransformHandle.cpp
--- a/Engine/Source/Editor/UMGEditor/Private/Designer/STransformHandle.cpp
+++ b/Engine/Source/Editor/UMGEditor/Private/Designer/STransformHandle.cpp
@@ -37,22 +37,26 @@ EVisibility STransformHandle::GetHandleVisibility() const
 {
 	ETransformMode::Type TransformMode = Designer->GetTransformMode();
 
-	switch ( TransformMode )
+	// Only show the handles for visible elements in the designer.
+	FWidgetReference SelectedWidget = Designer->GetSelectedWidget();
+	if ( !SelectedWidget.GetTemplate()->bHiddenInDesigner )
 	{
-	case ETransformMode::Layout:
-	{
-		FWidgetReference SelectedWidget = Designer->GetSelectedWidget();
-		if ( UPanelSlot* TemplateSlot = SelectedWidget.GetTemplate()->Slot )
+		switch ( TransformMode )
 		{
-			if ( TemplateSlot->CanResize(DragDirection) )
+		case ETransformMode::Layout:
+		{
+			if ( UPanelSlot* TemplateSlot = SelectedWidget.GetTemplate()->Slot )
 			{
-				return EVisibility::Visible;
+				if ( TemplateSlot->CanResize(DragDirection) )
+				{
+					return EVisibility::Visible;
+				}
 			}
+			break;
+		}
+		case ETransformMode::Render:
+			return EVisibility::Visible;
 		}
-		break;
-	}
-	case ETransformMode::Render:
-		return EVisibility::Visible;
 	}
 
 	return EVisibility::Collapsed;
diff --git a/Engine/Source/Editor/UMGEditor/Private/Extensions/CanvasSlotExtension.cpp b/Engine/Source/Editor/UMGEditor/Private/Extensions/CanvasSlotExtension.cpp
--- a/Engine/Source/Editor/UMGEditor/Private/Extensions/CanvasSlotExtension.cpp
+++ b/Engine/Source/Editor/UMGEditor/Private/Extensions/CanvasSlotExtension.cpp
@@ -133,7 +133,7 @@ EVisibility FCanvasSlotExtension::GetAnchorVisibility(EAnchorWidget::Type Anchor
 	for ( const FWidgetReference& Selection : SelectionCache )
 	{
 		UWidget* PreviewWidget = Selection.GetPreview();
-		if ( PreviewWidget && PreviewWidget->Slot )
+		if ( PreviewWidget && PreviewWidget->Slot && !PreviewWidget->bHiddenInDesigner )
 		{
 			if ( UCanvasPanelSlot* PreviewCanvasSlot = Cast<UCanvasPanelSlot>(PreviewWidget->Slot) )
 			{
