UMG - You can now change the category of user widgets by selecting 'Root' which are the defaults for the class, then simply adjust the category at the top.

diff --git a/Engine/Source/Editor/UMGEditor/Private/Designer/SDesignSurface.cpp b/Engine/Source/Editor/UMGEditor/Private/Designer/SDesignSurface.cpp
--- a/Engine/Source/Editor/UMGEditor/Private/Designer/SDesignSurface.cpp
+++ b/Engine/Source/Editor/UMGEditor/Private/Designer/SDesignSurface.cpp
@@ -29,9 +29,7 @@ struct FFixedZoomLevelsContainer : public FZoomLevelsContainer
 		// Initialize zoom levels if not done already
 		if ( ZoomLevels.Num() == 0 )
 		{
-			ZoomLevels.Reserve(20);
-			ZoomLevels.Add(FZoomLevelEntry(0.100f, LOCTEXT("ZoomLevel", "-12"), EGraphRenderingLOD::LowestDetail));
-			ZoomLevels.Add(FZoomLevelEntry(0.125f, LOCTEXT("ZoomLevel", "-11"), EGraphRenderingLOD::LowestDetail));
+			ZoomLevels.Reserve(21);
 			ZoomLevels.Add(FZoomLevelEntry(0.150f, LOCTEXT("ZoomLevel", "-10"), EGraphRenderingLOD::LowestDetail));
 			ZoomLevels.Add(FZoomLevelEntry(0.175f, LOCTEXT("ZoomLevel", "-9"), EGraphRenderingLOD::LowestDetail));
 			ZoomLevels.Add(FZoomLevelEntry(0.200f, LOCTEXT("ZoomLevel", "-8"), EGraphRenderingLOD::LowestDetail));
@@ -44,15 +42,15 @@ struct FFixedZoomLevelsContainer : public FZoomLevelsContainer
 			ZoomLevels.Add(FZoomLevelEntry(0.875f, LOCTEXT("ZoomLevel", "-1"), EGraphRenderingLOD::DefaultDetail));
 			ZoomLevels.Add(FZoomLevelEntry(1.000f, LOCTEXT("ZoomLevel", "1:1"), EGraphRenderingLOD::DefaultDetail));
 			ZoomLevels.Add(FZoomLevelEntry(1.250f, LOCTEXT("ZoomLevel", "+1"), EGraphRenderingLOD::DefaultDetail));
-			ZoomLevels.Add(FZoomLevelEntry(1.375f, LOCTEXT("ZoomLevel", "+2"), EGraphRenderingLOD::DefaultDetail));
-			ZoomLevels.Add(FZoomLevelEntry(1.500f, LOCTEXT("ZoomLevel", "+3"), EGraphRenderingLOD::FullyZoomedIn));
-			ZoomLevels.Add(FZoomLevelEntry(1.675f, LOCTEXT("ZoomLevel", "+4"), EGraphRenderingLOD::FullyZoomedIn));
-			ZoomLevels.Add(FZoomLevelEntry(1.750f, LOCTEXT("ZoomLevel", "+5"), EGraphRenderingLOD::FullyZoomedIn));
-			ZoomLevels.Add(FZoomLevelEntry(1.875f, LOCTEXT("ZoomLevel", "+6"), EGraphRenderingLOD::FullyZoomedIn));
-			ZoomLevels.Add(FZoomLevelEntry(2.000f, LOCTEXT("ZoomLevel", "+7"), EGraphRenderingLOD::FullyZoomedIn));
-			ZoomLevels.Add(FZoomLevelEntry(2.250f, LOCTEXT("ZoomLevel", "+8"), EGraphRenderingLOD::FullyZoomedIn));
-			ZoomLevels.Add(FZoomLevelEntry(2.375f, LOCTEXT("ZoomLevel", "+9"), EGraphRenderingLOD::FullyZoomedIn));
-			ZoomLevels.Add(FZoomLevelEntry(2.500f, LOCTEXT("ZoomLevel", "+10"), EGraphRenderingLOD::FullyZoomedIn));
+			ZoomLevels.Add(FZoomLevelEntry(1.500f, LOCTEXT("ZoomLevel", "+2"), EGraphRenderingLOD::DefaultDetail));
+			ZoomLevels.Add(FZoomLevelEntry(1.750f, LOCTEXT("ZoomLevel", "+3"), EGraphRenderingLOD::FullyZoomedIn));
+			ZoomLevels.Add(FZoomLevelEntry(2.000f, LOCTEXT("ZoomLevel", "+4"), EGraphRenderingLOD::FullyZoomedIn));
+			ZoomLevels.Add(FZoomLevelEntry(2.250f, LOCTEXT("ZoomLevel", "+5"), EGraphRenderingLOD::FullyZoomedIn));
+			ZoomLevels.Add(FZoomLevelEntry(2.500f, LOCTEXT("ZoomLevel", "+6"), EGraphRenderingLOD::FullyZoomedIn));
+			ZoomLevels.Add(FZoomLevelEntry(2.750f, LOCTEXT("ZoomLevel", "+7"), EGraphRenderingLOD::FullyZoomedIn));
+			ZoomLevels.Add(FZoomLevelEntry(3.000f, LOCTEXT("ZoomLevel", "+8"), EGraphRenderingLOD::FullyZoomedIn));
+			ZoomLevels.Add(FZoomLevelEntry(3.250f, LOCTEXT("ZoomLevel", "+9"), EGraphRenderingLOD::FullyZoomedIn));
+			ZoomLevels.Add(FZoomLevelEntry(3.500f, LOCTEXT("ZoomLevel", "+10"), EGraphRenderingLOD::FullyZoomedIn));
 		}
 	}
 
diff --git a/Engine/Source/Editor/UMGEditor/Private/Details/SWidgetDetailsView.cpp b/Engine/Source/Editor/UMGEditor/Private/Details/SWidgetDetailsView.cpp
--- a/Engine/Source/Editor/UMGEditor/Private/Details/SWidgetDetailsView.cpp
+++ b/Engine/Source/Editor/UMGEditor/Private/Details/SWidgetDetailsView.cpp
@@ -43,6 +43,31 @@ void SWidgetDetailsView::Construct(const FArguments& InArgs, TSharedPtr<FWidgetB
 	[
 		SNew(SVerticalBox)
 
+		+ SVerticalBox::Slot()
+		.AutoHeight()
+		.Padding(0, 0, 0, 6)
+		[
+			SNew(SHorizontalBox)
+			.Visibility(this, &SWidgetDetailsView::GetCategoryAreaVisibility)
+
+			+ SHorizontalBox::Slot()
+			.AutoWidth()
+			.Padding(0, 0, 6, 0)
+			[
+				SNew(SBox)
+				.WidthOverride(200.0f)
+				.VAlign(VAlign_Center)
+				[
+					SAssignNew(NameTextBox, SEditableTextBox)
+					.SelectAllTextWhenFocused(true)
+					.ToolTipText(LOCTEXT("CategoryToolTip", "Sets the category of the widget"))
+					.HintText(LOCTEXT("Category", "Category"))
+					.Text(this, &SWidgetDetailsView::GetCategoryText)
+					.OnTextCommitted(this, &SWidgetDetailsView::HandleCategoryTextCommitted)
+				]
+			]
+		]
+
 		+ SVerticalBox::Slot()
 		.AutoHeight()
 		.Padding(0, 0, 0, 6)
@@ -218,6 +243,50 @@ EVisibility SWidgetDetailsView::GetNameAreaVisibility() const
 	return EVisibility::Collapsed;
 }
 
+EVisibility SWidgetDetailsView::GetCategoryAreaVisibility() const
+{
+	if ( SelectedObjects.Num() == 1 )
+	{
+		UUserWidget* Widget = Cast<UUserWidget>(SelectedObjects[0].Get());
+		if ( Widget && Widget->HasAnyFlags(RF_ClassDefaultObject) )
+		{
+			return EVisibility::Visible;
+		}
+	}
+
+	return EVisibility::Collapsed;
+}
+
+void SWidgetDetailsView::HandleCategoryTextCommitted(const FText& Text, ETextCommit::Type CommitType)
+{
+	if ( SelectedObjects.Num() == 1 && !Text.IsEmptyOrWhitespace() )
+	{
+		if ( UUserWidget* Widget = Cast<UUserWidget>(SelectedObjects[0].Get()) )
+		{
+			UUserWidget* WidgetCDO = Widget->GetClass()->GetDefaultObject<UUserWidget>();
+			WidgetCDO->PaletteCategory = Text;
+
+			// We mark it as structurally modified to ensure that all palettes refresh to move the widget into
+			// the correct new category.
+			FBlueprintEditorUtils::MarkBlueprintAsStructurallyModified(BlueprintEditor.Pin()->GetBlueprintObj());
+		}
+	}
+}
+
+FText SWidgetDetailsView::GetCategoryText() const
+{
+	if ( SelectedObjects.Num() == 1 )
+	{
+		if ( UUserWidget* Widget = Cast<UUserWidget>(SelectedObjects[0].Get()) )
+		{
+			UUserWidget* WidgetCDO = Widget->GetClass()->GetDefaultObject<UUserWidget>();
+			return WidgetCDO->PaletteCategory;
+		}
+	}
+
+	return FText::GetEmpty();
+}
+
 const FSlateBrush* SWidgetDetailsView::GetNameIcon() const
 {
 	if ( SelectedObjects.Num() == 1 )
@@ -229,7 +298,7 @@ const FSlateBrush* SWidgetDetailsView::GetNameIcon() const
 		}
 	}
 
-	return NULL;
+	return nullptr;
 }
 
 FText SWidgetDetailsView::GetNameText() const
diff --git a/Engine/Source/Editor/UMGEditor/Private/Details/SWidgetDetailsView.h b/Engine/Source/Editor/UMGEditor/Private/Details/SWidgetDetailsView.h
--- a/Engine/Source/Editor/UMGEditor/Private/Details/SWidgetDetailsView.h
+++ b/Engine/Source/Editor/UMGEditor/Private/Details/SWidgetDetailsView.h
@@ -54,6 +54,10 @@ private:
 
 	void HandleIsVariableChanged(ESlateCheckBoxState::Type CheckState);
 
+	EVisibility GetCategoryAreaVisibility() const;
+	FText GetCategoryText() const;
+	void HandleCategoryTextCommitted(const FText& Text, ETextCommit::Type CommitType);
+
 private:
 	/** The editor that owns this details view */
 	TWeakPtr<class FWidgetBlueprintEditor> BlueprintEditor;
diff --git a/Engine/Source/Editor/UMGEditor/Private/Templates/WidgetTemplateBlueprintClass.cpp b/Engine/Source/Editor/UMGEditor/Private/Templates/WidgetTemplateBlueprintClass.cpp
--- a/Engine/Source/Editor/UMGEditor/Private/Templates/WidgetTemplateBlueprintClass.cpp
+++ b/Engine/Source/Editor/UMGEditor/Private/Templates/WidgetTemplateBlueprintClass.cpp
@@ -27,8 +27,17 @@ FWidgetTemplateBlueprintClass::~FWidgetTemplateBlueprintClass()
 
 FText FWidgetTemplateBlueprintClass::GetCategory() const
 {
-	auto DefaultUserWidget = UUserWidget::StaticClass()->GetDefaultObject<UUserWidget>();
-	return DefaultUserWidget->GetPaletteCategory();
+	if ( WidgetClass.Get() )
+	{
+		UUserWidget* DefaultUserWidget = WidgetClass->GetDefaultObject<UUserWidget>();
+		return DefaultUserWidget->GetPaletteCategory();
+	}
+	else
+	{
+		//TODO UMG if the blueprint is unloaded we need to extract it from the asset metadata.
+		auto DefaultUserWidget = UUserWidget::StaticClass()->GetDefaultObject<UUserWidget>();
+		return DefaultUserWidget->GetPaletteCategory();
+	}
 }
 
 UWidget* FWidgetTemplateBlueprintClass::Create(UWidgetTree* Tree)
diff --git a/Engine/Source/Runtime/UMG/Private/UserWidget.cpp b/Engine/Source/Runtime/UMG/Private/UserWidget.cpp
--- a/Engine/Source/Runtime/UMG/Private/UserWidget.cpp
+++ b/Engine/Source/Runtime/UMG/Private/UserWidget.cpp
@@ -28,6 +28,7 @@ UUserWidget::UUserWidget(const FPostConstructInitializeProperties& PCIP)
 #if WITH_EDITORONLY_DATA
 	bUseDesignTimeSize = false;
 	DesignTimeSize = FVector2D(100, 100);
+	PaletteCategory = LOCTEXT("UserCreated", "User Created");
 #endif
 }
 
@@ -683,7 +684,7 @@ const FSlateBrush* UUserWidget::GetEditorIcon()
 
 const FText UUserWidget::GetPaletteCategory()
 {
-	return LOCTEXT("UserCreated", "User Created");
+	return PaletteCategory;
 }
 
 #endif
diff --git a/Engine/Source/Runtime/UMG/Public/Blueprint/UserWidget.h b/Engine/Source/Runtime/UMG/Public/Blueprint/UserWidget.h
--- a/Engine/Source/Runtime/UMG/Public/Blueprint/UserWidget.h
+++ b/Engine/Source/Runtime/UMG/Public/Blueprint/UserWidget.h
@@ -358,6 +358,10 @@ public:
 	UPROPERTY()
 	bool bUseDesignTimeSize;
 
+	/** The category this widget appears in the palette. */
+	UPROPERTY()
+	FText PaletteCategory;
+
 #endif
 
 protected:
