UMG - The Category override field for user widgets no longer shows "User Created" in the field if it's still the default, it leaves it blank so that the user see's that the it's the category field.

diff --git a/Engine/Source/Editor/UMGEditor/Private/Details/SWidgetDetailsView.cpp b/Engine/Source/Editor/UMGEditor/Private/Details/SWidgetDetailsView.cpp
--- a/Engine/Source/Editor/UMGEditor/Private/Details/SWidgetDetailsView.cpp
+++ b/Engine/Source/Editor/UMGEditor/Private/Details/SWidgetDetailsView.cpp
@@ -280,7 +280,15 @@ FText SWidgetDetailsView::GetCategoryText() const
 		if ( UUserWidget* Widget = Cast<UUserWidget>(SelectedObjects[0].Get()) )
 		{
 			UUserWidget* WidgetCDO = Widget->GetClass()->GetDefaultObject<UUserWidget>();
-			return WidgetCDO->PaletteCategory;
+			FText Category = WidgetCDO->PaletteCategory;
+			if ( Category.EqualToCaseIgnored(UUserWidget::StaticClass()->GetDefaultObject<UUserWidget>()->PaletteCategory) )
+			{
+				return FText::GetEmpty();
+			}
+			else
+			{
+				return Category;
+			}
 		}
 	}
 
