UMG - The designer surface now allows zooming below 1:1 without holding control.

diff --git a/Engine/Source/Editor/UMGEditor/Private/Designer/SDesignSurface.cpp b/Engine/Source/Editor/UMGEditor/Private/Designer/SDesignSurface.cpp
--- a/Engine/Source/Editor/UMGEditor/Private/Designer/SDesignSurface.cpp
+++ b/Engine/Source/Editor/UMGEditor/Private/Designer/SDesignSurface.cpp
@@ -86,7 +86,7 @@ struct FFixedZoomLevelsContainer : public FZoomLevelsContainer
 
 	int32 GetDefaultZoomLevel() const override
 	{
-		return 12;
+		return 10;
 	}
 
 	EGraphRenderingLOD::Type GetLOD(int32 InZoomLevel) const override
@@ -130,6 +130,8 @@ void SDesignSurface::Construct(const FArguments& InArgs)
 	bAllowContinousZoomInterpolation = false;
 	bTeleportInsteadOfScrollingWhenZoomingToFit = false;
 
+	bRequireControlToOverZoom = false;
+
 	ZoomTargetTopLeft = FVector2D::ZeroVector;
 	ZoomTargetBottomRight = FVector2D::ZeroVector;
 
@@ -236,7 +238,7 @@ FReply SDesignSurface::OnMouseWheel(const FGeometry& MyGeometry, const FPointerE
 	// We want to zoom into this point; i.e. keep it the same fraction offset into the panel
 	const FVector2D WidgetSpaceCursorPos = MyGeometry.AbsoluteToLocal(MouseEvent.GetScreenSpacePosition());
 	const int32 ZoomLevelDelta = FMath::FloorToInt(MouseEvent.GetWheelDelta());
-	ChangeZoomLevel(ZoomLevelDelta, WidgetSpaceCursorPos, MouseEvent.IsControlDown());
+	ChangeZoomLevel(ZoomLevelDelta, WidgetSpaceCursorPos, !bRequireControlToOverZoom || MouseEvent.IsControlDown());
 
 	return FReply::Handled();
 }
diff --git a/Engine/Source/Editor/UMGEditor/Private/Designer/SDesignSurface.h b/Engine/Source/Editor/UMGEditor/Private/Designer/SDesignSurface.h
--- a/Engine/Source/Editor/UMGEditor/Private/Designer/SDesignSurface.h
+++ b/Engine/Source/Editor/UMGEditor/Private/Designer/SDesignSurface.h
@@ -100,4 +100,7 @@ protected:
 
 	FVector2D ZoomTargetTopLeft;
 	FVector2D ZoomTargetBottomRight;
+
+	/** Does the user need to press Control in order to over-zoom. */
+	bool bRequireControlToOverZoom;
 };
