UMG - Shift is no longer used in the designer to 'move the widget out of the parent'.  That's now Alt.  Shift is now used for Axis locked movement like photoshop.

diff --git a/Engine/Source/Editor/UMGEditor/Private/Designer/SDesignerView.cpp b/Engine/Source/Editor/UMGEditor/Private/Designer/SDesignerView.cpp
--- a/Engine/Source/Editor/UMGEditor/Private/Designer/SDesignerView.cpp
+++ b/Engine/Source/Editor/UMGEditor/Private/Designer/SDesignerView.cpp
@@ -129,7 +129,6 @@ void SDesignerView::Construct(const FArguments& InArgs, TSharedPtr<FWidgetBluepr
 
 	HoverTime = 0;
 
-	bMouseDown = false;
 	bMovingExistingWidget = false;
 
 	// TODO UMG - Register these with the module through some public interface to allow for new extensions to be registered.
@@ -510,7 +509,7 @@ FText SDesignerView::GetInfoBarText() const
 	switch ( DesignerMessage )
 	{
 	case EDesignerMessage::MoveFromParent:
-		return LOCTEXT("PressShiftToMove", "Press SHIFT to move the widget out of the current parent");
+		return LOCTEXT("PressShiftToMove", "Press Alt to move the widget out of the current parent");
 	}
 
 	return FText::GetEmpty();
@@ -852,8 +851,7 @@ FReply SDesignerView::OnMouseButtonDown(const FGeometry& MyGeometry, const FPoin
 				ResolvePendingSelectedWidgets();
 			}
 
-			bMouseDown = true;
-			ScreenMouseDownLocation = MouseEvent.GetScreenSpacePosition();
+			DraggingStartPositionScreenSpace = MouseEvent.GetScreenSpacePosition();
 		}
 	}
 
@@ -867,7 +865,6 @@ FReply SDesignerView::OnMouseButtonUp(const FGeometry& MyGeometry, const FPointe
 	{
 		ResolvePendingSelectedWidgets();
 
-		bMouseDown = false;
 		bMovingExistingWidget = false;
 		DesignerMessage = EDesignerMessage::None;
 
@@ -901,7 +898,7 @@ FReply SDesignerView::OnMouseMove(const FGeometry& MyGeometry, const FPointerEve
 		return SurfaceHandled;
 	}
 
-	if ( bMouseDown )
+	if ( MouseEvent.IsMouseButtonDown(EKeys::LeftMouseButton) )
 	{
 		if ( SelectedWidget.IsValid() && !bMovingExistingWidget )
 		{
@@ -1201,10 +1198,10 @@ FReply SDesignerView::OnDragDetected(const FGeometry& MyGeometry, const FPointer
 		// Clear any pending selected widgets, the user has already decided what widget they want.
 		PendingSelectedWidget = FWidgetReference();
 
-		//
+		// Determine The offset to keep the widget from the mouse while dragging
 		FArrangedWidget ArrangedWidget(SNullWidget::NullWidget, FGeometry());
 		FDesignTimeUtils::GetArrangedWidget(SelectedWidget.GetPreview()->GetCachedWidget().ToSharedRef(), ArrangedWidget);
-		SelectedWidgetContextMenuLocation = ArrangedWidget.Geometry.AbsoluteToLocal(ScreenMouseDownLocation);
+		SelectedWidgetContextMenuLocation = ArrangedWidget.Geometry.AbsoluteToLocal(DraggingStartPositionScreenSpace);
 
 		ClearExtensionWidgets();
 
@@ -1412,9 +1409,9 @@ UWidget* SDesignerView::ProcessDropAndAddWidget(const FGeometry& MyGeometry, con
 	{
 		SelectedDragDropOp->SetCursorOverride(TOptional<EMouseCursor::Type>());
 
-		// If they've pressed shift, and we were staying in the parent, disable that
+		// If they've pressed alt, and we were staying in the parent, disable that
 		// and adjust the designer message to no longer warn.
-		if ( DragDropEvent.IsShiftDown() && SelectedDragDropOp->bStayingInParent )
+		if ( DragDropEvent.IsAltDown() && SelectedDragDropOp->bStayingInParent )
 		{
 			SelectedDragDropOp->bStayingInParent = false;
 			DesignerMessage = EDesignerMessage::None;
@@ -1465,15 +1462,49 @@ UWidget* SDesignerView::ProcessDropAndAddWidget(const FGeometry& MyGeometry, con
 					Widget->GetParent()->RemoveChild(Widget);
 				}
 
-				FVector2D LocalPosition = WidgetUnderCursorGeometry.AbsoluteToLocal(DragDropEvent.GetScreenSpacePosition());
+				FVector2D ScreenSpacePosition = DragDropEvent.GetScreenSpacePosition();
+
+				const UWidgetDesignerSettings* DesignerSettings = GetDefault<UWidgetDesignerSettings>();
+				bool bGridSnapX, bGridSnapY;
+				bGridSnapX = bGridSnapY = DesignerSettings->GridSnapEnabled;
+
+				// As long as shift is pressed and we're staying in the same parent,
+				// allow the user to lock the movement to a specific axis.
+				const bool bLockToAxis =
+					FSlateApplication::Get().GetModifierKeys().IsShiftDown() &&
+					SelectedDragDropOp->bStayingInParent;
+
+				if ( bLockToAxis )
+				{
+					// Choose the largest axis of movement as the primary axis to lock to.
+					FVector2D DragDelta = ScreenSpacePosition - DraggingStartPositionScreenSpace;
+					if ( FMath::Abs(DragDelta.X) > FMath::Abs(DragDelta.Y) )
+					{
+						// Lock to X Axis
+						ScreenSpacePosition.Y = DraggingStartPositionScreenSpace.Y;
+						bGridSnapY = false;
+					}
+					else
+					{
+						// Lock To Y Axis
+						ScreenSpacePosition.X = DraggingStartPositionScreenSpace.X;
+						bGridSnapX = false;
+					}
+				}
+
+				FVector2D LocalPosition = WidgetUnderCursorGeometry.AbsoluteToLocal(ScreenSpacePosition);
 				if ( UPanelSlot* Slot = NewParent->AddChild(Widget) )
 				{
 					FVector2D NewPosition = LocalPosition - SelectedWidgetContextMenuLocation;
 
-					const UWidgetDesignerSettings* DesignerSettings = GetDefault<UWidgetDesignerSettings>();
-					if ( DesignerSettings->GridSnapEnabled )
+					// Perform grid snapping on X and Y if we need to.
+					if ( bGridSnapX )
 					{
 						NewPosition.X = ( (int32)NewPosition.X ) - ( ( (int32)NewPosition.X ) % DesignerSettings->GridSnapSize );
+					}
+
+					if ( bGridSnapY )
+					{
 						NewPosition.Y = ( (int32)NewPosition.Y ) - ( ( (int32)NewPosition.Y ) % DesignerSettings->GridSnapSize );
 					}
 
@@ -1539,7 +1570,6 @@ UWidget* SDesignerView::ProcessDropAndAddWidget(const FGeometry& MyGeometry, con
 
 FReply SDesignerView::OnDrop(const FGeometry& MyGeometry, const FDragDropEvent& DragDropEvent)
 {
-	bMouseDown = false;
 	bMovingExistingWidget = false;
 
 	UWidgetBlueprint* BP = GetBlueprint();
diff --git a/Engine/Source/Editor/UMGEditor/Private/Designer/SDesignerView.h b/Engine/Source/Editor/UMGEditor/Private/Designer/SDesignerView.h
--- a/Engine/Source/Editor/UMGEditor/Private/Designer/SDesignerView.h
+++ b/Engine/Source/Editor/UMGEditor/Private/Designer/SDesignerView.h
@@ -213,11 +213,8 @@ private:
 	 */
 	FWidgetReference PendingSelectedWidget;
 
-	/**  */
-	bool bMouseDown;
-
-	/**  */
-	FVector2D ScreenMouseDownLocation;
+	/** The position in screen space where the user began dragging a widget */
+	FVector2D DraggingStartPositionScreenSpace;
 
 	/** An existing widget is being moved in its current container, or in to a new container. */
 	bool bMovingExistingWidget;
