UMGEditor private PCH cleanup.

diff --git a/Engine/Source/Editor/Kismet/Public/BlueprintEditorModes.h b/Engine/Source/Editor/Kismet/Public/BlueprintEditorModes.h
--- a/Engine/Source/Editor/Kismet/Public/BlueprintEditorModes.h
+++ b/Engine/Source/Editor/Kismet/Public/BlueprintEditorModes.h
@@ -4,7 +4,7 @@
 
 #include "BlueprintEditorModule.h"
 #include "BlueprintUtilities.h"
-
+#include "WorkflowOrientedApp/ApplicationMode.h"
 
 struct KISMET_API FBlueprintEditorApplicationModes
 {
diff --git a/Engine/Source/Editor/Sequencer/Public/MovieSceneTrackEditor.h b/Engine/Source/Editor/Sequencer/Public/MovieSceneTrackEditor.h
--- a/Engine/Source/Editor/Sequencer/Public/MovieSceneTrackEditor.h
+++ b/Engine/Source/Editor/Sequencer/Public/MovieSceneTrackEditor.h
@@ -5,6 +5,7 @@
 //@todo Sequencer - these have to be here for now because this class contains a small amount of implementation inside this header to avoid exporting this class
 #include "ISequencer.h"
 #include "MovieScene.h"
+#include "ScopedTransaction.h"
 
 class ISequencerSection;
 
diff --git a/Engine/Source/Editor/UMGEditor/Classes/WidgetBlueprintFactory.h b/Engine/Source/Editor/UMGEditor/Classes/WidgetBlueprintFactory.h
--- a/Engine/Source/Editor/UMGEditor/Classes/WidgetBlueprintFactory.h
+++ b/Engine/Source/Editor/UMGEditor/Classes/WidgetBlueprintFactory.h
@@ -1,7 +1,7 @@
 // Copyright 1998-2014 Epic Games, Inc. All Rights Reserved.
 
 #pragma once
-
+#include "Factories/Factory.h"
 #include "WidgetBlueprintFactory.generated.h"
 
 UCLASS(HideCategories=Object, MinimalAPI)
diff --git a/Engine/Source/Editor/UMGEditor/Private/AssetTypeActions_WidgetBlueprint.h b/Engine/Source/Editor/UMGEditor/Private/AssetTypeActions_WidgetBlueprint.h
--- a/Engine/Source/Editor/UMGEditor/Private/AssetTypeActions_WidgetBlueprint.h
+++ b/Engine/Source/Editor/UMGEditor/Private/AssetTypeActions_WidgetBlueprint.h
@@ -1,7 +1,7 @@
 // Copyright 1998-2014 Epic Games, Inc. All Rights Reserved.
 
 #pragma once
-
+#include "AssetTypeActions_Base.h"
 //TODO Use base class FAssetTypeActions_Blueprint later
 class FAssetTypeActions_WidgetBlueprint : public FAssetTypeActions_Base
 {
diff --git a/Engine/Source/Editor/UMGEditor/Private/BlueprintModes/WidgetBlueprintApplicationMode.h b/Engine/Source/Editor/UMGEditor/Private/BlueprintModes/WidgetBlueprintApplicationMode.h
--- a/Engine/Source/Editor/UMGEditor/Private/BlueprintModes/WidgetBlueprintApplicationMode.h
+++ b/Engine/Source/Editor/UMGEditor/Private/BlueprintModes/WidgetBlueprintApplicationMode.h
@@ -1,6 +1,7 @@
 // Copyright 1998-2014 Epic Games, Inc. All Rights Reserved.
 
 #pragma once
+#include "BlueprintEditorModes.h"
 
 /////////////////////////////////////////////////////
 // FWidgetBlueprintApplicationMode
diff --git a/Engine/Source/Editor/UMGEditor/Private/BlueprintModes/WidgetDesignerApplicationMode.h b/Engine/Source/Editor/UMGEditor/Private/BlueprintModes/WidgetDesignerApplicationMode.h
--- a/Engine/Source/Editor/UMGEditor/Private/BlueprintModes/WidgetDesignerApplicationMode.h
+++ b/Engine/Source/Editor/UMGEditor/Private/BlueprintModes/WidgetDesignerApplicationMode.h
@@ -1,6 +1,7 @@
 // Copyright 1998-2014 Epic Games, Inc. All Rights Reserved.
 
 #pragma once
+#include "BlueprintModes/WidgetBlueprintApplicationMode.h"
 
 /////////////////////////////////////////////////////
 // FWidgetDesignerApplicationMode
diff --git a/Engine/Source/Editor/UMGEditor/Private/BlueprintModes/WidgetGraphApplicationMode.h b/Engine/Source/Editor/UMGEditor/Private/BlueprintModes/WidgetGraphApplicationMode.h
--- a/Engine/Source/Editor/UMGEditor/Private/BlueprintModes/WidgetGraphApplicationMode.h
+++ b/Engine/Source/Editor/UMGEditor/Private/BlueprintModes/WidgetGraphApplicationMode.h
@@ -1,6 +1,7 @@
 // Copyright 1998-2014 Epic Games, Inc. All Rights Reserved.
 
 #pragma once
+#include "BlueprintModes/WidgetBlueprintApplicationMode.h"
 
 /////////////////////////////////////////////////////
 // FWidgetGraphApplicationMode
diff --git a/Engine/Source/Editor/UMGEditor/Private/Customizations/CanvasSlotCustomization.cpp b/Engine/Source/Editor/UMGEditor/Private/Customizations/CanvasSlotCustomization.cpp
--- a/Engine/Source/Editor/UMGEditor/Private/Customizations/CanvasSlotCustomization.cpp
+++ b/Engine/Source/Editor/UMGEditor/Private/Customizations/CanvasSlotCustomization.cpp
@@ -1,7 +1,7 @@
 // Copyright 1998-2014 Epic Games, Inc. All Rights Reserved.
 
 #include "UMGEditorPrivatePCH.h"
-
+#include "Components/CanvasPanelSlot.h"
 #include "PropertyEditing.h"
 #include "ObjectEditorUtils.h"
 #include "WidgetGraphSchema.h"
diff --git a/Engine/Source/Editor/UMGEditor/Private/Customizations/DetailCustomizations.cpp b/Engine/Source/Editor/UMGEditor/Private/Customizations/DetailCustomizations.cpp
--- a/Engine/Source/Editor/UMGEditor/Private/Customizations/DetailCustomizations.cpp
+++ b/Engine/Source/Editor/UMGEditor/Private/Customizations/DetailCustomizations.cpp
@@ -9,6 +9,7 @@
 #include "WidgetGraphSchema.h"
 #include "ScopedTransaction.h"
 #include "BlueprintEditorUtils.h"
+#include "WidgetGraphSchema.h"
 
 #define LOCTEXT_NAMESPACE "UMG"
 
diff --git a/Engine/Source/Editor/UMGEditor/Private/Designer/DesignTimeUtils.cpp b/Engine/Source/Editor/UMGEditor/Private/Designer/DesignTimeUtils.cpp
--- a/Engine/Source/Editor/UMGEditor/Private/Designer/DesignTimeUtils.cpp
+++ b/Engine/Source/Editor/UMGEditor/Private/Designer/DesignTimeUtils.cpp
@@ -1,7 +1,7 @@
 // Copyright 1998-2014 Epic Games, Inc. All Rights Reserved.
 
 #include "UMGEditorPrivatePCH.h"
-
+#include "Designer/DesignTimeUtils.h"
 #define LOCTEXT_NAMESPACE "UMGEditor"
 
 bool FDesignTimeUtils::GetArrangedWidget(TSharedRef<SWidget> Widget, FArrangedWidget& ArrangedWidget)
diff --git a/Engine/Source/Editor/UMGEditor/Private/Designer/DesignTimeUtils.h b/Engine/Source/Editor/UMGEditor/Private/Designer/DesignTimeUtils.h
--- a/Engine/Source/Editor/UMGEditor/Private/Designer/DesignTimeUtils.h
+++ b/Engine/Source/Editor/UMGEditor/Private/Designer/DesignTimeUtils.h
@@ -1,6 +1,9 @@
 // Copyright 1998-2014 Epic Games, Inc. All Rights Reserved.
 
 #pragma once
+class SWidget;
+class FArrangedWidget;
+class FWidgetPath;
 
 /**
  * A set of utility functions used at design time for the widget blueprint editor.
diff --git a/Engine/Source/Editor/UMGEditor/Private/Designer/SDesignerToolBar.cpp b/Engine/Source/Editor/UMGEditor/Private/Designer/SDesignerToolBar.cpp
--- a/Engine/Source/Editor/UMGEditor/Private/Designer/SDesignerToolBar.cpp
+++ b/Engine/Source/Editor/UMGEditor/Private/Designer/SDesignerToolBar.cpp
@@ -1,10 +1,9 @@
 // Copyright 1998-2014 Epic Games, Inc. All Rights Reserved.
 
 #include "UMGEditorPrivatePCH.h"
-
-#include "SDesignerToolBar.h"
 #include "DesignerCommands.h"
-
+#include "SDesignerToolBar.h"
+#include "Settings/WidgetDesignerSettings.h"
 #include "SViewportToolBarComboMenu.h"
 
 #define LOCTEXT_NAMESPACE "UMG"
diff --git a/Engine/Source/Editor/UMGEditor/Private/Designer/SDesignerView.cpp b/Engine/Source/Editor/UMGEditor/Private/Designer/SDesignerView.cpp
--- a/Engine/Source/Editor/UMGEditor/Private/Designer/SDesignerView.cpp
+++ b/Engine/Source/Editor/UMGEditor/Private/Designer/SDesignerView.cpp
@@ -1,6 +1,11 @@
 // Copyright 1998-2014 Epic Games, Inc. All Rights Reserved.
 
 #include "UMGEditorPrivatePCH.h"
+#include "Extensions/CanvasSlotExtension.h"
+#include "Extensions/GridSlotExtension.h"
+#include "Extensions/HorizontalSlotExtension.h"
+#include "Extensions/UniformGridSlotExtension.h"
+#include "Extensions/VerticalSlotExtension.h"
 #include "SDesignerView.h"
 
 #include "BlueprintEditor.h"
diff --git a/Engine/Source/Editor/UMGEditor/Private/Extensions/CanvasSlotExtension.cpp b/Engine/Source/Editor/UMGEditor/Private/Extensions/CanvasSlotExtension.cpp
--- a/Engine/Source/Editor/UMGEditor/Private/Extensions/CanvasSlotExtension.cpp
+++ b/Engine/Source/Editor/UMGEditor/Private/Extensions/CanvasSlotExtension.cpp
@@ -1,6 +1,7 @@
 // Copyright 1998-2014 Epic Games, Inc. All Rights Reserved.
 
 #include "UMGEditorPrivatePCH.h"
+#include "Extensions/CanvasSlotExtension.h"
 
 #define LOCTEXT_NAMESPACE "UMG"
 
diff --git a/Engine/Source/Editor/UMGEditor/Private/Extensions/GridSlotExtension.cpp b/Engine/Source/Editor/UMGEditor/Private/Extensions/GridSlotExtension.cpp
--- a/Engine/Source/Editor/UMGEditor/Private/Extensions/GridSlotExtension.cpp
+++ b/Engine/Source/Editor/UMGEditor/Private/Extensions/GridSlotExtension.cpp
@@ -1,6 +1,7 @@
 ﻿// Copyright 1998-2014 Epic Games, Inc. All Rights Reserved.
 
 #include "UMGEditorPrivatePCH.h"
+#include "Components/GridSlot.h"
 
 #define LOCTEXT_NAMESPACE "UMG"
 
diff --git a/Engine/Source/Editor/UMGEditor/Private/Extensions/HorizontalSlotExtension.cpp b/Engine/Source/Editor/UMGEditor/Private/Extensions/HorizontalSlotExtension.cpp
--- a/Engine/Source/Editor/UMGEditor/Private/Extensions/HorizontalSlotExtension.cpp
+++ b/Engine/Source/Editor/UMGEditor/Private/Extensions/HorizontalSlotExtension.cpp
@@ -1,6 +1,8 @@
 ﻿// Copyright 1998-2014 Epic Games, Inc. All Rights Reserved.
 
 #include "UMGEditorPrivatePCH.h"
+#include "Components/HorizontalBox.h"
+#include "Components/HorizontalBoxSlot.h"
 #include "Kismet2/BlueprintEditorUtils.h"
 
 #define LOCTEXT_NAMESPACE "UMG"
diff --git a/Engine/Source/Editor/UMGEditor/Private/Extensions/UniformGridSlotExtension.cpp b/Engine/Source/Editor/UMGEditor/Private/Extensions/UniformGridSlotExtension.cpp
--- a/Engine/Source/Editor/UMGEditor/Private/Extensions/UniformGridSlotExtension.cpp
+++ b/Engine/Source/Editor/UMGEditor/Private/Extensions/UniformGridSlotExtension.cpp
@@ -1,7 +1,7 @@
 ﻿// Copyright 1998-2014 Epic Games, Inc. All Rights Reserved.
 
 #include "UMGEditorPrivatePCH.h"
-
+#include "Components/UniformGridSlot.h"
 #define LOCTEXT_NAMESPACE "UMG"
 
 FUniformGridSlotExtension::FUniformGridSlotExtension()
diff --git a/Engine/Source/Editor/UMGEditor/Private/Extensions/VerticalSlotExtension.cpp b/Engine/Source/Editor/UMGEditor/Private/Extensions/VerticalSlotExtension.cpp
--- a/Engine/Source/Editor/UMGEditor/Private/Extensions/VerticalSlotExtension.cpp
+++ b/Engine/Source/Editor/UMGEditor/Private/Extensions/VerticalSlotExtension.cpp
@@ -1,6 +1,8 @@
 ﻿// Copyright 1998-2014 Epic Games, Inc. All Rights Reserved.
 
 #include "UMGEditorPrivatePCH.h"
+#include "Components/VerticalBox.h"
+#include "Components/VerticalBoxSlot.h"
 #include "Kismet2/BlueprintEditorUtils.h"
 
 #define LOCTEXT_NAMESPACE "UMG"
diff --git a/Engine/Source/Editor/UMGEditor/Private/Hierarchy/SHierarchyView.h b/Engine/Source/Editor/UMGEditor/Private/Hierarchy/SHierarchyView.h
--- a/Engine/Source/Editor/UMGEditor/Private/Hierarchy/SHierarchyView.h
+++ b/Engine/Source/Editor/UMGEditor/Private/Hierarchy/SHierarchyView.h
@@ -2,11 +2,11 @@
 
 #pragma once
 
-#include "SCompoundWidget.h"
 #include "BlueprintEditor.h"
-#include "TreeFilterHandler.h"
+#include "Misc/TextFilter.h"
+#include "SCompoundWidget.h"
 #include "SHierarchyViewItem.h"
-
+#include "TreeFilterHandler.h"
 //TODO rename SUMGEditorHierarchy
 
 /**
diff --git a/Engine/Source/Editor/UMGEditor/Private/Nodes/K2Node_CreateDragDropOperation.cpp b/Engine/Source/Editor/UMGEditor/Private/Nodes/K2Node_CreateDragDropOperation.cpp
--- a/Engine/Source/Editor/UMGEditor/Private/Nodes/K2Node_CreateDragDropOperation.cpp
+++ b/Engine/Source/Editor/UMGEditor/Private/Nodes/K2Node_CreateDragDropOperation.cpp
@@ -1,12 +1,12 @@
 // Copyright 1998-2014 Epic Games, Inc. All Rights Reserved.
 
 #include "UMGEditorPrivatePCH.h"
-
-#include "KismetCompiler.h"
+#include "Blueprint/DragDropOperation.h"
+#include "Blueprint/WidgetBlueprintLibrary.h"
 #include "BlueprintNodeSpawner.h"
 #include "EditorCategoryUtils.h"
+#include "KismetCompiler.h"
 #include "K2ActionMenuBuilder.h"
-
 #include "K2Node_CreateDragDropOperation.h"
 
 #define LOCTEXT_NAMESPACE "UMG"
diff --git a/Engine/Source/Editor/UMGEditor/Private/Nodes/K2Node_CreateDragDropOperation.h b/Engine/Source/Editor/UMGEditor/Private/Nodes/K2Node_CreateDragDropOperation.h
--- a/Engine/Source/Editor/UMGEditor/Private/Nodes/K2Node_CreateDragDropOperation.h
+++ b/Engine/Source/Editor/UMGEditor/Private/Nodes/K2Node_CreateDragDropOperation.h
@@ -3,6 +3,7 @@
 #pragma once
 
 #include "K2Node.h"
+#include "K2Node_ConstructObjectFromClass.h"
 #include "K2Node_CreateDragDropOperation.generated.h"
 
 UCLASS()
diff --git a/Engine/Source/Editor/UMGEditor/Private/Nodes/K2Node_CreateWidget.cpp b/Engine/Source/Editor/UMGEditor/Private/Nodes/K2Node_CreateWidget.cpp
--- a/Engine/Source/Editor/UMGEditor/Private/Nodes/K2Node_CreateWidget.cpp
+++ b/Engine/Source/Editor/UMGEditor/Private/Nodes/K2Node_CreateWidget.cpp
@@ -1,7 +1,7 @@
 // Copyright 1998-2014 Epic Games, Inc. All Rights Reserved.
 
 #include "UMGEditorPrivatePCH.h"
-
+#include "Blueprint/WidgetBlueprintLibrary.h"
 #include "KismetCompiler.h"
 #include "BlueprintNodeSpawner.h"
 #include "EditorCategoryUtils.h"
diff --git a/Engine/Source/Editor/UMGEditor/Private/Nodes/K2Node_CreateWidget.h b/Engine/Source/Editor/UMGEditor/Private/Nodes/K2Node_CreateWidget.h
--- a/Engine/Source/Editor/UMGEditor/Private/Nodes/K2Node_CreateWidget.h
+++ b/Engine/Source/Editor/UMGEditor/Private/Nodes/K2Node_CreateWidget.h
@@ -3,6 +3,7 @@
 #pragma once
 
 #include "K2Node.h"
+#include "K2Node_ConstructObjectFromClass.h"
 #include "K2Node_CreateWidget.generated.h"
 
 UCLASS()
diff --git a/Engine/Source/Editor/UMGEditor/Private/Palette/SPaletteView.h b/Engine/Source/Editor/UMGEditor/Private/Palette/SPaletteView.h
--- a/Engine/Source/Editor/UMGEditor/Private/Palette/SPaletteView.h
+++ b/Engine/Source/Editor/UMGEditor/Private/Palette/SPaletteView.h
@@ -1,7 +1,7 @@
 // Copyright 1998-2014 Epic Games, Inc. All Rights Reserved.
 
 #pragma once
-
+#include "Misc/TextFilter.h"
 #include "SCompoundWidget.h"
 #include "BlueprintEditor.h"
 #include "TreeFilterHandler.h"
diff --git a/Engine/Source/Editor/UMGEditor/Private/TabFactory/AnimationTabSummoner.cpp b/Engine/Source/Editor/UMGEditor/Private/TabFactory/AnimationTabSummoner.cpp
--- a/Engine/Source/Editor/UMGEditor/Private/TabFactory/AnimationTabSummoner.cpp
+++ b/Engine/Source/Editor/UMGEditor/Private/TabFactory/AnimationTabSummoner.cpp
@@ -7,6 +7,7 @@
 #include "ISequencer.h"
 #include "WidgetAnimation.h"
 #include "MovieScene.h"
+#include "UMGStyle.h"
 
 #define LOCTEXT_NAMESPACE "UMG"
 
diff --git a/Engine/Source/Editor/UMGEditor/Private/TabFactory/DesignerTabSummoner.cpp b/Engine/Source/Editor/UMGEditor/Private/TabFactory/DesignerTabSummoner.cpp
--- a/Engine/Source/Editor/UMGEditor/Private/TabFactory/DesignerTabSummoner.cpp
+++ b/Engine/Source/Editor/UMGEditor/Private/TabFactory/DesignerTabSummoner.cpp
@@ -5,6 +5,7 @@
 #include "DesignerTabSummoner.h"
 #include "SDesignerView.h"
 #include "SScissorRectBox.h"
+#include "UMGStyle.h"
 
 #define LOCTEXT_NAMESPACE "UMG"
 
diff --git a/Engine/Source/Editor/UMGEditor/Private/TabFactory/PaletteTabSummoner.cpp b/Engine/Source/Editor/UMGEditor/Private/TabFactory/PaletteTabSummoner.cpp
--- a/Engine/Source/Editor/UMGEditor/Private/TabFactory/PaletteTabSummoner.cpp
+++ b/Engine/Source/Editor/UMGEditor/Private/TabFactory/PaletteTabSummoner.cpp
@@ -4,6 +4,7 @@
 
 #include "PaletteTabSummoner.h"
 #include "SPaletteView.h"
+#include "UMGStyle.h"
 
 #define LOCTEXT_NAMESPACE "UMG"
 
diff --git a/Engine/Source/Editor/UMGEditor/Private/UMGEditorModule.cpp b/Engine/Source/Editor/UMGEditor/Private/UMGEditorModule.cpp
--- a/Engine/Source/Editor/UMGEditor/Private/UMGEditorModule.cpp
+++ b/Engine/Source/Editor/UMGEditor/Private/UMGEditorModule.cpp
@@ -13,7 +13,7 @@
 #include "Editor/Sequencer/Public/ISequencerModule.h"
 #include "Animation/MarginTrackEditor.h"
 #include "Animation/Sequencer2DTransformTrackEditor.h"
-
+#include "IUMGModule.h"
 #include "DesignerCommands.h"
 
 const FName UMGEditorAppIdentifier = FName(TEXT("UMGEditorApp"));
diff --git a/Engine/Source/Editor/UMGEditor/Private/UMGEditorPrivatePCH.h b/Engine/Source/Editor/UMGEditor/Private/UMGEditorPrivatePCH.h
--- a/Engine/Source/Editor/UMGEditor/Private/UMGEditorPrivatePCH.h
+++ b/Engine/Source/Editor/UMGEditor/Private/UMGEditorPrivatePCH.h
@@ -2,54 +2,10 @@
 
 #pragma once
 
-#include "Slate.h"
-#include "SlateCore.h"
-#include "UnrealEd.h"
-
-#include "IDocumentation.h"
-
-#include "BlueprintGraphDefinitions.h"
-
-#include "BlueprintEditorUtils.h"
-
-#include "ScopedTransaction.h"
-#include "TextFilter.h"
-
-#include "ObjectEditorUtils.h"
-
-#include "UMG.h"
-#include "UMGStyle.h"
-
-#include "WidgetDesignerSettings.h"
-
-#include "UMGEditorClasses.h"
-#include "WidgetBlueprint.h"
-
-#include "BlueprintEditorModes.h"
-#include "WidgetBlueprintApplicationModes.h"
-#include "WidgetBlueprintApplicationMode.h"
-
-#include "WidgetReference.h"
-#include "WidgetBlueprintEditor.h"
-
-#include "AssetTypeActions_Base.h"
-
-#include "DesignTimeUtils.h"
-
-#include "DesignerExtension.h"
-#include "VerticalSlotExtension.h"
-#include "HorizontalSlotExtension.h"
-#include "CanvasSlotExtension.h"
-#include "UniformGridSlotExtension.h"
-#include "GridSlotExtension.h"
-
-#include "DetailCustomizations.h"
-
-#include "WidgetTemplate.h"
-#include "WidgetTemplateDragDropOp.h"
-
-#include "WidgetBlueprintEditorUtils.h"
-
-#include "WidgetBlueprintLibrary.h"
-
-#include "WidgetBlueprintCompiler.h"
+#include "EditorStyleSet.h"
+#include "Internationalization/Internationalization.h"
+#include "Internationalization/Text.h"
+#include "Misc/AssertionMacros.h"
+#include "Templates/SharedPointer.h"
+#include "Templates/Casts.h"
+#include "Widgets/DeclarativeSyntaxSupport.h"
diff --git a/Engine/Source/Editor/UMGEditor/Private/WidgetBlueprint.cpp b/Engine/Source/Editor/UMGEditor/Private/WidgetBlueprint.cpp
--- a/Engine/Source/Editor/UMGEditor/Private/WidgetBlueprint.cpp
+++ b/Engine/Source/Editor/UMGEditor/Private/WidgetBlueprint.cpp
@@ -1,6 +1,7 @@
 // Copyright 1998-2014 Epic Games, Inc. All Rights Reserved.
 
 #include "UMGEditorPrivatePCH.h"
+#include "Blueprint/WidgetTree.h"
 #include "Runtime/MovieSceneCore/Classes/MovieScene.h"
 #include "PropertyTag.h"
 
diff --git a/Engine/Source/Editor/UMGEditor/Private/WidgetBlueprintCompiler.cpp b/Engine/Source/Editor/UMGEditor/Private/WidgetBlueprintCompiler.cpp
--- a/Engine/Source/Editor/UMGEditor/Private/WidgetBlueprintCompiler.cpp
+++ b/Engine/Source/Editor/UMGEditor/Private/WidgetBlueprintCompiler.cpp
@@ -8,6 +8,7 @@
 #include "KismetReinstanceUtilities.h"
 #include "MovieScene.h"
 #include "K2Node_FunctionEntry.h"
+#include "Blueprint/WidgetTree.h"
 
 FOnWidgetBlueprintCompiled FWidgetBlueprintCompiler::OnWidgetBlueprintCompiled = FOnWidgetBlueprintCompiled();
 
diff --git a/Engine/Source/Editor/UMGEditor/Private/WidgetBlueprintEditor.cpp b/Engine/Source/Editor/UMGEditor/Private/WidgetBlueprintEditor.cpp
--- a/Engine/Source/Editor/UMGEditor/Private/WidgetBlueprintEditor.cpp
+++ b/Engine/Source/Editor/UMGEditor/Private/WidgetBlueprintEditor.cpp
@@ -12,11 +12,13 @@
 #include "PropertyCustomizationHelpers.h"
 
 #include "WidgetBlueprintApplicationModes.h"
-//#include "WidgetDefafaultsApplicationMode.h"
+#include "WidgetBlueprintEditorUtils.h"
 #include "WidgetDesignerApplicationMode.h"
 #include "WidgetGraphApplicationMode.h"
 
 #include "WidgetBlueprintEditorToolbar.h"
+#include "Blueprint/WidgetTree.h"
+#include "Components/CanvasPanel.h"
 
 #define LOCTEXT_NAMESPACE "UMG"
 
diff --git a/Engine/Source/Editor/UMGEditor/Private/WidgetBlueprintEditor.h b/Engine/Source/Editor/UMGEditor/Private/WidgetBlueprintEditor.h
--- a/Engine/Source/Editor/UMGEditor/Private/WidgetBlueprintEditor.h
+++ b/Engine/Source/Editor/UMGEditor/Private/WidgetBlueprintEditor.h
@@ -5,6 +5,7 @@
 #include "Editor/Kismet/Public/BlueprintEditor.h"
 #include "ISequencer.h"
 #include "PreviewScene.h"
+#include "WidgetReference.h"
 
 class ISequencer;
 
diff --git a/Engine/Source/Editor/UMGEditor/Private/WidgetBlueprintEditorToolbar.cpp b/Engine/Source/Editor/UMGEditor/Private/WidgetBlueprintEditorToolbar.cpp
--- a/Engine/Source/Editor/UMGEditor/Private/WidgetBlueprintEditorToolbar.cpp
+++ b/Engine/Source/Editor/UMGEditor/Private/WidgetBlueprintEditorToolbar.cpp
@@ -1,6 +1,7 @@
 // Copyright 1998-2014 Epic Games, Inc. All Rights Reserved.
 
 #include "UMGEditorPrivatePCH.h"
+#include "IDocumentation.h"
 #include "WidgetBlueprintEditorToolbar.h"
 
 #include "WidgetBlueprintEditor.h"
diff --git a/Engine/Source/Editor/UMGEditor/Private/WidgetBlueprintEditorUtils.cpp b/Engine/Source/Editor/UMGEditor/Private/WidgetBlueprintEditorUtils.cpp
--- a/Engine/Source/Editor/UMGEditor/Private/WidgetBlueprintEditorUtils.cpp
+++ b/Engine/Source/Editor/UMGEditor/Private/WidgetBlueprintEditorUtils.cpp
@@ -1,7 +1,8 @@
 // Copyright 1998-2014 Epic Games, Inc. All Rights Reserved.
 
 #include "UMGEditorPrivatePCH.h"
-
+#include "Components/PanelSlot.h"
+#include "Components/PanelWidget.h"
 #include "WidgetBlueprintEditorUtils.h"
 #include "WidgetBlueprintEditor.h"
 #include "Kismet2NameValidators.h"
diff --git a/Engine/Source/Editor/UMGEditor/Private/WidgetBlueprintFactory.cpp b/Engine/Source/Editor/UMGEditor/Private/WidgetBlueprintFactory.cpp
--- a/Engine/Source/Editor/UMGEditor/Private/WidgetBlueprintFactory.cpp
+++ b/Engine/Source/Editor/UMGEditor/Private/WidgetBlueprintFactory.cpp
@@ -5,6 +5,7 @@
 #include "Kismet2/BlueprintEditorUtils.h"
 #include "Kismet2/KismetEditorUtilities.h"
 
+#include "WidgetBlueprintFactory.h"
 #define LOCTEXT_NAMESPACE "UWidgetBlueprintFactory"
 
 /*------------------------------------------------------------------------------
diff --git a/Engine/Source/Editor/UMGEditor/Private/WidgetGraphSchema.cpp b/Engine/Source/Editor/UMGEditor/Private/WidgetGraphSchema.cpp
--- a/Engine/Source/Editor/UMGEditor/Private/WidgetGraphSchema.cpp
+++ b/Engine/Source/Editor/UMGEditor/Private/WidgetGraphSchema.cpp
@@ -1,6 +1,7 @@
 // Copyright 1998-2014 Epic Games, Inc. All Rights Reserved.
 
 #include "UMGEditorPrivatePCH.h"
+#include "WidgetGraphSchema.h"
 
 UWidgetGraphSchema::UWidgetGraphSchema(const FPostConstructInitializeProperties& PCIP)
 : Super(PCIP)
diff --git a/Engine/Source/Editor/UMGEditor/Public/WidgetBlueprint.h b/Engine/Source/Editor/UMGEditor/Public/WidgetBlueprint.h
--- a/Engine/Source/Editor/UMGEditor/Public/WidgetBlueprint.h
+++ b/Engine/Source/Editor/UMGEditor/Public/WidgetBlueprint.h
@@ -1,7 +1,7 @@
 // Copyright 1998-2014 Epic Games, Inc. All Rights Reserved.
 
 #pragma once
-
+#include "Blueprint/WidgetBlueprintGeneratedClass.h"
 #include "Engine/Blueprint.h"
 #include "UserWidget.h"
 #include "WidgetAnimation.h"
diff --git a/Engine/Source/Runtime/UMG/Public/Blueprint/UserWidget.h b/Engine/Source/Runtime/UMG/Public/Blueprint/UserWidget.h
--- a/Engine/Source/Runtime/UMG/Public/Blueprint/UserWidget.h
+++ b/Engine/Source/Runtime/UMG/Public/Blueprint/UserWidget.h
@@ -1,12 +1,13 @@
 // Copyright 1998-2014 Epic Games, Inc. All Rights Reserved.
 
 #pragma once
-
+#include "Components/SlateWrapperTypes.h"
+#include "Components/Widget.h"
 #include "Geometry.h"
 #include "Engine/GameInstance.h"
 
 #include "NamedSlotInterface.h"
-
+#include "Slate/Anchors.h"
 #include "UserWidget.generated.h"
 
 static FGeometry NullGeometry;
@@ -14,7 +15,9 @@ static FSlateRect NullRect;
 static FSlateWindowElementList NullElementList;
 static FWidgetStyle NullStyle;
 
+class UWidget;
 class UWidgetAnimation;
+struct FAnchors;
 
 /**
  * The state passed into OnPaint that we can expose as a single painting structure to blueprints to
diff --git a/Engine/Source/Runtime/UMG/Public/Blueprint/WidgetBlueprintGeneratedClass.h b/Engine/Source/Runtime/UMG/Public/Blueprint/WidgetBlueprintGeneratedClass.h
--- a/Engine/Source/Runtime/UMG/Public/Blueprint/WidgetBlueprintGeneratedClass.h
+++ b/Engine/Source/Runtime/UMG/Public/Blueprint/WidgetBlueprintGeneratedClass.h
@@ -1,10 +1,13 @@
 // Copyright 1998-2014 Epic Games, Inc. All Rights Reserved.
 
 #pragma once
-
+#include "Engine/BlueprintGeneratedClass.h"
 #include "WidgetBlueprintGeneratedClass.generated.h"
 
 class UMovieScene;
+class UStructProperty;
+class UUserWidget;
+class UWidgetAnimation;
 
 UENUM()
 namespace EBindingKind
@@ -74,5 +77,5 @@ class UMG_API UWidgetBlueprintGeneratedClass : public UBlueprintGeneratedClass
 	 * call into its generated class and ask to be initialized.  The class will perform all the delegate
 	 * binding and wiring nessesary to have the user's widget perform as desired.
 	 */
-	void InitializeWidget(class UUserWidget* UserWidget) const;
+	void InitializeWidget(UUserWidget* UserWidget) const;
 };
\ No newline at end of file
diff --git a/Engine/Source/Runtime/UMG/Public/Components/CanvasPanel.h b/Engine/Source/Runtime/UMG/Public/Components/CanvasPanel.h
--- a/Engine/Source/Runtime/UMG/Public/Components/CanvasPanel.h
+++ b/Engine/Source/Runtime/UMG/Public/Components/CanvasPanel.h
@@ -1,9 +1,11 @@
 // Copyright 1998-2014 Epic Games, Inc. All Rights Reserved.
 
 #pragma once
-
+#include "Components/PanelWidget.h"
 #include "CanvasPanel.generated.h"
 
+class UCanvasPanelSlot;
+
 /**
  * The canvas panel is a designer friendly panel that allows widgets to be laid out at arbitrary 
  * locations, anchored and z-ordered with other children of the canvas.  The canvas is a great widget
diff --git a/Engine/Source/Runtime/UMG/Public/Components/CanvasPanelSlot.h b/Engine/Source/Runtime/UMG/Public/Components/CanvasPanelSlot.h
--- a/Engine/Source/Runtime/UMG/Public/Components/CanvasPanelSlot.h
+++ b/Engine/Source/Runtime/UMG/Public/Components/CanvasPanelSlot.h
@@ -3,7 +3,7 @@
 #pragma once
 
 #include "Anchors.h"
-
+#include "Slate/SConstraintCanvas.h"
 #include "CanvasPanelSlot.generated.h"
 
 USTRUCT(BlueprintType)
diff --git a/Engine/Source/Runtime/UMG/Public/Components/GridSlot.h b/Engine/Source/Runtime/UMG/Public/Components/GridSlot.h
--- a/Engine/Source/Runtime/UMG/Public/Components/GridSlot.h
+++ b/Engine/Source/Runtime/UMG/Public/Components/GridSlot.h
@@ -1,7 +1,7 @@
 // Copyright 1998-2014 Epic Games, Inc. All Rights Reserved.
 
 #pragma once
-
+#include "Components/PanelSlot.h"
 #include "SlateWrapperTypes.h"
 
 #include "GridSlot.generated.h"
diff --git a/Engine/Source/Runtime/UMG/Public/Components/HorizontalBox.h b/Engine/Source/Runtime/UMG/Public/Components/HorizontalBox.h
--- a/Engine/Source/Runtime/UMG/Public/Components/HorizontalBox.h
+++ b/Engine/Source/Runtime/UMG/Public/Components/HorizontalBox.h
@@ -1,9 +1,11 @@
 // Copyright 1998-2014 Epic Games, Inc. All Rights Reserved.
 
 #pragma once
-
+#include "Components/PanelWidget.h"
 #include "HorizontalBox.generated.h"
 
+class UHorizontalBoxSlot;
+
 /**
  * Allows widgets to be laid out in a flow horizontally.
  */
diff --git a/Engine/Source/Runtime/UMG/Public/Components/HorizontalBoxSlot.h b/Engine/Source/Runtime/UMG/Public/Components/HorizontalBoxSlot.h
--- a/Engine/Source/Runtime/UMG/Public/Components/HorizontalBoxSlot.h
+++ b/Engine/Source/Runtime/UMG/Public/Components/HorizontalBoxSlot.h
@@ -1,7 +1,7 @@
 // Copyright 1998-2014 Epic Games, Inc. All Rights Reserved.
 
 #pragma once
-
+#include "Components/PanelSlot.h"
 #include "SlateWrapperTypes.h"
 
 #include "HorizontalBoxSlot.generated.h"
diff --git a/Engine/Source/Runtime/UMG/Public/Components/NamedSlotInterface.h b/Engine/Source/Runtime/UMG/Public/Components/NamedSlotInterface.h
--- a/Engine/Source/Runtime/UMG/Public/Components/NamedSlotInterface.h
+++ b/Engine/Source/Runtime/UMG/Public/Components/NamedSlotInterface.h
@@ -1,9 +1,11 @@
 // Copyright 1998-2014 Epic Games, Inc. All Rights Reserved.
 
 #pragma once
-
+#include "Interface.h"
 #include "NamedSlotInterface.generated.h"
 
+class UWidget;
+
 /**  */
 UINTERFACE(MinimalAPI, meta=( CannotImplementInterfaceInBlueprint ))
 class UNamedSlotInterface : public UInterface
diff --git a/Engine/Source/Runtime/UMG/Public/Components/PanelWidget.h b/Engine/Source/Runtime/UMG/Public/Components/PanelWidget.h
--- a/Engine/Source/Runtime/UMG/Public/Components/PanelWidget.h
+++ b/Engine/Source/Runtime/UMG/Public/Components/PanelWidget.h
@@ -1,7 +1,8 @@
 // Copyright 1998-2014 Epic Games, Inc. All Rights Reserved.
 
 #pragma once
-
+#include "Components/PanelSlot.h"
+#include "Components/Widget.h"
 #include "PanelWidget.generated.h"
 
 /** The base class for all UMG panel widgets.  Panel widgets layout a collection of child widgets. */
diff --git a/Engine/Source/Runtime/UMG/Public/Components/UniformGridSlot.h b/Engine/Source/Runtime/UMG/Public/Components/UniformGridSlot.h
--- a/Engine/Source/Runtime/UMG/Public/Components/UniformGridSlot.h
+++ b/Engine/Source/Runtime/UMG/Public/Components/UniformGridSlot.h
@@ -1,7 +1,7 @@
 // Copyright 1998-2014 Epic Games, Inc. All Rights Reserved.
 
 #pragma once
-
+#include "Components/PanelSlot.h"
 #include "SlateWrapperTypes.h"
 
 #include "UniformGridSlot.generated.h"
diff --git a/Engine/Source/Runtime/UMG/Public/Components/VerticalBox.h b/Engine/Source/Runtime/UMG/Public/Components/VerticalBox.h
--- a/Engine/Source/Runtime/UMG/Public/Components/VerticalBox.h
+++ b/Engine/Source/Runtime/UMG/Public/Components/VerticalBox.h
@@ -1,9 +1,11 @@
 // Copyright 1998-2014 Epic Games, Inc. All Rights Reserved.
 
 #pragma once
-
+#include "Components/PanelWidget.h"
 #include "VerticalBox.generated.h"
 
+class UVerticalBoxSlot;
+
 /**
  * A vertical box widget is a layout panel allowing child widgets to be automatically laid out
  * vertically.
diff --git a/Engine/Source/Runtime/UMG/Public/Components/VerticalBoxSlot.h b/Engine/Source/Runtime/UMG/Public/Components/VerticalBoxSlot.h
--- a/Engine/Source/Runtime/UMG/Public/Components/VerticalBoxSlot.h
+++ b/Engine/Source/Runtime/UMG/Public/Components/VerticalBoxSlot.h
@@ -1,7 +1,7 @@
 // Copyright 1998-2014 Epic Games, Inc. All Rights Reserved.
 
 #pragma once
-
+#include "Components/PanelSlot.h"
 #include "SlateWrapperTypes.h"
 
 #include "VerticalBoxSlot.generated.h"
diff --git a/Engine/Source/Runtime/UMG/Public/Components/Widget.h b/Engine/Source/Runtime/UMG/Public/Components/Widget.h
--- a/Engine/Source/Runtime/UMG/Public/Components/Widget.h
+++ b/Engine/Source/Runtime/UMG/Public/Components/Widget.h
@@ -1,12 +1,14 @@
 // Copyright 1998-2014 Epic Games, Inc. All Rights Reserved.
 
 #pragma once
-
+#include "Components/Visual.h"
 #include "SlateWrapperTypes.h"
 #include "WidgetTransform.h"
 
 #include "Widget.generated.h"
 
+class UPanelSlot;
+
 /**
  * Helper macro for binding to a delegate or using the constant value when constructing the underlying SWidget
  */
