Temporary backout of CL2321071

diff --git a/Engine/Source/Editor/UMGEditor/Private/UMGEditorPrivatePCH.h b/Engine/Source/Editor/UMGEditor/Private/UMGEditorPrivatePCH.h
--- a/Engine/Source/Editor/UMGEditor/Private/UMGEditorPrivatePCH.h
+++ b/Engine/Source/Editor/UMGEditor/Private/UMGEditorPrivatePCH.h
@@ -1,11 +1,54 @@
 // Copyright 1998-2014 Epic Games, Inc. All Rights Reserved.
 
 #pragma once
-#include "Platform.h"
-#include "EditorStyleSet.h"
-#include "Internationalization/Internationalization.h"
-#include "Internationalization/Text.h"
-#include "Misc/AssertionMacros.h"
-#include "Templates/SharedPointer.h"
-#include "Templates/Casts.h"
-#include "Widgets/DeclarativeSyntaxSupport.h"
+#include "Slate.h"
+#include "SlateCore.h"
+#include "UnrealEd.h"
+
+#include "IDocumentation.h"
+
+#include "BlueprintGraphDefinitions.h"
+
+#include "BlueprintEditorUtils.h"
+
+#include "ScopedTransaction.h"
+#include "TextFilter.h"
+
+#include "ObjectEditorUtils.h"
+
+#include "UMG.h"
+#include "UMGStyle.h"
+
+#include "WidgetDesignerSettings.h"
+
+#include "UMGEditorClasses.h"
+#include "WidgetBlueprint.h"
+
+#include "BlueprintEditorModes.h"
+#include "WidgetBlueprintApplicationModes.h"
+#include "WidgetBlueprintApplicationMode.h"
+
+#include "WidgetReference.h"
+#include "WidgetBlueprintEditor.h"
+
+#include "AssetTypeActions_Base.h"
+
+#include "DesignTimeUtils.h"
+
+#include "DesignerExtension.h"
+#include "VerticalSlotExtension.h"
+#include "HorizontalSlotExtension.h"
+#include "CanvasSlotExtension.h"
+#include "UniformGridSlotExtension.h"
+#include "GridSlotExtension.h"
+
+#include "DetailCustomizations.h"
+
+#include "WidgetTemplate.h"
+#include "WidgetTemplateDragDropOp.h"
+
+#include "WidgetBlueprintEditorUtils.h"
+
+#include "WidgetBlueprintLibrary.h"
+
+#include "WidgetBlueprintCompiler.h"
