TTP 348925 - UMG Removing a selection changed notification that's emitted during object replacement passes during blueprint compiles.  It's causing massive memory explosions from regenerations of the details panel due to selection change notification.  Fixing in both master and 4.5.

diff --git a/Engine/Source/Editor/UMGEditor/Private/WidgetBlueprintEditor.cpp b/Engine/Source/Editor/UMGEditor/Private/WidgetBlueprintEditor.cpp
--- a/Engine/Source/Editor/UMGEditor/Private/WidgetBlueprintEditor.cpp
+++ b/Engine/Source/Editor/UMGEditor/Private/WidgetBlueprintEditor.cpp
@@ -212,9 +212,6 @@ void FWidgetBlueprintEditor::OnBlueprintChanged(UBlueprint* InBlueprint)
 
 void FWidgetBlueprintEditor::OnObjectsReplaced(const TMap<UObject*, UObject*>& ReplacementMap)
 {
-	// Objects being replaced could force a selection change
-	OnSelectedWidgetsChanging.Broadcast();
-
 	// Remove dead references and update references
 	for ( int32 HandleIndex = WidgetHandlePool.Num() - 1; HandleIndex >= 0; HandleIndex-- )
 	{
@@ -233,9 +230,6 @@ void FWidgetBlueprintEditor::OnObjectsReplaced(const TMap<UObject*, UObject*>& R
 			WidgetHandlePool.RemoveAtSwap(HandleIndex);
 		}
 	}
-
-	// Fire the selection updated event to ensure everyone is watching the same widgets.
-	OnSelectedWidgetsChanged.Broadcast();
 }
 
 bool FWidgetBlueprintEditor::CanDeleteSelectedWidgets()
