UMG - Now always dirtying the blueprint when any property on a widget changes.

diff --git a/Engine/Source/Editor/UMGEditor/Private/Details/SWidgetDetailsView.cpp b/Engine/Source/Editor/UMGEditor/Private/Details/SWidgetDetailsView.cpp
--- a/Engine/Source/Editor/UMGEditor/Private/Details/SWidgetDetailsView.cpp
+++ b/Engine/Source/Editor/UMGEditor/Private/Details/SWidgetDetailsView.cpp
@@ -462,6 +462,10 @@ void SWidgetDetailsView::NotifyPostChange(const FPropertyChangedEvent& PropertyC
 
 		const bool bIsModify = false;
 		Editor->MigrateFromChain(PropertyThatChanged, bIsModify);
+
+		// Any time we migrate a property value we need to mark the blueprint as structurally modified so users don't need to recompile it manually
+		// before they see it play in game using the latest version.
+		FBlueprintEditorUtils::MarkBlueprintAsStructurallyModified(BlueprintEditor.Pin()->GetBlueprintObj());
 	}
 
 	// If the property that changed is marked as "DesignerRebuild" we invalidate
