Fixed trying to set an error on an invisible text box.

TTP 349122: EDITOR: UMG: CRASH: Haivng no characters in the Name field causes editor to crash

diff --git a/Engine/Source/Editor/UMGEditor/Private/Details/SWidgetDetailsView.cpp b/Engine/Source/Editor/UMGEditor/Private/Details/SWidgetDetailsView.cpp
--- a/Engine/Source/Editor/UMGEditor/Private/Details/SWidgetDetailsView.cpp
+++ b/Engine/Source/Editor/UMGEditor/Private/Details/SWidgetDetailsView.cpp
@@ -58,7 +58,7 @@ void SWidgetDetailsView::Construct(const FArguments& InArgs, TSharedPtr<FWidgetB
 				.WidthOverride(200.0f)
 				.VAlign(VAlign_Center)
 				[
-					SAssignNew(NameTextBox, SEditableTextBox)
+					SNew(SEditableTextBox)
 					.SelectAllTextWhenFocused(true)
 					.ToolTipText(LOCTEXT("CategoryToolTip", "Sets the category of the widget"))
 					.HintText(LOCTEXT("Category", "Category"))
