UMG - We now mark blueprints as modified (but not structural) when properties change on widgets.

diff --git a/Engine/Source/Editor/UMGEditor/Private/Details/SWidgetDetailsView.cpp b/Engine/Source/Editor/UMGEditor/Private/Details/SWidgetDetailsView.cpp
--- a/Engine/Source/Editor/UMGEditor/Private/Details/SWidgetDetailsView.cpp
+++ b/Engine/Source/Editor/UMGEditor/Private/Details/SWidgetDetailsView.cpp
@@ -266,8 +266,7 @@ void SWidgetDetailsView::HandleCategoryTextCommitted(const FText& Text, ETextCom
 			UUserWidget* WidgetCDO = Widget->GetClass()->GetDefaultObject<UUserWidget>();
 			WidgetCDO->PaletteCategory = Text;
 
-			// We mark it as structurally modified to ensure that all palettes refresh to move the widget into
-			// the correct new category.
+			// Immediately force a rebuild so that all palettes update to show it in a new category.
 			FBlueprintEditorUtils::MarkBlueprintAsStructurallyModified(BlueprintEditor.Pin()->GetBlueprintObj());
 		}
 	}
@@ -463,9 +462,9 @@ void SWidgetDetailsView::NotifyPostChange(const FPropertyChangedEvent& PropertyC
 		const bool bIsModify = false;
 		Editor->MigrateFromChain(PropertyThatChanged, bIsModify);
 
-		// Any time we migrate a property value we need to mark the blueprint as structurally modified so users don't need to recompile it manually
-		// before they see it play in game using the latest version.
-		FBlueprintEditorUtils::MarkBlueprintAsStructurallyModified(BlueprintEditor.Pin()->GetBlueprintObj());
+		// Any time we migrate a property value we need to mark the blueprint as structurally modified so users don't need 
+		// to recompile it manually before they see it play in game using the latest version.
+		FBlueprintEditorUtils::MarkBlueprintAsModified(BlueprintEditor.Pin()->GetBlueprintObj());
 	}
 
 	// If the property that changed is marked as "DesignerRebuild" we invalidate
