UMG - If custom sizing is enabled, we apply a scale of 1 DPI wise to the designer preview.

diff --git a/Engine/Source/Editor/UMGEditor/Private/Designer/SDesignerView.cpp b/Engine/Source/Editor/UMGEditor/Private/Designer/SDesignerView.cpp
--- a/Engine/Source/Editor/UMGEditor/Private/Designer/SDesignerView.cpp
+++ b/Engine/Source/Editor/UMGEditor/Private/Designer/SDesignerView.cpp
@@ -491,6 +491,15 @@ FOptionalSize SDesignerView::GetPreviewHeight() const
 
 float SDesignerView::GetPreviewDPIScale() const
 {
+	// If the user is using a custom size then we disable the DPI scaling logic.
+	if ( UUserWidget* DefaultWidget = GetDefaultWidget() )
+	{
+		if ( DefaultWidget->bUseDesignTimeSize )
+		{
+			return 1.0f;
+		}
+	}
+
 	return GetDefault<URendererSettings>(URendererSettings::StaticClass())->GetDPIScaleBasedOnSize(FIntPoint(PreviewWidth, PreviewHeight));
 }
 
diff --git a/Engine/Source/Editor/UMGEditor/Private/WidgetBlueprintEditor.cpp b/Engine/Source/Editor/UMGEditor/Private/WidgetBlueprintEditor.cpp
--- a/Engine/Source/Editor/UMGEditor/Private/WidgetBlueprintEditor.cpp
+++ b/Engine/Source/Editor/UMGEditor/Private/WidgetBlueprintEditor.cpp
@@ -24,7 +24,7 @@
 
 FWidgetBlueprintEditor::FWidgetBlueprintEditor()
 	: PreviewScene(FPreviewScene::ConstructionValues().AllowAudioPlayback(true).ShouldSimulatePhysics(true))
-	, PreviewBlueprint(NULL)
+	, PreviewBlueprint(nullptr)
 {
 	
 }
