UMG - Refactoring the slate color customization to play nice with bindings.  Adding a new property existenion system to the details panel.  Allows for special case additions to be added to the property after it's customization has been used.  Now used for bindings.  User widgets now have foreground color and color and opacity support.

diff --git a/Engine/Source/Editor/DetailCustomizations/Private/MathStructCustomizations.cpp b/Engine/Source/Editor/DetailCustomizations/Private/MathStructCustomizations.cpp
--- a/Engine/Source/Editor/DetailCustomizations/Private/MathStructCustomizations.cpp
+++ b/Engine/Source/Editor/DetailCustomizations/Private/MathStructCustomizations.cpp
@@ -416,8 +416,6 @@ TSharedRef<IPropertyTypeCustomization> FColorStructCustomization::MakeInstance()
 
 void FColorStructCustomization::CustomizeHeader( TSharedRef<class IPropertyHandle> InStructPropertyHandle, class FDetailWidgetRow& InHeaderRow, IPropertyTypeCustomizationUtils& StructCustomizationUtils )
 {
-	TSharedPtr<FAssetThumbnailPool> Pool = StructCustomizationUtils.GetThumbnailPool();
-
 	StructPropertyHandle = InStructPropertyHandle;
 
 	bIsLinearColor = CastChecked<UStructProperty>( StructPropertyHandle->GetProperty() )->Struct->GetFName() == NAME_LinearColor;
@@ -435,8 +433,6 @@ void FColorStructCustomization::MakeHeaderRow( TSharedRef<class IPropertyHandle>
 	const bool bDisplayResetToDefault = false;
 	const FString DisplayNameOverride = TEXT("");
 
-	FSlateFontInfo NormalText = IDetailLayoutBuilder::GetDetailFont();
-
 	Row.NameContent()
 	[
 		StructPropertyHandle->CreatePropertyNameWidget( DisplayNameOverride, bDisplayResetToDefault )
@@ -445,7 +441,15 @@ void FColorStructCustomization::MakeHeaderRow( TSharedRef<class IPropertyHandle>
 	.MinDesiredWidth(250.0f)
 	.MaxDesiredWidth(250.0f)
 	[
-		SNew( SHorizontalBox )
+		CreateColorWidget()
+	];
+}
+
+TSharedRef<SWidget> FColorStructCustomization::CreateColorWidget()
+{
+	FSlateFontInfo NormalText = IDetailLayoutBuilder::GetDetailFont();
+
+	return SNew( SHorizontalBox )
 		+ SHorizontalBox::Slot()
 		.VAlign(VAlign_Center)
 		[
@@ -481,8 +485,7 @@ void FColorStructCustomization::MakeHeaderRow( TSharedRef<class IPropertyHandle>
 			.IgnoreAlpha(true)
 			.OnMouseButtonDown( this, &FColorStructCustomization::OnMouseButtonDownColorBlock )
 			.Size( FVector2D( 35.0f, 12.0f ) )
-		]
-	];
+		];
 }
 
 void FColorStructCustomization::GetSortedChildren( TSharedRef<IPropertyHandle> InStructPropertyHandle, TArray< TSharedRef<IPropertyHandle> >& OutChildren )
diff --git a/Engine/Source/Editor/DetailCustomizations/Private/MathStructCustomizations.h b/Engine/Source/Editor/DetailCustomizations/Private/MathStructCustomizations.h
--- a/Engine/Source/Editor/DetailCustomizations/Private/MathStructCustomizations.h
+++ b/Engine/Source/Editor/DetailCustomizations/Private/MathStructCustomizations.h
@@ -146,7 +146,7 @@ class FColorStructCustomization : public FMathStructCustomization
 public:
 	static TSharedRef<IPropertyTypeCustomization> MakeInstance();
 	
-private:
+protected:
 	FColorStructCustomization()
 		: bIgnoreAlpha( false )
 		, bIsInlineColorPickerVisible(false)
@@ -159,6 +159,9 @@ private:
 	virtual void MakeHeaderRow( TSharedRef<class IPropertyHandle>& InStructPropertyHandle, FDetailWidgetRow& Row ) override;
 	virtual void GetSortedChildren( TSharedRef<IPropertyHandle> StructPropertyHandle, TArray< TSharedRef<IPropertyHandle> >& OutChildren ) override;
 
+	/** Creates the color widget that when clicked spawns the color picker window. */
+	TSharedRef<SWidget> CreateColorWidget();
+
 	/**
 	 * Get the color used by this struct as a linear color value
 	 * @param InColor To be filled with the color value used by this struct, or white if this struct is being used to edit multiple values  
diff --git a/Engine/Source/Editor/DetailCustomizations/Private/SlateBrushCustomization.cpp b/Engine/Source/Editor/DetailCustomizations/Private/SlateBrushCustomization.cpp
--- a/Engine/Source/Editor/DetailCustomizations/Private/SlateBrushCustomization.cpp
+++ b/Engine/Source/Editor/DetailCustomizations/Private/SlateBrushCustomization.cpp
@@ -1039,7 +1039,7 @@ private:
 	static float TargetHeight;
 };
 
-float SSlateBrushStaticPreview::TargetHeight = 40.f;
+float SSlateBrushStaticPreview::TargetHeight = 18.0f;
 
 // SBrushResourceObjectBox
 ////////////////////////////////////////////////////////////////////////////////
diff --git a/Engine/Source/Editor/DetailCustomizations/Private/SlateColorCustomization.cpp b/Engine/Source/Editor/DetailCustomizations/Private/SlateColorCustomization.cpp
--- a/Engine/Source/Editor/DetailCustomizations/Private/SlateColorCustomization.cpp
+++ b/Engine/Source/Editor/DetailCustomizations/Private/SlateColorCustomization.cpp
@@ -8,30 +8,30 @@ TSharedRef<IPropertyTypeCustomization> FSlateColorCustomization::MakeInstance()
 	return MakeShareable( new FSlateColorCustomization() );
 }
 
-void FSlateColorCustomization::CustomizeHeader( TSharedRef<class IPropertyHandle> StructPropertyHandle, class FDetailWidgetRow& HeaderRow, IPropertyTypeCustomizationUtils& StructCustomizationUtils )
+void FSlateColorCustomization::CustomizeHeader(TSharedRef<class IPropertyHandle> InStructPropertyHandle, class FDetailWidgetRow& InHeaderRow, IPropertyTypeCustomizationUtils& StructCustomizationUtils)
 {
 	static const FName ColorUseRuleKey(TEXT("ColorUseRule"));
 	static const FName SpecifiedColorKey(TEXT("SpecifiedColor"));
 
-	ColorRuleHandle = StructPropertyHandle->GetChildHandle(ColorUseRuleKey);
-	SpecifiedColorHandle = StructPropertyHandle->GetChildHandle(SpecifiedColorKey);
+	StructPropertyHandle = InStructPropertyHandle;
+
+	ColorRuleHandle = InStructPropertyHandle->GetChildHandle(ColorUseRuleKey);
+	SpecifiedColorHandle = InStructPropertyHandle->GetChildHandle(SpecifiedColorKey);
 
 	check(ColorRuleHandle.IsValid());
 	check(SpecifiedColorHandle.IsValid());
+
+	ColorRuleHandle->MarkHiddenByCustomization();
+	SpecifiedColorHandle->SetOnPropertyValueChanged(FSimpleDelegate::CreateSP(SharedThis(this), &FSlateColorCustomization::OnValueChanged));
+
+	FColorStructCustomization::CustomizeHeader(SpecifiedColorHandle.ToSharedRef(), InHeaderRow, StructCustomizationUtils);
 }
 
-void FSlateColorCustomization::CustomizeChildren( TSharedRef<class IPropertyHandle> StructPropertyHandle, class IDetailChildrenBuilder& StructBuilder, IPropertyTypeCustomizationUtils& StructCustomizationUtils )
+void FSlateColorCustomization::MakeHeaderRow(TSharedRef<class IPropertyHandle>& InStructPropertyHandle, FDetailWidgetRow& Row)
 {
-	SpecifiedColorHandle->SetOnPropertyValueChanged( FSimpleDelegate::CreateSP( SharedThis( this ), &FSlateColorCustomization::OnValueChanged ) );
-	
-	IDetailPropertyRow& ColorRow = StructBuilder.AddChildProperty(SpecifiedColorHandle.ToSharedRef());
-
-	TSharedPtr<SWidget> NameWidget;
-	TSharedPtr<SWidget> ValueWidget;
-	FDetailWidgetRow Row;
-	ColorRow.GetDefaultWidgets(NameWidget, ValueWidget, Row);
+	// NOTE: Ignore InStructPropertyHandle, it's going to be the specified color handle that we passed to the color customization base class.
 
-	ColorRow.CustomWidget(/*bShowChildren*/ true)
+	Row
 		.NameContent()
 		[
 			SNew(STextBlock)
@@ -48,7 +48,7 @@ void FSlateColorCustomization::CustomizeChildren( TSharedRef<class IPropertyHand
 			+ SHorizontalBox::Slot()
 			.FillWidth(1.0f)
 			[
-				ValueWidget.ToSharedRef()
+				CreateColorWidget()
 			]
 
 			+ SHorizontalBox::Slot()
diff --git a/Engine/Source/Editor/DetailCustomizations/Private/SlateColorCustomization.h b/Engine/Source/Editor/DetailCustomizations/Private/SlateColorCustomization.h
--- a/Engine/Source/Editor/DetailCustomizations/Private/SlateColorCustomization.h
+++ b/Engine/Source/Editor/DetailCustomizations/Private/SlateColorCustomization.h
@@ -2,7 +2,9 @@
 
 #pragma once
 
-class FSlateColorCustomization : public IPropertyTypeCustomization
+#include "MathStructCustomizations.h"
+
+class FSlateColorCustomization : public FColorStructCustomization
 {
 public:
 	static TSharedRef<IPropertyTypeCustomization> MakeInstance();
@@ -10,7 +12,7 @@ public:
 	/** IPropertyTypeCustomization interface */
 	virtual void CustomizeHeader( TSharedRef<class IPropertyHandle> StructPropertyHandle, class FDetailWidgetRow& HeaderRow, IPropertyTypeCustomizationUtils& StructCustomizationUtils ) override;
 
-	virtual void CustomizeChildren( TSharedRef<class IPropertyHandle> StructPropertyHandle, class IDetailChildrenBuilder& StructBuilder, IPropertyTypeCustomizationUtils& StructCustomizationUtils ) override;              
+	virtual void MakeHeaderRow(TSharedRef<class IPropertyHandle>& InStructPropertyHandle, FDetailWidgetRow& Row) override;
 
 private:
 
@@ -25,6 +27,9 @@ private:
 
 private:
 
+	/** slate color struct handle */
+	TSharedPtr<IPropertyHandle> StructPropertyHandle;
+
 	TSharedPtr<IPropertyHandle> ColorRuleHandle;
 	TSharedPtr<IPropertyHandle> SpecifiedColorHandle;
 };
diff --git a/Engine/Source/Editor/PropertyEditor/Private/IDetailsViewPrivate.h b/Engine/Source/Editor/PropertyEditor/Private/IDetailsViewPrivate.h
--- a/Engine/Source/Editor/PropertyEditor/Private/IDetailsViewPrivate.h
+++ b/Engine/Source/Editor/PropertyEditor/Private/IDetailsViewPrivate.h
@@ -79,6 +79,8 @@ public:
 
 	virtual TSharedPtr<IDetailKeyframeHandler> GetKeyframeHandler() = 0;
 
+	virtual TSharedPtr<IDetailPropertyExtensionHandler> GetExtensionHandler() = 0;
+
 	/**
 	 * @return The thumbnail pool that should be used for thumbnails being rendered in this view
 	 */
diff --git a/Engine/Source/Editor/PropertyEditor/Private/SDetailSingleItemRow.cpp b/Engine/Source/Editor/PropertyEditor/Private/SDetailSingleItemRow.cpp
--- a/Engine/Source/Editor/PropertyEditor/Private/SDetailSingleItemRow.cpp
+++ b/Engine/Source/Editor/PropertyEditor/Private/SDetailSingleItemRow.cpp
@@ -5,6 +5,7 @@
 #include "DetailItemNode.h"
 #include "PropertyEditorHelpers.h"
 #include "IDetailKeyframeHandler.h"
+#include "IDetailPropertyExtensionHandler.h"
 
 namespace DetailWidgetConstants
 {
@@ -96,6 +97,8 @@ void SDetailSingleItemRow::Construct( const FArguments& InArgs, FDetailLayoutCus
 				Row.ValueWidget.Widget
 			];
 
+		ValueWidget = CreateExtensionWidget(ValueWidget.ToSharedRef(), *Customization, InOwnerTreeNode );
+
 		if( Row.IsEnabledAttr.IsBound() )
 		{
 			ValueWidget->SetEnabled( Row.IsEnabledAttr );
@@ -144,6 +147,7 @@ void SDetailSingleItemRow::Construct( const FArguments& InArgs, FDetailLayoutCus
 				.OnSlotResized( ColumnSizeData.OnWidthChanged )
 				[
 					SNew( SHorizontalBox )
+
 					+ SHorizontalBox::Slot()
 					.Padding( DetailWidgetConstants::RightRowPadding )
 					.HAlign( Row.ValueWidget.HorizontalAlignment )
@@ -300,6 +304,40 @@ const FSlateBrush* SDetailSingleItemRow::GetBorderImage() const
 	}
 }
 
+TSharedRef<SWidget> SDetailSingleItemRow::CreateExtensionWidget(TSharedRef<SWidget> ValueWidget, FDetailLayoutCustomization& InCustomization, TSharedRef<IDetailTreeNode> InTreeNode)
+{
+	IDetailsViewPrivate& DetailsView = InTreeNode->GetDetailsView();
+	TSharedPtr<IDetailPropertyExtensionHandler> ExtensionHandler = DetailsView.GetExtensionHandler();
+
+	if ( ExtensionHandler.IsValid() )
+	{
+		if ( InCustomization.HasPropertyNode() && ExtensionHandler.IsValid() )
+		{
+			TSharedPtr<IPropertyHandle> Handle = PropertyEditorHelpers::GetPropertyHandle(InCustomization.GetPropertyNode().ToSharedRef(), nullptr, nullptr);
+
+			UClass* ObjectClass = InCustomization.GetPropertyNode()->FindObjectItemParent()->GetObjectBaseClass();
+			if ( ExtensionHandler->IsPropertyExtenable(ObjectClass, *Handle) )
+			{
+				ValueWidget = SNew(SHorizontalBox)
+
+					+ SHorizontalBox::Slot()
+					.FillWidth(1.0f)
+					[
+						ValueWidget
+					]
+
+				+ SHorizontalBox::Slot()
+					.AutoWidth()
+					[
+						ExtensionHandler->GenerateExtensionWidget(ObjectClass, Handle)
+					];
+			}
+		}
+	}
+
+	return ValueWidget;
+}
+
 TSharedRef<SWidget> SDetailSingleItemRow::CreateKeyframeButton( FDetailLayoutCustomization& InCustomization, TSharedRef<IDetailTreeNode> InTreeNode )
 {
 	IDetailsViewPrivate& DetailsView = InTreeNode->GetDetailsView();
diff --git a/Engine/Source/Editor/PropertyEditor/Private/SDetailSingleItemRow.h b/Engine/Source/Editor/PropertyEditor/Private/SDetailSingleItemRow.h
--- a/Engine/Source/Editor/PropertyEditor/Private/SDetailSingleItemRow.h
+++ b/Engine/Source/Editor/PropertyEditor/Private/SDetailSingleItemRow.h
@@ -29,6 +29,7 @@ private:
 	void OnPasteProperty();
 	bool CanPasteProperty() const;
 	const FSlateBrush* GetBorderImage() const;
+	TSharedRef<SWidget> CreateExtensionWidget( TSharedRef<SWidget> ValueWidget, FDetailLayoutCustomization& InCustomization, TSharedRef<IDetailTreeNode> InTreeNode );
 	TSharedRef<SWidget> CreateKeyframeButton( FDetailLayoutCustomization& InCustomization, TSharedRef<IDetailTreeNode> InTreeNode );
 	FReply OnAddKeyframeClicked();
 	bool IsHighlighted() const;
diff --git a/Engine/Source/Editor/PropertyEditor/Private/SDetailsViewBase.cpp b/Engine/Source/Editor/PropertyEditor/Private/SDetailsViewBase.cpp
--- a/Engine/Source/Editor/PropertyEditor/Private/SDetailsViewBase.cpp
+++ b/Engine/Source/Editor/PropertyEditor/Private/SDetailsViewBase.cpp
@@ -258,7 +258,6 @@ bool SDetailsViewBase::IsPropertyEditingEnabled() const
 	return !IsPropertyEditingEnabledDelegate.IsBound() || IsPropertyEditingEnabledDelegate.Execute();
 }
 
-
 void SDetailsViewBase::SetKeyframeHandler( TSharedPtr<class IDetailKeyframeHandler> InKeyframeHandler )
 {
 	KeyframeHandler = InKeyframeHandler;
@@ -269,6 +268,16 @@ TSharedPtr<IDetailKeyframeHandler> SDetailsViewBase::GetKeyframeHandler()
 	return KeyframeHandler;
 }
 
+void SDetailsViewBase::SetExtensionHandler(TSharedPtr<class IDetailPropertyExtensionHandler> InExtensionHandler)
+{
+	ExtensionHandler = InExtensionHandler;
+}
+
+TSharedPtr<IDetailPropertyExtensionHandler> SDetailsViewBase::GetExtensionHandler()
+{
+	return ExtensionHandler;
+}
+
 void SDetailsViewBase::SetGenericLayoutDetailsDelegate(FOnGetDetailCustomizationInstance OnGetGenericDetails)
 {
 	GenericLayoutDelegate = OnGetGenericDetails;
diff --git a/Engine/Source/Editor/PropertyEditor/Private/SDetailsViewBase.h b/Engine/Source/Editor/PropertyEditor/Private/SDetailsViewBase.h
--- a/Engine/Source/Editor/PropertyEditor/Private/SDetailsViewBase.h
+++ b/Engine/Source/Editor/PropertyEditor/Private/SDetailsViewBase.h
@@ -143,6 +143,9 @@ public:
 	virtual void SetKeyframeHandler( TSharedPtr<class IDetailKeyframeHandler> InKeyframeHandler ) override;
 	virtual TSharedPtr<IDetailKeyframeHandler> GetKeyframeHandler() override;
 
+	virtual void SetExtensionHandler(TSharedPtr<class IDetailPropertyExtensionHandler> InExtensionHandler);
+	virtual TSharedPtr<IDetailPropertyExtensionHandler> GetExtensionHandler();
+
 	/**
 	 * Requests that an item in the tree be expanded or collapsed
 	 *
@@ -424,8 +427,12 @@ protected:
 	/** Root tree node that needs to be destroyed when safe */
 	TSharedPtr<FComplexPropertyNode> RootNodePendingKill;
 
+	/** The handler for the keyframe UI, determines if the key framing UI should be displayed. */
 	TSharedPtr<IDetailKeyframeHandler> KeyframeHandler;
 
+	/** Property extension handler returns additional UI to apply after the customization is applied to the property. */
+	TSharedPtr<IDetailPropertyExtensionHandler> ExtensionHandler;
+
 	/** External property nodes which need to validated each tick */
 	TArray< TWeakPtr<FObjectPropertyNode> > ExternalRootPropertyNodes;
 
diff --git a/Engine/Source/Editor/PropertyEditor/Public/IDetailPropertyExtensionHandler.h b/Engine/Source/Editor/PropertyEditor/Public/IDetailPropertyExtensionHandler.h
new file mode 100644
--- /dev/null
+++ b/Engine/Source/Editor/PropertyEditor/Public/IDetailPropertyExtensionHandler.h
@@ -0,0 +1,13 @@
+// Copyright 1998-2014 Epic Games, Inc. All Rights Reserved.
+
+#pragma once
+
+class IDetailPropertyExtensionHandler
+{
+public:
+	virtual ~IDetailPropertyExtensionHandler(){ }
+
+	virtual bool IsPropertyExtenable(const UClass* InObjectClass, const class IPropertyHandle& PropertyHandle) const = 0;
+
+	virtual TSharedRef<SWidget> GenerateExtensionWidget(const UClass* InObjectClass, TSharedPtr<IPropertyHandle> PropertyHandle) = 0;
+};
diff --git a/Engine/Source/Editor/PropertyEditor/Public/IDetailsView.h b/Engine/Source/Editor/PropertyEditor/Public/IDetailsView.h
--- a/Engine/Source/Editor/PropertyEditor/Public/IDetailsView.h
+++ b/Engine/Source/Editor/PropertyEditor/Public/IDetailsView.h
@@ -152,6 +152,8 @@ public:
 
 	virtual void SetKeyframeHandler( TSharedPtr<class IDetailKeyframeHandler> InKeyframeHandler ) = 0;
 
+	virtual void SetExtensionHandler(TSharedPtr<class IDetailPropertyExtensionHandler> InExtensionandler) = 0;
+
 	/**
 	 * @return true if property editing is enabled (based on the FIsPropertyEditingEnabled delegate)
 	 */ 
diff --git a/Engine/Source/Editor/UMGEditor/Private/Customizations/DetailCustomizations.cpp b/Engine/Source/Editor/UMGEditor/Private/Customizations/DetailCustomizations.cpp
--- a/Engine/Source/Editor/UMGEditor/Private/Customizations/DetailCustomizations.cpp
+++ b/Engine/Source/Editor/UMGEditor/Private/Customizations/DetailCustomizations.cpp
@@ -365,7 +365,7 @@ void FBlueprintWidgetCustomization::PerformBindingCustomization(IDetailLayoutBui
 			{
 				if ( DelegateProperty->GetName().EndsWith(TEXT("Delegate")) )
 				{
-					CreateDelegateCustomization(DetailLayout, DelegateProperty, Widget);
+					//CreateDelegateCustomization(DetailLayout, DelegateProperty, Widget);
 				}
 				else if ( DelegateProperty->GetName().EndsWith(TEXT("Event")) )
 				{
diff --git a/Engine/Source/Editor/UMGEditor/Private/Details/DetailWidgetExtensionHandler.cpp b/Engine/Source/Editor/UMGEditor/Private/Details/DetailWidgetExtensionHandler.cpp
new file mode 100644
--- /dev/null
+++ b/Engine/Source/Editor/UMGEditor/Private/Details/DetailWidgetExtensionHandler.cpp
@@ -0,0 +1,49 @@
+// Copyright 1998-2014 Epic Games, Inc. All Rights Reserved.
+
+#include "UMGEditorPrivatePCH.h"
+#include "DetailWidgetExtensionHandler.h"
+#include "PropertyHandle.h"
+#include "SPropertyBinding.h"
+
+FDetailWidgetExtensionHandler::FDetailWidgetExtensionHandler(TSharedPtr<FWidgetBlueprintEditor> InBlueprintEditor)
+	: BlueprintEditor( InBlueprintEditor )
+{}
+
+bool FDetailWidgetExtensionHandler::IsPropertyExtenable(const UClass* InObjectClass, const IPropertyHandle& InPropertyHandle) const
+{
+	// TODO UMG make this work for multiple widgets.
+	if ( InPropertyHandle.GetNumOuterObjects() == 1 )
+	{
+		UProperty* Property = InPropertyHandle.GetProperty();
+		FString DelegateName = Property->GetName() + "Delegate";
+
+		if ( UClass* ContainerClass = Cast<UClass>(Property->GetOuter()) )
+		{
+			UDelegateProperty* DelegateProperty = FindField<UDelegateProperty>(ContainerClass, FName(*DelegateName));
+			if ( DelegateProperty )
+			{
+				return true;
+			}
+		}
+	}
+
+	return false;
+}
+
+TSharedRef<SWidget> FDetailWidgetExtensionHandler::GenerateExtensionWidget(const UClass* InObjectClass, TSharedPtr<IPropertyHandle> InPropertyHandle)
+{
+	UProperty* Property = InPropertyHandle->GetProperty();
+	FString DelegateName = Property->GetName() + "Delegate";
+	
+	UDelegateProperty* DelegateProperty = FindFieldChecked<UDelegateProperty>(CastChecked<UClass>(Property->GetOuter()), FName(*DelegateName));
+
+	const bool bIsEditable = Property->HasAnyPropertyFlags(CPF_Edit | CPF_EditConst);
+	const bool bDoSignaturesMatch = DelegateProperty->SignatureFunction->GetReturnProperty()->SameType(Property);
+
+	if ( !ensure(bIsEditable && bDoSignaturesMatch) )
+	{
+		return SNullWidget::NullWidget;
+	}
+
+	return SNew(SPropertyBinding, BlueprintEditor.Pin().ToSharedRef(), DelegateProperty, InPropertyHandle.ToSharedRef());
+}
diff --git a/Engine/Source/Editor/UMGEditor/Private/Details/DetailWidgetExtensionHandler.h b/Engine/Source/Editor/UMGEditor/Private/Details/DetailWidgetExtensionHandler.h
new file mode 100644
--- /dev/null
+++ b/Engine/Source/Editor/UMGEditor/Private/Details/DetailWidgetExtensionHandler.h
@@ -0,0 +1,18 @@
+// Copyright 1998-2014 Epic Games, Inc. All Rights Reserved.
+
+#pragma once
+
+#include "Editor/PropertyEditor/Public/IDetailPropertyExtensionHandler.h"
+
+class FDetailWidgetExtensionHandler : public IDetailPropertyExtensionHandler
+{
+public:
+	FDetailWidgetExtensionHandler(TSharedPtr<class FWidgetBlueprintEditor> InBlueprintEditor);
+
+	virtual bool IsPropertyExtenable(const UClass* InObjectClass, const class IPropertyHandle& PropertyHandle) const override;
+
+	virtual TSharedRef<SWidget> GenerateExtensionWidget(const UClass* InObjectClass, TSharedPtr<IPropertyHandle> PropertyHandle) override;
+
+private:
+	TWeakPtr<class FWidgetBlueprintEditor> BlueprintEditor;
+};
\ No newline at end of file
diff --git a/Engine/Source/Editor/UMGEditor/Private/Details/SPropertyBinding.cpp b/Engine/Source/Editor/UMGEditor/Private/Details/SPropertyBinding.cpp
new file mode 100644
--- /dev/null
+++ b/Engine/Source/Editor/UMGEditor/Private/Details/SPropertyBinding.cpp
@@ -0,0 +1,524 @@
+// Copyright 1998-2014 Epic Games, Inc. All Rights Reserved.
+
+#include "UMGEditorPrivatePCH.h"
+
+#include "SPropertyBinding.h"
+
+#define LOCTEXT_NAMESPACE "UMG"
+
+/////////////////////////////////////////////////////
+// SPropertyBinding
+
+void SPropertyBinding::Construct(const FArguments& InArgs, TSharedRef<FWidgetBlueprintEditor> InEditor, UDelegateProperty* DelegateProperty, TSharedRef<IPropertyHandle> Property)
+{
+	Editor = InEditor;
+	Blueprint = InEditor->GetWidgetBlueprintObj();
+
+	TArray<UObject*> Objects;
+	Property->GetOuterObjects(Objects);
+
+	UWidget* Widget = CastChecked<UWidget>(Objects[0]);
+
+	ChildSlot
+	[
+		SNew(SHorizontalBox)
+
+		+ SHorizontalBox::Slot()
+		.AutoWidth()
+		[
+			SNew(SComboButton)
+			.OnGetMenuContent(this, &SPropertyBinding::OnGenerateDelegateMenu, Widget, Property, DelegateProperty->SignatureFunction, InArgs._GeneratePureBindings)
+			.ContentPadding(1)
+			.ButtonContent()
+			[
+				SNew(SHorizontalBox)
+				+ SHorizontalBox::Slot()
+				.AutoWidth()
+				.VAlign(VAlign_Center)
+				[
+					SNew(SImage)
+					.Image(this, &SPropertyBinding::GetCurrentBindingImage, Property)
+					.ColorAndOpacity(FLinearColor(0.25f, 0.25f, 0.25f))
+				]
+				+ SHorizontalBox::Slot()
+				.AutoWidth()
+				.VAlign(VAlign_Center)
+				.Padding(4, 1, 0, 0)
+				[
+					SNew(STextBlock)
+					.Text(this, &SPropertyBinding::GetCurrentBindingText, Property)
+					.Font(IDetailLayoutBuilder::GetDetailFont())
+				]
+			]
+		]
+
+		+ SHorizontalBox::Slot()
+		.AutoWidth()
+		[
+			SNew(SButton)
+			.ButtonStyle(FEditorStyle::Get(), "HoverHintOnly")
+			.Visibility(this, &SPropertyBinding::GetGotoBindingVisibility, Property)
+			.OnClicked(this, &SPropertyBinding::HandleGotoBindingClicked, Property)
+			.VAlign(VAlign_Center)
+			.ToolTipText(LOCTEXT("GotoFunction", "Goto Function"))
+			[
+				SNew(SImage)
+				.Image(FEditorStyle::GetBrush("PropertyWindow.Button_Browse"))
+			]
+		]
+	];
+}
+
+void SPropertyBinding::RefreshBlueprintMemberCache(const UFunction* DelegateSignature)
+{
+	const UWidgetGraphSchema* Schema = GetDefault<UWidgetGraphSchema>();
+	const FSlateFontInfo DetailFontInfo = IDetailLayoutBuilder::GetDetailFont();
+
+	BlueprintFunctionCache.Reset();
+	BlueprintPropertyCache.Reset();
+
+	// Get the current skeleton class, think header for the blueprint.
+	UBlueprintGeneratedClass* SkeletonClass = Cast<UBlueprintGeneratedClass>(Blueprint->SkeletonGeneratedClass);
+
+	// Grab functions implemented by the blueprint
+	for ( int32 i = 0; i < Blueprint->FunctionGraphs.Num(); i++ )
+	{
+		UEdGraph* Graph = Blueprint->FunctionGraphs[i];
+		check(Graph);
+
+		// Ignore the construction script, we just want user created functions.
+		if ( Graph->GetFName() == Schema->FN_UserConstructionScript )
+		{
+			continue;
+		}
+
+		FName FunctionFName = Graph->GetFName();
+
+		UFunction* Function = SkeletonClass->FindFunctionByName(FunctionFName, EIncludeSuperFlag::IncludeSuper);
+		if ( Function == NULL )
+		{
+			continue;
+		}
+
+		// We ignore CPF_ReturnParm because all that matters for binding to script functions is that the number of out parameters match.
+		if ( Function->IsSignatureCompatibleWith(DelegateSignature, UFunction::GetDefaultIgnoredSignatureCompatibilityFlags() | CPF_ReturnParm) )
+		{
+			FGraphDisplayInfo DisplayInfo;
+			Graph->GetSchema()->GetGraphDisplayInformation(*Graph, DisplayInfo);
+
+			TSharedPtr<FunctionInfo> NewFuncAction = MakeShareable(new FunctionInfo());
+			NewFuncAction->DisplayName = DisplayInfo.PlainName;
+			NewFuncAction->Tooltip = DisplayInfo.Tooltip;
+			NewFuncAction->FuncName = FunctionFName;
+			NewFuncAction->EdGraph = Graph;
+
+			BlueprintFunctionCache.Add(NewFuncAction);
+		}
+	}
+
+	// Grab functions implemented by the blueprint
+	for ( TFieldIterator<UProperty> PropIt(SkeletonClass, EFieldIteratorFlags::ExcludeSuper); PropIt; ++PropIt )
+	{
+		UProperty* Prop = *PropIt;
+
+		if ( UProperty* ReturnProperty = DelegateSignature->GetReturnProperty() )
+		{
+			if ( ReturnProperty->SameType(Prop) )
+			{
+				if ( Prop->HasAnyPropertyFlags(UP::BlueprintReadWrite) )
+				{
+					BlueprintPropertyCache.Add(Prop);
+				}
+			}
+		}
+	}
+}
+
+TSharedRef<SWidget> SPropertyBinding::OnGenerateDelegateMenu(UWidget* Widget, TSharedRef<IPropertyHandle> PropertyHandle, UFunction* DelegateSignature, bool bIsPure)
+{
+	RefreshBlueprintMemberCache(DelegateSignature);
+
+	const bool bInShouldCloseWindowAfterMenuSelection = true;
+	FMenuBuilder MenuBuilder(bInShouldCloseWindowAfterMenuSelection, NULL);
+
+	static FName PropertyIcon(TEXT("Kismet.Tabs.Variables"));
+	static FName FunctionIcon(TEXT("GraphEditor.Function_16x"));
+
+	//FLinearColor ColorOut;
+	//const UClass* VarClass = PropertyHandle->GetPropertyClass();
+	//BrushOut = FBlueprintEditor::GetVarIconAndColor(VarClass, PropertyHandle->GetProperty()->GetFName(), ColorOut);
+
+	if ( CanRemoveBinding(PropertyHandle) )
+	{
+		MenuBuilder.BeginSection("RemoveBinding");
+		{
+			MenuBuilder.AddMenuEntry(
+				LOCTEXT("RemoveBinding", "Remove Binding"),
+				LOCTEXT("RemoveBindingToolTip", "Removes the current binding"),
+				FSlateIcon(),
+				FUIAction(FExecuteAction::CreateSP(this, &SPropertyBinding::HandleRemoveBinding, PropertyHandle))
+				);
+		}
+		MenuBuilder.EndSection(); //RemoveBinding
+	}
+
+	MenuBuilder.BeginSection("CreateBinding");
+	{
+		MenuBuilder.AddMenuEntry(
+			LOCTEXT("CreateBinding", "Create Binding"),
+			LOCTEXT("CreateBindingToolTip", "Creates a new function for this property to be bound to"),
+			FSlateIcon(),
+			FUIAction(FExecuteAction::CreateSP(this, &SPropertyBinding::HandleCreateAndAddBinding, Widget, PropertyHandle, DelegateSignature, bIsPure))
+			);
+	}
+	MenuBuilder.EndSection(); //CreateBinding
+
+	MenuBuilder.BeginSection("Functions", LOCTEXT("Functions", "Functions"));
+	{
+		for ( TSharedPtr<FunctionInfo>& Function : BlueprintFunctionCache )
+		{
+			MenuBuilder.AddMenuEntry(
+				Function->DisplayName,
+				FText::FromString(Function->Tooltip),
+				FSlateIcon(FEditorStyle::GetStyleSetName(), FunctionIcon),
+				FUIAction(FExecuteAction::CreateSP(this, &SPropertyBinding::HandleAddFunctionBinding, PropertyHandle, Function))
+				);
+		}
+	}
+	MenuBuilder.EndSection(); //Functions
+
+	MenuBuilder.BeginSection("Properties", LOCTEXT("Properties", "Properties"));
+	{
+		for ( UProperty* ExistingProperty : BlueprintPropertyCache )
+		{
+			MenuBuilder.AddMenuEntry(
+				ExistingProperty->GetDisplayNameText(),
+				ExistingProperty->GetToolTipText(),
+				FSlateIcon(FEditorStyle::GetStyleSetName(), PropertyIcon),
+				FUIAction(FExecuteAction::CreateSP(this, &SPropertyBinding::HandleAddPropertyBinding, PropertyHandle, ExistingProperty))
+				);
+		}
+	}
+	MenuBuilder.EndSection(); //Functions
+
+
+	FDisplayMetrics DisplayMetrics;
+	FSlateApplication::Get().GetDisplayMetrics(DisplayMetrics);
+
+	return
+		SNew(SVerticalBox)
+
+		+ SVerticalBox::Slot()
+		.MaxHeight(DisplayMetrics.PrimaryDisplayHeight * 0.5)
+		[
+			MenuBuilder.MakeWidget()
+		];
+}
+
+const FSlateBrush* SPropertyBinding::GetCurrentBindingImage(TSharedRef<IPropertyHandle> PropertyHandle) const
+{
+	static FName PropertyIcon(TEXT("Kismet.Tabs.Variables"));
+	static FName FunctionIcon(TEXT("GraphEditor.Function_16x"));
+
+	TArray<UObject*> OuterObjects;
+	PropertyHandle->GetOuterObjects(OuterObjects);
+
+	//TODO UMG O(N) Isn't good for this, needs to be map, but map isn't serialized, need cached runtime map for fast lookups.
+
+	FName PropertyName = PropertyHandle->GetProperty()->GetFName();
+	for ( int32 ObjectIndex = 0; ObjectIndex < OuterObjects.Num(); ObjectIndex++ )
+	{
+		// Ignore null outer objects
+		if ( OuterObjects[ObjectIndex] == NULL )
+		{
+			continue;
+		}
+
+		//TODO UMG handle multiple things selected
+
+		for ( const FDelegateEditorBinding& Binding : Blueprint->Bindings )
+		{
+			if ( Binding.ObjectName == OuterObjects[ObjectIndex]->GetName() && Binding.PropertyName == PropertyName )
+			{
+				if ( Binding.Kind == EBindingKind::Function )
+				{
+					return FEditorStyle::GetBrush(FunctionIcon);
+				}
+				else // Property
+				{
+					return FEditorStyle::GetBrush(PropertyIcon);
+				}
+			}
+		}
+	}
+
+	return NULL;
+}
+
+FText SPropertyBinding::GetCurrentBindingText(TSharedRef<IPropertyHandle> PropertyHandle) const
+{
+	TArray<UObject*> OuterObjects;
+	PropertyHandle->GetOuterObjects(OuterObjects);
+
+	//TODO UMG O(N) Isn't good for this, needs to be map, but map isn't serialized, need cached runtime map for fast lookups.
+
+	FName PropertyName = PropertyHandle->GetProperty()->GetFName();
+	for ( int32 ObjectIndex = 0; ObjectIndex < OuterObjects.Num(); ObjectIndex++ )
+	{
+		// Ignore null outer objects
+		if ( OuterObjects[ObjectIndex] == NULL )
+		{
+			continue;
+		}
+
+		//TODO UMG handle multiple things selected
+
+		for ( const FDelegateEditorBinding& Binding : Blueprint->Bindings )
+		{
+			if ( Binding.ObjectName == OuterObjects[ObjectIndex]->GetName() && Binding.PropertyName == PropertyName )
+			{
+				if ( Binding.Kind == EBindingKind::Function )
+				{
+					FName FoundName = Blueprint->GetFieldNameFromClassByGuid<UFunction>(Blueprint->GeneratedClass, Binding.MemberGuid);
+					return FText::FromString(FName::NameToDisplayString(FoundName.ToString(), false));
+				}
+				else // Property
+				{
+					FName FoundName = Blueprint->GetFieldNameFromClassByGuid<UProperty>(Blueprint->GeneratedClass, Binding.MemberGuid);
+					return FText::FromString(FName::NameToDisplayString(FoundName.ToString(), false));
+				}
+			}
+		}
+
+		//TODO UMG Do something about missing functions, little exclamation points if they're missing and such.
+
+		break;
+	}
+
+	return LOCTEXT("Bind", "Bind");
+}
+
+bool SPropertyBinding::CanRemoveBinding(TSharedRef<IPropertyHandle> PropertyHandle)
+{
+	FName PropertyName = PropertyHandle->GetProperty()->GetFName();
+
+	TArray<UObject*> OuterObjects;
+	PropertyHandle->GetOuterObjects(OuterObjects);
+	for ( UObject* SelectedObject : OuterObjects )
+	{
+		for ( const FDelegateEditorBinding& Binding : Blueprint->Bindings )
+		{
+			if ( Binding.ObjectName == SelectedObject->GetName() && Binding.PropertyName == PropertyName )
+			{
+				return true;
+			}
+		}
+	}
+
+	return false;
+}
+
+void SPropertyBinding::HandleRemoveBinding(TSharedRef<IPropertyHandle> PropertyHandle)
+{
+	const FScopedTransaction Transaction(LOCTEXT("BindDelegate", "Remove Binding"));
+
+	Blueprint->Modify();
+
+	TArray<UObject*> OuterObjects;
+	PropertyHandle->GetOuterObjects(OuterObjects);
+	for ( UObject* SelectedObject : OuterObjects )
+	{
+		FDelegateEditorBinding Binding;
+		Binding.ObjectName = SelectedObject->GetName();
+		Binding.PropertyName = PropertyHandle->GetProperty()->GetFName();
+
+		Blueprint->Bindings.Remove(Binding);
+	}
+
+	FBlueprintEditorUtils::MarkBlueprintAsStructurallyModified(Blueprint);
+}
+
+void SPropertyBinding::HandleAddFunctionBinding(TSharedRef<IPropertyHandle> PropertyHandle, TSharedPtr<FunctionInfo> SelectedFunction)
+{
+	const FScopedTransaction Transaction(LOCTEXT("BindDelegate", "Set Binding"));
+
+	Blueprint->Modify();
+
+	TArray<UObject*> OuterObjects;
+	PropertyHandle->GetOuterObjects(OuterObjects);
+	for ( UObject* SelectedObject : OuterObjects )
+	{
+		FDelegateEditorBinding Binding;
+		Binding.ObjectName = SelectedObject->GetName();
+		Binding.PropertyName = PropertyHandle->GetProperty()->GetFName();
+		Binding.FunctionName = SelectedFunction->FuncName;
+		Binding.MemberGuid = SelectedFunction->EdGraph->GraphGuid;
+		Binding.Kind = EBindingKind::Function;
+
+		Blueprint->Bindings.Remove(Binding);
+		Blueprint->Bindings.AddUnique(Binding);
+	}
+
+	FBlueprintEditorUtils::MarkBlueprintAsStructurallyModified(Blueprint);
+}
+
+void SPropertyBinding::HandleAddPropertyBinding(TSharedRef<IPropertyHandle> PropertyHandle, UProperty* SelectedProperty)
+{
+	const FScopedTransaction Transaction(LOCTEXT("BindDelegate", "Set Binding"));
+
+	// Get the current skeleton class, think header for the blueprint.
+	UBlueprintGeneratedClass* SkeletonClass = Cast<UBlueprintGeneratedClass>(Blueprint->SkeletonGeneratedClass);
+
+	Blueprint->Modify();
+
+	FGuid MemberGuid;
+	UBlueprint::GetGuidFromClassByFieldName<UProperty>(SkeletonClass, SelectedProperty->GetFName(), MemberGuid);
+
+	TArray<UObject*> OuterObjects;
+	PropertyHandle->GetOuterObjects(OuterObjects);
+	for ( UObject* SelectedObject : OuterObjects )
+	{
+		FDelegateEditorBinding Binding;
+		Binding.ObjectName = SelectedObject->GetName();
+		Binding.PropertyName = PropertyHandle->GetProperty()->GetFName();
+		Binding.SourceProperty = SelectedProperty->GetFName();
+		Binding.MemberGuid = MemberGuid;
+		Binding.Kind = EBindingKind::Property;
+
+		Blueprint->Bindings.Remove(Binding);
+		Blueprint->Bindings.AddUnique(Binding);
+	}
+
+	FBlueprintEditorUtils::MarkBlueprintAsStructurallyModified(Blueprint);
+}
+
+void SPropertyBinding::HandleCreateAndAddBinding(UWidget* Widget, TSharedRef<IPropertyHandle> PropertyHandle, UFunction* DelegateSignature, bool bIsPure)
+{
+	const FScopedTransaction Transaction(LOCTEXT("CreateDelegate", "Create Binding"));
+
+	Blueprint->Modify();
+
+	FString Pre = bIsPure ? FString(TEXT("Get")) : FString(TEXT("On"));
+
+	FString WidgetName;
+	if ( Widget && !Widget->IsGeneratedName() )
+	{
+		WidgetName = TEXT("_") + Widget->GetName() + TEXT("_");
+	}
+
+	FString Post = PropertyHandle->GetProperty()->GetName();
+	Post.RemoveFromStart(TEXT("On"));
+	Post.RemoveFromEnd(TEXT("Event"));
+
+	// Create the function graph.
+	FString FunctionName = Pre + WidgetName + Post;
+	UEdGraph* FunctionGraph = FBlueprintEditorUtils::CreateNewGraph(Blueprint, FBlueprintEditorUtils::FindUniqueKismetName(Blueprint, FunctionName), UEdGraph::StaticClass(), UEdGraphSchema_K2::StaticClass());
+
+	// Only mark bindings as pure that need to be.
+	if ( bIsPure )
+	{
+		const UEdGraphSchema_K2* Schema_K2 = Cast<UEdGraphSchema_K2>(FunctionGraph->GetSchema());
+		Schema_K2->AddExtraFunctionFlags(FunctionGraph, FUNC_BlueprintPure);
+	}
+
+	// Add the binding to the blueprint
+	TSharedPtr<FunctionInfo> SelectedFunction = MakeShareable(new FunctionInfo());
+	SelectedFunction->FuncName = FunctionGraph->GetFName();
+	SelectedFunction->EdGraph = FunctionGraph;
+
+	HandleAddFunctionBinding(PropertyHandle, SelectedFunction);
+
+	bool bUserCreated = true;
+	FBlueprintEditorUtils::AddFunctionGraph(Blueprint, FunctionGraph, bUserCreated, DelegateSignature);
+
+	GotoFunction(FunctionGraph);
+}
+
+EVisibility SPropertyBinding::GetGotoBindingVisibility(TSharedRef<IPropertyHandle> PropertyHandle) const
+{
+	TArray<UObject*> OuterObjects;
+	PropertyHandle->GetOuterObjects(OuterObjects);
+
+	//TODO UMG O(N) Isn't good for this, needs to be map, but map isn't serialized, need cached runtime map for fast lookups.
+
+	FName PropertyName = PropertyHandle->GetProperty()->GetFName();
+	for ( int32 ObjectIndex = 0; ObjectIndex < OuterObjects.Num(); ObjectIndex++ )
+	{
+		// Ignore null outer objects
+		if ( OuterObjects[ObjectIndex] == NULL )
+		{
+			continue;
+		}
+
+		//TODO UMG handle multiple things selected
+
+		for ( const FDelegateEditorBinding& Binding : Blueprint->Bindings )
+		{
+			if ( Binding.ObjectName == OuterObjects[ObjectIndex]->GetName() && Binding.PropertyName == PropertyName )
+			{
+				if ( Binding.Kind == EBindingKind::Function )
+				{
+					return EVisibility::Visible;
+				}
+			}
+		}
+	}
+
+	return EVisibility::Collapsed;
+}
+
+FReply SPropertyBinding::HandleGotoBindingClicked(TSharedRef<IPropertyHandle> PropertyHandle)
+{
+	TArray<UObject*> OuterObjects;
+	PropertyHandle->GetOuterObjects(OuterObjects);
+
+	//TODO UMG O(N) Isn't good for this, needs to be map, but map isn't serialized, need cached runtime map for fast lookups.
+
+	FName PropertyName = PropertyHandle->GetProperty()->GetFName();
+	for ( int32 ObjectIndex = 0; ObjectIndex < OuterObjects.Num(); ObjectIndex++ )
+	{
+		// Ignore null outer objects
+		if ( OuterObjects[ObjectIndex] == NULL )
+		{
+			continue;
+		}
+
+		//TODO UMG handle multiple things selected
+
+		for ( const FDelegateEditorBinding& Binding : Blueprint->Bindings )
+		{
+			if ( Binding.ObjectName == OuterObjects[ObjectIndex]->GetName() && Binding.PropertyName == PropertyName )
+			{
+				if ( Binding.Kind == EBindingKind::Function )
+				{
+					TArray<UEdGraph*> AllGraphs;
+					Blueprint->GetAllGraphs(AllGraphs);
+
+					for ( UEdGraph* Graph : AllGraphs )
+					{
+						if ( Graph->GraphGuid == Binding.MemberGuid )
+						{
+							GotoFunction(Graph);
+						}
+					}
+
+
+					// Either way return
+					return FReply::Handled();
+				}
+			}
+		}
+	}
+
+	return FReply::Unhandled();
+}
+
+void SPropertyBinding::GotoFunction(UEdGraph* FunctionGraph)
+{
+	Editor.Pin()->SetCurrentMode(FWidgetBlueprintApplicationModes::GraphMode);
+
+	Editor.Pin()->OpenDocument(FunctionGraph, FDocumentTracker::OpenNewDocument);
+}
+
+
+#undef LOCTEXT_NAMESPACE
diff --git a/Engine/Source/Editor/UMGEditor/Private/Details/SPropertyBinding.h b/Engine/Source/Editor/UMGEditor/Private/Details/SPropertyBinding.h
new file mode 100644
--- /dev/null
+++ b/Engine/Source/Editor/UMGEditor/Private/Details/SPropertyBinding.h
@@ -0,0 +1,59 @@
+// Copyright 1998-2014 Epic Games, Inc. All Rights Reserved.
+
+#pragma once
+
+#include "SCompoundWidget.h"
+
+class SPropertyBinding : public SCompoundWidget
+{
+public:
+
+	SLATE_BEGIN_ARGS(SPropertyBinding)
+		: _GeneratePureBindings(true)
+		{}
+		SLATE_ARGUMENT(bool, GeneratePureBindings)
+	SLATE_END_ARGS()
+
+	void Construct(const FArguments& InArgs, TSharedRef<FWidgetBlueprintEditor> InEditor, UDelegateProperty* DelegateProperty, TSharedRef<IPropertyHandle> Property);
+
+protected:
+	struct FunctionInfo
+	{
+		FText DisplayName;
+		FString Tooltip;
+
+		FName FuncName;
+		UEdGraph* EdGraph;
+	};
+
+	void RefreshBlueprintMemberCache(const UFunction* DelegateSignature);
+
+	TSharedRef<SWidget> OnGenerateDelegateMenu(UWidget* Widget, TSharedRef<IPropertyHandle> PropertyHandle, UFunction* DelegateSignature, bool bIsPure);
+
+	const FSlateBrush* GetCurrentBindingImage(TSharedRef<IPropertyHandle> PropertyHandle) const;
+	FText GetCurrentBindingText(TSharedRef<IPropertyHandle> PropertyHandle) const;
+
+	bool CanRemoveBinding(TSharedRef<IPropertyHandle> PropertyHandle);
+	void HandleRemoveBinding(TSharedRef<IPropertyHandle> PropertyHandle);
+
+	void HandleAddFunctionBinding(TSharedRef<IPropertyHandle> PropertyHandle, TSharedPtr<FunctionInfo> SelectedFunction);
+	void HandleAddPropertyBinding(TSharedRef<IPropertyHandle> PropertyHandle, UProperty* SelectedProperty);
+
+	void HandleCreateAndAddBinding(UWidget* Widget, TSharedRef<IPropertyHandle> PropertyHandle, UFunction* DelegateSignature, bool bIsPure);
+	void GotoFunction(UEdGraph* FunctionGraph);
+
+	EVisibility GetGotoBindingVisibility(TSharedRef<IPropertyHandle> PropertyHandle) const;
+
+	FReply HandleGotoBindingClicked(TSharedRef<IPropertyHandle> PropertyHandle);
+
+	FReply AddOrViewEventBinding(TSharedPtr<FEdGraphSchemaAction> Action);
+
+private:
+
+	TWeakPtr<FWidgetBlueprintEditor> Editor;
+
+	UWidgetBlueprint* Blueprint;
+
+	TArray< TSharedPtr<FunctionInfo> > BlueprintFunctionCache;
+	TArray< UProperty* > BlueprintPropertyCache;
+};
diff --git a/Engine/Source/Editor/UMGEditor/Private/Details/SWidgetDetailsView.cpp b/Engine/Source/Editor/UMGEditor/Private/Details/SWidgetDetailsView.cpp
--- a/Engine/Source/Editor/UMGEditor/Private/Details/SWidgetDetailsView.cpp
+++ b/Engine/Source/Editor/UMGEditor/Private/Details/SWidgetDetailsView.cpp
@@ -8,6 +8,7 @@
 #include "Kismet2NameValidators.h"
 #include "ISequencer.h"
 #include "Animation/UMGDetailKeyframeHandler.h"
+#include "DetailWidgetExtensionHandler.h"
 #include "EditorClassUtils.h"
 
 #include "Customizations/SlateBrushCustomization.h"
@@ -39,6 +40,10 @@ void SWidgetDetailsView::Construct(const FArguments& InArgs, TSharedPtr<FWidgetB
 	TSharedRef<IDetailKeyframeHandler> KeyframeHandler = MakeShareable( new FUMGDetailKeyframeHandler( InBlueprintEditor ) );
 	PropertyView->SetKeyframeHandler( KeyframeHandler );
 
+	// Create a handler for property binding via the details panel
+	TSharedRef<FDetailWidgetExtensionHandler> BindingHandler = MakeShareable( new FDetailWidgetExtensionHandler( InBlueprintEditor ) );
+	PropertyView->SetExtensionHandler(BindingHandler);
+
 	ChildSlot
 	[
 		SNew(SVerticalBox)
diff --git a/Engine/Source/Runtime/UMG/Private/Components/Widget.cpp b/Engine/Source/Runtime/UMG/Private/Components/Widget.cpp
--- a/Engine/Source/Runtime/UMG/Private/Components/Widget.cpp
+++ b/Engine/Source/Runtime/UMG/Private/Components/Widget.cpp
@@ -244,7 +244,7 @@ TSharedRef<SWidget> UWidget::TakeWidget()
 		// Otherwise we need to recreate the wrapper widget
 		else
 		{
-			TSharedPtr<SWidget> SafeGCWidget = SNew(SObjectWidget, Cast<UUserWidget>(this))
+			TSharedPtr<SObjectWidget> SafeGCWidget = SNew(SObjectWidget, Cast<UUserWidget>(this))
 				[
 					SafeWidget.ToSharedRef()
 				];
diff --git a/Engine/Source/Runtime/UMG/Private/UserWidget.cpp b/Engine/Source/Runtime/UMG/Private/UserWidget.cpp
--- a/Engine/Source/Runtime/UMG/Private/UserWidget.cpp
+++ b/Engine/Source/Runtime/UMG/Private/UserWidget.cpp
@@ -21,9 +21,9 @@ UUserWidget::UUserWidget(const FObjectInitializer& ObjectInitializer)
 	Visiblity = ESlateVisibility::SelfHitTestInvisible;
 
 	bInitialized = false;
-	CachedWorld = NULL;
-
 	bSupportsKeyboardFocus = true;
+	ColorAndOpacity = FLinearColor::White;
+	ForegroundColor = FSlateColor::UseForeground();
 
 #if WITH_EDITORONLY_DATA
 	bUseDesignTimeSize = false;
@@ -78,6 +78,12 @@ void UUserWidget::ReleaseSlateResources(bool bReleaseChildren)
 {
 	Super::ReleaseSlateResources(bReleaseChildren);
 
+	if ( !HasAnyFlags(RF_ClassDefaultObject) )
+	{
+		// Notify the blueprint that it's been removed from the slate hierarchy.
+		//Destruct();
+	}
+
 	UWidget* RootWidget = GetRootWidgetComponent();
 	if ( RootWidget )
 	{
@@ -85,6 +91,43 @@ void UUserWidget::ReleaseSlateResources(bool bReleaseChildren)
 	}
 }
 
+void UUserWidget::SynchronizeProperties()
+{
+	Super::SynchronizeProperties();
+
+	TSharedPtr<SObjectWidget> SafeGCWidget = MyGCWidget.Pin();
+	if ( SafeGCWidget.IsValid() )
+	{
+		TAttribute<FLinearColor> ColorBinding = OPTIONAL_BINDING(FLinearColor, ColorAndOpacity);
+		TAttribute<FSlateColor> ForegroundColorBinding = OPTIONAL_BINDING(FSlateColor, ForegroundColor);
+
+		SafeGCWidget->SetColorAndOpacity(ColorBinding);
+		SafeGCWidget->SetForegroundColor(ForegroundColorBinding);
+	}
+}
+
+void UUserWidget::SetColorAndOpacity(FLinearColor InColorAndOpacity)
+{
+	ColorAndOpacity = InColorAndOpacity;
+
+	TSharedPtr<SObjectWidget> SafeGCWidget = MyGCWidget.Pin();
+	if ( SafeGCWidget.IsValid() )
+	{
+		SafeGCWidget->SetColorAndOpacity(ColorAndOpacity);
+	}
+}
+
+void UUserWidget::SetForegroundColor(FSlateColor InForegroundColor)
+{
+	ForegroundColor = InForegroundColor;
+
+	TSharedPtr<SObjectWidget> SafeGCWidget = MyGCWidget.Pin();
+	if ( SafeGCWidget.IsValid() )
+	{
+		SafeGCWidget->SetForegroundColor(ForegroundColor);
+	}
+}
+
 void UUserWidget::PostInitProperties()
 {
 	Super::PostInitProperties();
@@ -131,6 +174,11 @@ void UUserWidget::Construct_Implementation()
 
 }
 
+void UUserWidget::Destruct_Implementation()
+{
+
+}
+
 void UUserWidget::Tick_Implementation(FGeometry MyGeometry, float InDeltaTime)
 {
 
diff --git a/Engine/Source/Runtime/UMG/Public/Blueprint/UserWidget.h b/Engine/Source/Runtime/UMG/Public/Blueprint/UserWidget.h
--- a/Engine/Source/Runtime/UMG/Public/Blueprint/UserWidget.h
+++ b/Engine/Source/Runtime/UMG/Public/Blueprint/UserWidget.h
@@ -112,6 +112,7 @@ UCLASS(Abstract, editinlinenew, BlueprintType, Blueprintable, meta=( Category="U
 class UMG_API UUserWidget : public UWidget, public INamedSlotInterface
 {
 	GENERATED_UCLASS_BODY()
+
 public:
 	//UObject interface
 	virtual void PostInitProperties() override;
@@ -126,6 +127,10 @@ public:
 	virtual void ReleaseSlateResources(bool bReleaseChildren) override;
 	// End of UVisual interface
 
+	// UWidget interface
+	virtual void SynchronizeProperties() override;
+	// End of UWidget interface
+
 	// UNamedSlotInterface Begin
 	virtual void GetSlotNames(TArray<FName>& SlotNames) const override;
 	virtual UWidget* GetContentForSlot(FName SlotName) const override;
@@ -198,10 +203,20 @@ public:
 	UFUNCTION(BlueprintCallable, Category="Player")
 	class APawn* GetOwningPlayerPawn() const;
 
-	/** Called when the widget is constructed */
+	/**
+	 * Called after the underlying slate widget is constructed.  Depending on how the slate object is used
+	 * this event may be called multiple times due to adding and removing from the hierarchy.
+	 */
 	UFUNCTION(BlueprintNativeEvent, Category="User Interface", meta=(Keywords="Begin Play"))
 	void Construct();
 
+	/**
+	 * Whenever the underlying slate widgets are torn down.  This is called to allow you to do any last minute
+	 * destruction work.
+	 */
+	UFUNCTION(BlueprintNativeEvent, Category="User Interface", meta=(Keywords="End Play"))
+	void Destruct();
+
 	UFUNCTION(BlueprintNativeEvent, Category="User Interface")
 	void Tick(FGeometry MyGeometry, float InDeltaTime);
 
@@ -281,6 +296,14 @@ public:
 
 	//virtual bool OnVisualizeTooltip(const TSharedPtr<SWidget>& TooltipContent);
 
+	/**  */
+	UFUNCTION(BlueprintCallable, Category="Appearance")
+	void SetColorAndOpacity(FLinearColor InColorAndOpacity);
+
+	/**  */
+	UFUNCTION(BlueprintCallable, Category="Appearance")
+	void SetForegroundColor(FSlateColor InForegroundColor);
+
 	/**
 	 * Plays an animation in this widget
 	 * 
@@ -334,6 +357,21 @@ public:
 #endif
 
 public:
+	/**  */
+	UPROPERTY(EditDefaultsOnly, BlueprintReadOnly, Category="Style")
+	FLinearColor ColorAndOpacity;
+
+	/**  */
+	UPROPERTY()
+	FGetLinearColor ColorAndOpacityDelegate;
+
+	/**  */
+	UPROPERTY(EditDefaultsOnly, BlueprintReadOnly, Category="Style")
+	FSlateColor ForegroundColor;
+
+	/**  */
+	UPROPERTY()
+	FGetSlateColor ForegroundColorDelegate;
 
 	/** Called when the visibility changes. */
 	UPROPERTY(EditAnywhere, BlueprintReadWrite, AdvancedDisplay, Category="Behavior")
@@ -391,9 +429,8 @@ private:
 
 	FLocalPlayerContext PlayerContext;
 
+	/** Has this widget been initialized by its class yet? */
 	bool bInitialized;
-
-	mutable UWorld* CachedWorld;
 };
 
 template< class T >
diff --git a/Engine/Source/Runtime/UMG/Public/Components/Image.h b/Engine/Source/Runtime/UMG/Public/Components/Image.h
--- a/Engine/Source/Runtime/UMG/Public/Components/Image.h
+++ b/Engine/Source/Runtime/UMG/Public/Components/Image.h
@@ -34,7 +34,7 @@ public:
 
 	/** A bindable delegate for the ColorAndOpacity. */
 	UPROPERTY()
-	FGetSlateColor ColorAndOpacityDelegate;
+	FGetLinearColor ColorAndOpacityDelegate;
 
 	UPROPERTY(EditDefaultsOnly, Category=Events)
 	FOnPointerEvent OnMouseButtonDownEvent;
diff --git a/Engine/Source/Runtime/UMG/Public/Components/Widget.h b/Engine/Source/Runtime/UMG/Public/Components/Widget.h
--- a/Engine/Source/Runtime/UMG/Public/Components/Widget.h
+++ b/Engine/Source/Runtime/UMG/Public/Components/Widget.h
@@ -39,7 +39,7 @@ public:
 	DECLARE_DYNAMIC_DELEGATE_RetVal(FVector, FGetVector);
 	//DECLARE_DYNAMIC_DELEGATE_RetVal(FVector4, FGetVector4);
 	DECLARE_DYNAMIC_DELEGATE_RetVal(FMargin, FGetMargin);
-	DECLARE_DYNAMIC_DELEGATE_RetVal(FLinearColor, FGetSlateColor);
+	DECLARE_DYNAMIC_DELEGATE_RetVal(FSlateColor, FGetSlateColor);
 	DECLARE_DYNAMIC_DELEGATE_RetVal(FLinearColor, FGetLinearColor);
 	DECLARE_DYNAMIC_DELEGATE_RetVal(ESlateVisibility::Type, FGetSlateVisibility);
 	DECLARE_DYNAMIC_DELEGATE_RetVal(EMouseCursor::Type, FGetMouseCursor);
@@ -341,7 +341,7 @@ protected:
 	TWeakPtr<SWidget> MyWidget;
 
 	/** The underlying SWidget contained in a SObjectWidget */
-	TWeakPtr<SWidget> MyGCWidget;
+	TWeakPtr<class SObjectWidget> MyGCWidget;
 	
 	/** Is this widget being displayed on a designer surface */
 	UPROPERTY(Transient)
