Fix variable shadowing

diff --git a/Engine/Source/Editor/BlueprintGraph/Public/BlueprintNodeSpawner.h b/Engine/Source/Editor/BlueprintGraph/Public/BlueprintNodeSpawner.h
--- a/Engine/Source/Editor/BlueprintGraph/Public/BlueprintNodeSpawner.h
+++ b/Engine/Source/Editor/BlueprintGraph/Public/BlueprintNodeSpawner.h
@@ -232,9 +232,9 @@ UBlueprintNodeSpawner* UBlueprintNodeSpawner::Create(UObject* Outer, FCustomizeN
 
 //------------------------------------------------------------------------------
 template <class NodeType>
-NodeType* UBlueprintNodeSpawner::SpawnNode(TSubclassOf<UEdGraphNode> NodeClass, UEdGraph* ParentGraph, FBindingSet const& Bindings, FVector2D const Location, FCustomizeNodeDelegate PostSpawnDelegate) const
+NodeType* UBlueprintNodeSpawner::SpawnNode(TSubclassOf<UEdGraphNode> InNodeClass, UEdGraph* ParentGraph, FBindingSet const& Bindings, FVector2D const Location, FCustomizeNodeDelegate PostSpawnDelegate) const
 {
-	return CastChecked<NodeType>(SpawnEdGraphNode(NodeClass, ParentGraph, Bindings, Location, PostSpawnDelegate));
+	return CastChecked<NodeType>(SpawnEdGraphNode(InNodeClass, ParentGraph, Bindings, Location, PostSpawnDelegate));
 }
 
 //------------------------------------------------------------------------------
diff --git a/Engine/Source/Editor/Kismet/Private/SMyBlueprint.cpp b/Engine/Source/Editor/Kismet/Private/SMyBlueprint.cpp
--- a/Engine/Source/Editor/Kismet/Private/SMyBlueprint.cpp
+++ b/Engine/Source/Editor/Kismet/Private/SMyBlueprint.cpp
@@ -703,13 +703,13 @@ void SMyBlueprint::CollectAllActions(FGraphActionListBuilderBase& OutAllActions)
 				SortedActions.Add( SortIndex, Action );
 			}
 
-			void GetAllActions( FGraphActionListBuilderBase& OutAllActions )
+			void GetAllActions( FGraphActionListBuilderBase& OutActions )
 			{
 				SortedActions.KeySort( TLess<int32>() );
 
 				for( auto Iter : SortedActions )
 				{
-					OutAllActions.AddAction( Iter.Value );
+					OutActions.AddAction( Iter.Value );
 				}
 			}
 
@@ -1830,15 +1830,15 @@ void SMyBlueprint::OnDeleteEntry()
 				}
 				else if (Actions[i]->GetTypeId() == FEdGraphSchemaAction_K2LocalVar::StaticGetTypeId())
 				{
-					FEdGraphSchemaAction_K2LocalVar* LocalVarAction = (FEdGraphSchemaAction_K2LocalVar*)Actions[i].Get();
+					FEdGraphSchemaAction_K2LocalVar* K2LocalVarAction = (FEdGraphSchemaAction_K2LocalVar*)Actions[i].Get();
 
-					FBlueprintEditorUtils::RemoveLocalVariable(GetBlueprintObj(), LocalVarAction->GetVariableScope(), LocalVarAction->GetVariableName());
+					FBlueprintEditorUtils::RemoveLocalVariable(GetBlueprintObj(), K2LocalVarAction->GetVariableScope(), K2LocalVarAction->GetVariableName());
 				}
 				else if (Actions[i]->GetTypeId() == FEdGraphSchemaAction_K2Graph::StaticGetTypeId())
 				{
-					FEdGraphSchemaAction_K2Graph* GraphAction = (FEdGraphSchemaAction_K2Graph*)Actions[i].Get();
+					FEdGraphSchemaAction_K2Graph* K2GraphAction = (FEdGraphSchemaAction_K2Graph*)Actions[i].Get();
 
-					OnDeleteGraph(GraphAction->EdGraph, GraphAction->GraphType);
+					OnDeleteGraph(K2GraphAction->EdGraph, K2GraphAction->GraphType);
 				}
 				else if (Actions[i]->GetTypeId() == FEdGraphSchemaAction_K2Delegate::StaticGetTypeId())
 				{
diff --git a/Engine/Source/Editor/Persona/Private/Persona.cpp b/Engine/Source/Editor/Persona/Private/Persona.cpp
--- a/Engine/Source/Editor/Persona/Private/Persona.cpp
+++ b/Engine/Source/Editor/Persona/Private/Persona.cpp
@@ -503,11 +503,11 @@ void FPersona::ExtendMenu()
 	UAnimBlueprint* AnimBlueprint = GetAnimBlueprint();
 	if(AnimBlueprint)
 	{
-		TSharedPtr<FExtender> MenuExtender = MakeShareable(new FExtender);
-		FKismet2Menu::SetupBlueprintEditorMenu(MenuExtender, *this);
-		AddMenuExtender(MenuExtender);
+		TSharedPtr<FExtender> AnimBPMenuExtender = MakeShareable(new FExtender);
+		FKismet2Menu::SetupBlueprintEditorMenu(AnimBPMenuExtender, *this);
+		AddMenuExtender(AnimBPMenuExtender);
 
-		MenuExtender->AddMenuExtension(
+		AnimBPMenuExtender->AddMenuExtension(
 			"FileLoadAndSave",
 			EExtensionHook::After,
 			GetToolkitCommands(),
@@ -2811,10 +2811,10 @@ void FPersona::OnImportAsset(enum EFBXImportType DefaultImportType)
 void FPersona::OnReimportMesh()
 {
 	// Reimport the asset
-	UDebugSkelMeshComponent* PreviewComponent = GetPreviewMeshComponent();
-	if(PreviewComponent && PreviewComponent->SkeletalMesh)
+	UDebugSkelMeshComponent* PreviewMeshComponent = GetPreviewMeshComponent();
+	if(PreviewMeshComponent && PreviewMeshComponent->SkeletalMesh)
 	{
-		FReimportManager::Instance()->Reimport(PreviewComponent->SkeletalMesh, true);
+		FReimportManager::Instance()->Reimport(PreviewMeshComponent->SkeletalMesh, true);
 	}
 }
 
diff --git a/Engine/Source/Editor/Persona/Private/SRigWindow.h b/Engine/Source/Editor/Persona/Private/SRigWindow.h
--- a/Engine/Source/Editor/Persona/Private/SRigWindow.h
+++ b/Engine/Source/Editor/Persona/Private/SRigWindow.h
@@ -32,9 +32,9 @@ public:
 
 protected:
 	/** Hidden constructor, always use Make above */
-	FDisplayedBoneMappingInfo(const FName InNodeName, const FString DisplayName, USkeleton* InSkeleton)
+	FDisplayedBoneMappingInfo(const FName InNodeName, const FString InDisplayName, USkeleton* InSkeleton)
 		: Name( InNodeName )
-		, DisplayName( DisplayName )
+		, DisplayName( InDisplayName )
 		, Skeleton( InSkeleton )
 	{}
 
diff --git a/Engine/Source/Editor/UMGEditor/Private/Customizations/CanvasSlotCustomization.cpp b/Engine/Source/Editor/UMGEditor/Private/Customizations/CanvasSlotCustomization.cpp
--- a/Engine/Source/Editor/UMGEditor/Private/Customizations/CanvasSlotCustomization.cpp
+++ b/Engine/Source/Editor/UMGEditor/Private/Customizations/CanvasSlotCustomization.cpp
@@ -127,12 +127,14 @@ private:
 	{
 		FScopedTransaction Transaction(LOCTEXT("ChangeAnchors", "Changed Anchors"));
 
-		FString Value = FString::Printf(TEXT("(Minimum=(X=%f,Y=%f),Maximum=(X=%f,Y=%f))"), Anchors.Minimum.X, Anchors.Minimum.Y, Anchors.Maximum.X, Anchors.Maximum.Y);
-		AnchorsHandle->SetValueFromFormattedString(Value);
+		{
+			const FString Value = FString::Printf(TEXT("(Minimum=(X=%f,Y=%f),Maximum=(X=%f,Y=%f))"), Anchors.Minimum.X, Anchors.Minimum.Y, Anchors.Maximum.X, Anchors.Maximum.Y);
+			AnchorsHandle->SetValueFromFormattedString(Value);
+		}
 
 		if ( FSlateApplication::Get().GetModifierKeys().IsShiftDown() )
 		{
-			FString Value = FString::Printf(TEXT("(X=%f,Y=%f)"), Anchors.IsStretchedHorizontal() ? 0 : Anchors.Minimum.X, Anchors.IsStretchedVertical() ? 0 : Anchors.Minimum.Y);
+			const FString Value = FString::Printf(TEXT("(X=%f,Y=%f)"), Anchors.IsStretchedHorizontal() ? 0 : Anchors.Minimum.X, Anchors.IsStretchedVertical() ? 0 : Anchors.Minimum.Y);
 			AlignmentHandle->SetValueFromFormattedString(Value);
 		}
 
@@ -142,7 +144,7 @@ private:
 			OffsetsHandle->AccessRawData(RawOffsetData);
 			FMargin* Offsets = reinterpret_cast<FMargin*>( RawOffsetData[0] );
 
-			FString Value = FString::Printf(TEXT("(Left=%f,Top=%f,Right=%f,Bottom=%f)"), 0, 0, Anchors.IsStretchedHorizontal() ? 0 : Offsets->Right, Anchors.IsStretchedVertical() ? 0 : Offsets->Bottom);
+			const FString Value = FString::Printf(TEXT("(Left=%f,Top=%f,Right=%f,Bottom=%f)"), 0, 0, Anchors.IsStretchedHorizontal() ? 0 : Offsets->Right, Anchors.IsStretchedVertical() ? 0 : Offsets->Bottom);
 			OffsetsHandle->SetValueFromFormattedString(Value);
 		}
 
diff --git a/Engine/Source/Editor/UMGEditor/Private/Customizations/DetailCustomizations.cpp b/Engine/Source/Editor/UMGEditor/Private/Customizations/DetailCustomizations.cpp
--- a/Engine/Source/Editor/UMGEditor/Private/Customizations/DetailCustomizations.cpp
+++ b/Engine/Source/Editor/UMGEditor/Private/Customizations/DetailCustomizations.cpp
@@ -372,9 +372,9 @@ void FBlueprintWidgetCustomization::PerformBindingCustomization(IDetailLayoutBui
 					CreateEventCustomization(DetailLayout, DelegateProperty, Widget);
 				}
 			}
-			else if ( UMulticastDelegateProperty* DelegateProperty = Cast<UMulticastDelegateProperty>(Property) )
+			else if ( UMulticastDelegateProperty* MulticastDelegateProperty = Cast<UMulticastDelegateProperty>(Property) )
 			{
-				CreateMulticastEventCustomization(DetailLayout, OutObjects[0].Get()->GetFName(), PropertyClass, DelegateProperty);
+				CreateMulticastEventCustomization(DetailLayout, OutObjects[0].Get()->GetFName(), PropertyClass, MulticastDelegateProperty);
 			}
 		}
 	}
diff --git a/Engine/Source/Editor/UMGEditor/Private/Designer/SDesignerView.cpp b/Engine/Source/Editor/UMGEditor/Private/Designer/SDesignerView.cpp
--- a/Engine/Source/Editor/UMGEditor/Private/Designer/SDesignerView.cpp
+++ b/Engine/Source/Editor/UMGEditor/Private/Designer/SDesignerView.cpp
@@ -588,9 +588,9 @@ void SDesignerView::MarkDesignModifed(bool bRequiresRecompile)
 
 bool SDesignerView::GetWidgetParentGeometry(const FWidgetReference& Widget, FGeometry& Geometry) const
 {
-	if ( UWidget* PreviewWidget = Widget.GetPreview() )
+	if ( UWidget* WidgetPreview = Widget.GetPreview() )
 	{
-		if ( UPanelWidget* Parent = PreviewWidget->GetParent() )
+		if ( UPanelWidget* Parent = WidgetPreview->GetParent() )
 		{
 			FWidgetReference ParentReference = BlueprintEditor.Pin()->GetReferenceFromPreview(Parent);
 			return GetWidgetGeometry(ParentReference, Geometry);
@@ -603,9 +603,9 @@ bool SDesignerView::GetWidgetParentGeometry(const FWidgetReference& Widget, FGeo
 
 bool SDesignerView::GetWidgetGeometry(const FWidgetReference& Widget, FGeometry& Geometry) const
 {
-	if ( UWidget* PreviewWidget = Widget.GetPreview() )
+	if ( UWidget* WidgetPreview = Widget.GetPreview() )
 	{
-		TSharedPtr<SWidget> CachedPreviewWidget = PreviewWidget->GetCachedWidget();
+		TSharedPtr<SWidget> CachedPreviewWidget = WidgetPreview->GetCachedWidget();
 		if ( CachedPreviewWidget.IsValid() )
 		{
 			const FArrangedWidget* ArrangedWidget = CachedWidgetGeometry.Find(CachedPreviewWidget.ToSharedRef());
@@ -685,14 +685,14 @@ FVector2D SDesignerView::GetExtensionPosition(TSharedRef<FDesignerSurfaceElement
 		FWidgetReference ParentRef = BlueprintEditor.Pin()->GetReferenceFromTemplate(SelectedWidget.GetTemplate()->GetParent());
 
 		UWidget* Preview = ParentRef.GetPreview();
-		TSharedPtr<SWidget> PreviewSlateWidget = Preview ? Preview->GetCachedWidget() : NULL;
-		if ( PreviewSlateWidget.IsValid() )
+		TSharedPtr<SWidget> CachedPreviewSlateWidget = Preview ? Preview->GetCachedWidget() : NULL;
+		if ( CachedPreviewSlateWidget.IsValid() )
 		{
 			FWidgetPath WidgetPath;
 			SelectedWidgetPath.ToWidgetPath(WidgetPath);
 		
 			FArrangedWidget ArrangedWidget(SNullWidget::NullWidget, FGeometry());
-			FDesignTimeUtils::GetArrangedWidgetRelativeToParent(WidgetPath, PreviewSlateWidget.ToSharedRef(), AsShared(), ArrangedWidget);
+			FDesignTimeUtils::GetArrangedWidgetRelativeToParent(WidgetPath, CachedPreviewSlateWidget.ToSharedRef(), AsShared(), ArrangedWidget);
 
 			ParentPosition = ArrangedWidget.Geometry.AbsolutePosition;
 			ParentSize = ArrangedWidget.Geometry.Size * GetPreviewScale();
@@ -933,19 +933,19 @@ FReply SDesignerView::OnMouseMove(const FGeometry& MyGeometry, const FPointerEve
 
 				BeginTransaction(LOCTEXT("MoveWidgetRT", "Move Widget (Render Transform)"));
 
-				if ( UWidget* PreviewWidget = SelectedWidget.GetPreview() )
+				if ( UWidget* WidgetPreview = SelectedWidget.GetPreview() )
 				{
 					FGeometry ParentGeometry;
 					if ( GetWidgetParentGeometry(SelectedWidget, ParentGeometry) )
 					{
 						const FSlateRenderTransform& AbsoluteToLocalTransform = Inverse(ParentGeometry.GetAccumulatedRenderTransform());
 
-						FWidgetTransform RenderTransform = PreviewWidget->RenderTransform;
+						FWidgetTransform RenderTransform = WidgetPreview->RenderTransform;
 						RenderTransform.Translation += AbsoluteToLocalTransform.TransformVector(MouseEvent.GetCursorDelta());
 
 						static const FName RenderTransformName(TEXT("RenderTransform"));
 
-						FObjectEditorUtils::SetPropertyValue<UWidget, FWidgetTransform>(PreviewWidget, RenderTransformName, RenderTransform);
+						FObjectEditorUtils::SetPropertyValue<UWidget, FWidgetTransform>(WidgetPreview, RenderTransformName, RenderTransform);
 						FObjectEditorUtils::SetPropertyValue<UWidget, FWidgetTransform>(SelectedWidget.GetTemplate(), RenderTransformName, RenderTransform);
 					}
 				}
diff --git a/Engine/Source/Editor/UMGEditor/Private/Extensions/CanvasSlotExtension.cpp b/Engine/Source/Editor/UMGEditor/Private/Extensions/CanvasSlotExtension.cpp
--- a/Engine/Source/Editor/UMGEditor/Private/Extensions/CanvasSlotExtension.cpp
+++ b/Engine/Source/Editor/UMGEditor/Private/Extensions/CanvasSlotExtension.cpp
@@ -291,8 +291,8 @@ FReply FCanvasSlotExtension::HandleAnchorDragging(const FGeometry& Geometry, con
 			{
 				UCanvasPanelSlot* PreviewCanvasSlot = Cast<UCanvasPanelSlot>(PreviewWidget->Slot);
 
-				FGeometry Geometry;
-				if ( Canvas->GetGeometryForSlot(PreviewCanvasSlot, Geometry) )
+				FGeometry GeometryForSlot;
+				if ( Canvas->GetGeometryForSlot(PreviewCanvasSlot, GeometryForSlot) )
 				{
 					FGeometry CanvasGeometry = Canvas->GetCanvasWidget()->GetCachedGeometry();
 					FVector2D StartLocalPosition = CanvasGeometry.AbsoluteToLocal(MouseDownPosition);
diff --git a/Engine/Source/Editor/UMGEditor/Private/Hierarchy/SHierarchyViewItem.cpp b/Engine/Source/Editor/UMGEditor/Private/Hierarchy/SHierarchyViewItem.cpp
--- a/Engine/Source/Editor/UMGEditor/Private/Hierarchy/SHierarchyViewItem.cpp
+++ b/Engine/Source/Editor/UMGEditor/Private/Hierarchy/SHierarchyViewItem.cpp
@@ -162,10 +162,10 @@ FReply ProcessHierarchyDragDrop(const FDragDropEvent& DragDropEvent, bool bIsDro
 					}
 					else
 					{
-						TSharedPtr<SWidget> Widget = TargetItem.GetPreview()->GetCachedWidget();
-						if ( Widget.IsValid() )
+						TSharedPtr<SWidget> CachedWidget = TargetItem.GetPreview()->GetCachedWidget();
+						if ( CachedWidget.IsValid() )
 						{
-							Widget->SlatePrepass();
+							CachedWidget->SlatePrepass();
 						}
 					}
 				}
diff --git a/Engine/Source/Editor/UMGEditor/Private/TreeFilterHandler.h b/Engine/Source/Editor/UMGEditor/Private/TreeFilterHandler.h
--- a/Engine/Source/Editor/UMGEditor/Private/TreeFilterHandler.h
+++ b/Engine/Source/Editor/UMGEditor/Private/TreeFilterHandler.h
@@ -22,15 +22,15 @@ public:
 	}
 
 	/** Sets the TreeView to be filtered. */
-	void SetTreeView(STreeView<ItemType>* TreeView)
+	void SetTreeView(STreeView<ItemType>* InTreeView)
 	{
-		this->TreeView = TreeView;
+		TreeView = InTreeView;
 	}
 
 	/** Sets the filter which will be used to filter the items in the TreeView. */
-	void SetFilter(IFilter<ItemType>* Filter)
+	void SetFilter(IFilter<ItemType>* InFilter)
 	{
-		this->Filter = Filter;
+		Filter = InFilter;
 	}
 
 	/** 
@@ -38,19 +38,19 @@ public:
 	 * @param RootItems		An array of source root items to be displayed by the tree.
 	 * @param TreeRootItems	The array of root items which are set as the ItemsSource of the tree. This array will be modified to reflect any filtering.
 	 */
-	void SetRootItems(TArray<ItemType>* RootItems, TArray<ItemType>* TreeRootItems)
+	void SetRootItems(TArray<ItemType>* InRootItems, TArray<ItemType>* InTreeRootItems)
 	{
-		this->RootItems = RootItems;
-		this->TreeRootItems = TreeRootItems;
+		RootItems = InRootItems;
+		TreeRootItems = InTreeRootItems;
 	}
 
 	/**
 	 * Sets the delegate which should be used for traversing the children of the nodes in
 	 * the tree.
 	 */
-	void SetGetChildrenDelegate(FOnGetChildren GetChildrenDelegate)
+	void SetGetChildrenDelegate(FOnGetChildren InGetChildrenDelegate)
 	{
-		this->GetChildrenDelegate = GetChildrenDelegate;
+		GetChildrenDelegate = InGetChildrenDelegate;
 	}
 
 	/** Gets whether or not filtering is enabled for the TreeView. */
@@ -63,9 +63,9 @@ public:
 	 * Sets whether or not filtering the tree is enabled. Calling this DOES NOT refresh the tree,
 	 * this must be done by calling RefreshAndFilterTree.
 	 */
-	void SetIsEnabled(bool bIsEnabled)
+	void SetIsEnabled(bool bInIsEnabled)
 	{
-		this->bIsEnabled = bIsEnabled;
+		bIsEnabled = bInIsEnabled;
 	}
 
 	/** Removes any cached references to an item. */
diff --git a/Engine/Source/Editor/UMGEditor/Private/WidgetBlueprintCompiler.cpp b/Engine/Source/Editor/UMGEditor/Private/WidgetBlueprintCompiler.cpp
--- a/Engine/Source/Editor/UMGEditor/Private/WidgetBlueprintCompiler.cpp
+++ b/Engine/Source/Editor/UMGEditor/Private/WidgetBlueprintCompiler.cpp
@@ -29,7 +29,7 @@ void FWidgetBlueprintCompiler::CreateFunctionList()
 
 	for ( FDelegateEditorBinding& EditorBinding : WidgetBlueprint()->Bindings )
 	{
-		FName PropertyName = EditorBinding.SourceProperty;
+		const FName PropertyName = EditorBinding.SourceProperty;
 
 		UProperty* Property = FindField<UProperty>(Blueprint->SkeletonGeneratedClass, PropertyName);
 		if ( Property )
@@ -69,8 +69,6 @@ void FWidgetBlueprintCompiler::CreateFunctionList()
 			ReturnNode->NodePosY = EntryNode->NodePosY;
 			FunctionReturnCreator.Finalize();
 
-			UProperty* Property = FindField<UProperty>(Blueprint->SkeletonGeneratedClass, PropertyName);
-
 			FEdGraphPinType PinType;
 			K2Schema->ConvertPropertyToPinType(Property, /*out*/ PinType);
 
diff --git a/Engine/Source/Editor/UMGEditor/Private/WidgetReference.cpp b/Engine/Source/Editor/UMGEditor/Private/WidgetReference.cpp
--- a/Engine/Source/Editor/UMGEditor/Private/WidgetReference.cpp
+++ b/Engine/Source/Editor/UMGEditor/Private/WidgetReference.cpp
@@ -4,8 +4,8 @@
 
 #define LOCTEXT_NAMESPACE "UMG"
 
-FWidgetHandle::FWidgetHandle(UWidget* Widget)
-	: Widget(Widget)
+FWidgetHandle::FWidgetHandle(UWidget* InWidget)
+	: Widget(InWidget)
 {
 
 }
@@ -19,9 +19,9 @@ FWidgetReference::FWidgetReference()
 
 }
 
-FWidgetReference::FWidgetReference(TSharedPtr<FWidgetBlueprintEditor> WidgetEditor, TSharedPtr< FWidgetHandle > TemplateHandle)
-	: WidgetEditor(WidgetEditor)
-	, TemplateHandle(TemplateHandle)
+FWidgetReference::FWidgetReference(TSharedPtr<FWidgetBlueprintEditor> InWidgetEditor, TSharedPtr< FWidgetHandle > InTemplateHandle)
+	: WidgetEditor(InWidgetEditor)
+	, TemplateHandle(InTemplateHandle)
 {
 	
 }
diff --git a/Engine/Source/Runtime/Core/Public/Containers/Array.h b/Engine/Source/Runtime/Core/Public/Containers/Array.h
--- a/Engine/Source/Runtime/Core/Public/Containers/Array.h
+++ b/Engine/Source/Runtime/Core/Public/Containers/Array.h
@@ -3001,32 +3001,32 @@ public:
 	 * Serialization operator for TIndirectArray.
 	 *
 	 * @param	Ar		Archive to serialize with.
-	 * @param	Array	Array to serialize.
+	 * @param	A		Array to serialize.
 	 *
 	 * @returns Passing down serializing archive.
 	 */
-	friend FArchive& operator<<(FArchive& Ar, TIndirectArray& Array)
+	friend FArchive& operator<<(FArchive& Ar, TIndirectArray& A)
 	{
-		Array.CountBytes(Ar);
+		A.CountBytes(Ar);
 		if (Ar.IsLoading())
 		{
 			// Load array.
 			int32 NewNum;
 			Ar << NewNum;
-			Array.Empty(NewNum);
+			A.Empty(NewNum);
 			for (int32 Index = 0; Index < NewNum; Index++)
 			{
-				Ar << *new(Array)T;
+				Ar << *new(A)T;
 			}
 		}
 		else
 		{
 			// Save array.
-			int32 Num = Array.Num();
+			int32 Num = A.Num();
 			Ar << Num;
 			for (int32 Index = 0; Index < Num; Index++)
 			{
-				Ar << Array[Index];
+				Ar << A[Index];
 			}
 		}
 		return Ar;
diff --git a/Engine/Source/Runtime/Core/Public/GenericPlatform/GenericApplication.h b/Engine/Source/Runtime/Core/Public/GenericPlatform/GenericApplication.h
--- a/Engine/Source/Runtime/Core/Public/GenericPlatform/GenericApplication.h
+++ b/Engine/Source/Runtime/Core/Public/GenericPlatform/GenericApplication.h
@@ -67,7 +67,7 @@ public:
 						const bool bInIsRightAltDown,
 						const bool bInIsLeftCommandDown,
 						const bool bInIsRightCommandDown,
-						const bool bAreCapsLocked)
+						const bool bInAreCapsLocked)
 		: bIsLeftShiftDown( bInIsLeftShiftDown ),
 		  bIsRightShiftDown( bInIsRightShiftDown ),
 		  bIsLeftControlDown( bInIsLeftControlDown ),
@@ -76,7 +76,7 @@ public:
 		  bIsRightAltDown( bInIsRightAltDown ),
 		  bIsLeftCommandDown( bInIsLeftCommandDown ),
 		  bIsRightCommandDown( bInIsRightCommandDown ),
-		  bAreCapsLocked( bAreCapsLocked )
+		  bAreCapsLocked( bInAreCapsLocked )
 	{
 	}
 
diff --git a/Engine/Source/Runtime/Core/Public/Misc/AutomationTest.h b/Engine/Source/Runtime/Core/Public/Misc/AutomationTest.h
--- a/Engine/Source/Runtime/Core/Public/Misc/AutomationTest.h
+++ b/Engine/Source/Runtime/Core/Public/Misc/AutomationTest.h
@@ -520,9 +520,9 @@ public:
 	/**
 	 * Sets forcing smoke tests.
 	 */
-	void SetForceSmokeTests(bool bForceSmokeTests)
+	void SetForceSmokeTests(const bool bInForceSmokeTests)
 	{
-		this->bForceSmokeTests = bForceSmokeTests;
+		bForceSmokeTests = bInForceSmokeTests;
 	}
 private:
 
diff --git a/Engine/Source/Runtime/Core/Public/Templates/Sorting.h b/Engine/Source/Runtime/Core/Public/Templates/Sorting.h
--- a/Engine/Source/Runtime/Core/Public/Templates/Sorting.h
+++ b/Engine/Source/Runtime/Core/Public/Templates/Sorting.h
@@ -449,11 +449,11 @@ public:
 			}
 			else
 			{
-				for (int32 SubgroupStart = 0; SubgroupStart < Num; SubgroupStart += 2)
+				for (int32 Subgroup = 0; Subgroup < Num; Subgroup += 2)
 				{
-					if (SubgroupStart + 1 < Num && Predicate(First[SubgroupStart + 1], First[SubgroupStart]))
+					if (Subgroup + 1 < Num && Predicate(First[Subgroup + 1], First[Subgroup]))
 					{
-						Exchange(First[SubgroupStart], First[SubgroupStart + 1]);
+						Exchange(First[Subgroup], First[Subgroup + 1]);
 					}
 				}
 			}
diff --git a/Engine/Source/Runtime/Engine/Classes/GameFramework/PlayerController.h b/Engine/Source/Runtime/Engine/Classes/GameFramework/PlayerController.h
--- a/Engine/Source/Runtime/Engine/Classes/GameFramework/PlayerController.h
+++ b/Engine/Source/Runtime/Engine/Classes/GameFramework/PlayerController.h
@@ -524,7 +524,7 @@ public:
 	 * @param bSpeaking enable or disable voice chat
 	 */
 	UFUNCTION(exec)
-	virtual void ToggleSpeaking(bool bSpeaking);
+	virtual void ToggleSpeaking(bool bInSpeaking);
 
 	/**
 	 * Tells the client that the server has all the information it needs and that it
diff --git a/Engine/Source/Runtime/Engine/Public/VisualLogger/VisualLogger.inl b/Engine/Source/Runtime/Engine/Public/VisualLogger/VisualLogger.inl
--- a/Engine/Source/Runtime/Engine/Public/VisualLogger/VisualLogger.inl
+++ b/Engine/Source/Runtime/Engine/Public/VisualLogger/VisualLogger.inl
@@ -51,20 +51,20 @@ FVisualLogEntry* FVisualLogger::GetEntryToWrite(const class UObject* Object, flo
 	{
 		CurrentEntry->Reset();
 		CurrentEntry->TimeStamp = TimeStamp;
-		const class AActor* AsActor = Cast<class AActor>(Object);
-		if (AsActor)
+		const class AActor* ObjectAsActor = Cast<class AActor>(Object);
+		if (ObjectAsActor)
 		{
-			CurrentEntry->Location = AsActor->GetActorLocation();
-			AsActor->GrabDebugSnapshot(CurrentEntry);
+			CurrentEntry->Location = ObjectAsActor->GetActorLocation();
+			ObjectAsActor->GrabDebugSnapshot(CurrentEntry);
 		}
 		if (RedirectionMap.Contains(Object))
 		{
 			for (auto Child : RedirectionMap[Object])
 			{
-				const class AActor* AsActor = Cast<class AActor>(Child);
-				if (AsActor)
+				const class AActor* ChildAsActor = Cast<class AActor>(Child);
+				if (ChildAsActor)
 				{
-					AsActor->GrabDebugSnapshot(CurrentEntry);
+					ChildAsActor->GrabDebugSnapshot(CurrentEntry);
 				}
 			}
 		}
diff --git a/Engine/Source/Runtime/Slate/Public/Widgets/Views/SListView.h b/Engine/Source/Runtime/Slate/Public/Widgets/Views/SListView.h
--- a/Engine/Source/Runtime/Slate/Public/Widgets/Views/SListView.h
+++ b/Engine/Source/Runtime/Slate/Public/Widgets/Views/SListView.h
@@ -399,8 +399,8 @@ private:
 	class FWidgetGenerator
 	{
 	public:
-		FWidgetGenerator(SListView<ItemType>* OwnerList)
-			: OwnerList(OwnerList)
+		FWidgetGenerator(SListView<ItemType>* InOwnerList)
+			: OwnerList(InOwnerList)
 		{
 		}
 
diff --git a/Engine/Source/Runtime/UMG/Public/Blueprint/UserWidget.h b/Engine/Source/Runtime/UMG/Public/Blueprint/UserWidget.h
--- a/Engine/Source/Runtime/UMG/Public/Blueprint/UserWidget.h
+++ b/Engine/Source/Runtime/UMG/Public/Blueprint/UserWidget.h
@@ -36,19 +36,19 @@ public:
 		, MyClippingRect(NullRect)
 		, OutDrawElements(NullElementList)
 		, LayerId(0)
-		, InWidgetStyle(NullStyle)
+		, WidgetStyle(NullStyle)
 		, bParentEnabled(true)
 		, MaxLayer(0)
 	{ }
 
-	FPaintContext(const FGeometry& AllottedGeometry, const FSlateRect& MyClippingRect, FSlateWindowElementList& OutDrawElements, int32 LayerId, const FWidgetStyle& InWidgetStyle, bool bParentEnabled)
-		: AllottedGeometry(AllottedGeometry)
-		, MyClippingRect(MyClippingRect)
-		, OutDrawElements(OutDrawElements)
-		, LayerId(LayerId)
-		, InWidgetStyle(InWidgetStyle)
-		, bParentEnabled(bParentEnabled)
-		, MaxLayer(LayerId)
+	FPaintContext(const FGeometry& InAllottedGeometry, const FSlateRect& InMyClippingRect, FSlateWindowElementList& InOutDrawElements, const int32 InLayerId, const FWidgetStyle& InWidgetStyle, const bool bInParentEnabled)
+		: AllottedGeometry(InAllottedGeometry)
+		, MyClippingRect(InMyClippingRect)
+		, OutDrawElements(InOutDrawElements)
+		, LayerId(InLayerId)
+		, WidgetStyle(InWidgetStyle)
+		, bParentEnabled(bInParentEnabled)
+		, MaxLayer(InLayerId)
 	{
 	}
 
@@ -59,7 +59,7 @@ public:
 		const_cast<FSlateRect&>( MyClippingRect ) = Other.MyClippingRect;
 		OutDrawElements = Other.OutDrawElements;
 		LayerId = Other.LayerId;
-		const_cast<FWidgetStyle&>( InWidgetStyle ) = Other.InWidgetStyle;
+		const_cast<FWidgetStyle&>( WidgetStyle ) = Other.WidgetStyle;
 		bParentEnabled = Other.bParentEnabled;
 	}
 
@@ -69,7 +69,7 @@ public:
 	const FSlateRect& MyClippingRect;
 	FSlateWindowElementList& OutDrawElements;
 	int32 LayerId;
-	const FWidgetStyle& InWidgetStyle;
+	const FWidgetStyle& WidgetStyle;
 	bool bParentEnabled;
 
 	int32 MaxLayer;
