Editor - Adding an inclue for non-unity builds to SPropertyBinding.

diff --git a/Engine/Source/Editor/UMGEditor/Private/Details/SPropertyBinding.cpp b/Engine/Source/Editor/UMGEditor/Private/Details/SPropertyBinding.cpp
--- a/Engine/Source/Editor/UMGEditor/Private/Details/SPropertyBinding.cpp
+++ b/Engine/Source/Editor/UMGEditor/Private/Details/SPropertyBinding.cpp
@@ -3,6 +3,7 @@
 #include "UMGEditorPrivatePCH.h"
 
 #include "SPropertyBinding.h"
+#include "PropertyHandle.h"
 
 #define LOCTEXT_NAMESPACE "UMG"
 
