UMGEditor: Fixed non-unity

diff --git a/Engine/Source/Editor/UMGEditor/Private/Details/SPropertyBinding.cpp b/Engine/Source/Editor/UMGEditor/Private/Details/SPropertyBinding.cpp
--- a/Engine/Source/Editor/UMGEditor/Private/Details/SPropertyBinding.cpp
+++ b/Engine/Source/Editor/UMGEditor/Private/Details/SPropertyBinding.cpp
@@ -1,12 +1,14 @@
 // Copyright 1998-2014 Epic Games, Inc. All Rights Reserved.
 
 #include "UMGEditorPrivatePCH.h"
-
-#include "SPropertyBinding.h"
+#include "DetailLayoutBuilder.h"
 #include "PropertyHandle.h"
+#include "SPropertyBinding.h"
+
 
 #define LOCTEXT_NAMESPACE "UMG"
 
+
 /////////////////////////////////////////////////////
 // SPropertyBinding
 
