UMG - Stubbing in a widget navigation structure to UMG to begin exposing controller navigation customization support.

diff --git a/Engine/Source/Editor/UMGEditor/Private/Customizations/WidgetNavigationCustomization.cpp b/Engine/Source/Editor/UMGEditor/Private/Customizations/WidgetNavigationCustomization.cpp
new file mode 100644
--- /dev/null
+++ b/Engine/Source/Editor/UMGEditor/Private/Customizations/WidgetNavigationCustomization.cpp
@@ -0,0 +1,85 @@
+// Copyright 1998-2014 Epic Games, Inc. All Rights Reserved.
+
+#include "UMGEditorPrivatePCH.h"
+
+#include "PropertyEditing.h"
+#include "ObjectEditorUtils.h"
+#include "WidgetGraphSchema.h"
+#include "ScopedTransaction.h"
+#include "BlueprintEditorUtils.h"
+
+#include "WidgetNavigationCustomization.h"
+
+#define LOCTEXT_NAMESPACE "UMG"
+
+// FWidgetNavigationCustomization
+////////////////////////////////////////////////////////////////////////////////
+
+void FWidgetNavigationCustomization::CustomizeHeader(TSharedRef<IPropertyHandle> PropertyHandle, FDetailWidgetRow& HeaderRow, IPropertyTypeCustomizationUtils& CustomizationUtils)
+{
+	//HeaderRow
+	//.WholeRowContent()
+	//[
+	//	SNew(SButton)
+	//	.HAlign(HAlign_Center)
+	//	.OnClicked(this, &FWidgetNavigationCustomization::OnCustomizeNavigation, TWeakPtr<IPropertyHandle>(PropertyHandle))
+	//	[
+	//		SNew(STextBlock)
+	//		.Text(LOCTEXT("CustomizeNavigation", "Customize Navigation"))
+	//	]
+	//];
+}
+
+void FWidgetNavigationCustomization::CustomizeChildren(TSharedRef<IPropertyHandle> PropertyHandle, IDetailChildrenBuilder& ChildBuilder, IPropertyTypeCustomizationUtils& CustomizationUtils)
+{
+	//FillOutChildren(PropertyHandle, ChildBuilder, CustomizationUtils);
+}
+
+void FWidgetNavigationCustomization::FillOutChildren(TSharedRef<IPropertyHandle> PropertyHandle, IDetailChildrenBuilder& ChildBuilder, IPropertyTypeCustomizationUtils& CustomizationUtils)
+{
+	// Generate all the other children
+	uint32 NumChildren;
+	PropertyHandle->GetNumChildren(NumChildren);
+
+	for ( uint32 ChildIndex = 0; ChildIndex < NumChildren; ChildIndex++ )
+	{
+		TSharedRef<IPropertyHandle> ChildHandle = PropertyHandle->GetChildHandle(ChildIndex).ToSharedRef();
+		if ( ChildHandle->IsCustomized() )
+		{
+			continue;
+		}
+
+		if ( ChildHandle->GetProperty() == NULL )
+		{
+			FillOutChildren(ChildHandle, ChildBuilder, CustomizationUtils);
+		}
+		else
+		{
+			ChildBuilder.AddChildProperty(ChildHandle);
+		}
+	}
+}
+
+FReply FWidgetNavigationCustomization::OnCustomizeNavigation(TWeakPtr<IPropertyHandle> PropertyHandle)
+{
+	TArray<UObject*> OuterObjects;
+	PropertyHandle.Pin()->GetOuterObjects(OuterObjects);
+
+	const FScopedTransaction Transaction(LOCTEXT("InitializeNavigation", "Initialize Navigation"));
+
+	for ( UObject* OuterObject : OuterObjects )
+	{
+		if ( UWidget* Widget = Cast<UWidget>(OuterObject) )
+		{
+			if ( !Widget->Navigation )
+			{
+				Widget->Modify();
+				Widget->Navigation = ConstructObject<UWidgetNavigation>(UWidgetNavigation::StaticClass(), Widget);
+			}
+		}
+	}
+
+	return FReply::Handled();
+}
+
+#undef LOCTEXT_NAMESPACE
diff --git a/Engine/Source/Editor/UMGEditor/Private/Customizations/WidgetNavigationCustomization.h b/Engine/Source/Editor/UMGEditor/Private/Customizations/WidgetNavigationCustomization.h
new file mode 100644
--- /dev/null
+++ b/Engine/Source/Editor/UMGEditor/Private/Customizations/WidgetNavigationCustomization.h
@@ -0,0 +1,30 @@
+// Copyright 1998-2014 Epic Games, Inc. All Rights Reserved.
+
+#pragma once
+
+class FWidgetNavigationCustomization : public IPropertyTypeCustomization
+{
+public:
+	/** Makes a new instance of this detail layout class for a specific detail view requesting it */
+	static TSharedRef<class IPropertyTypeCustomization> MakeInstance(UBlueprint* Blueprint)
+	{
+		return MakeShareable(new FWidgetNavigationCustomization(Blueprint));
+	}
+
+	FWidgetNavigationCustomization(UBlueprint* InBlueprint)
+		: Blueprint(Cast<UWidgetBlueprint>(InBlueprint))
+	{
+	}
+	
+	virtual void CustomizeHeader(TSharedRef<IPropertyHandle> PropertyHandle, class FDetailWidgetRow& HeaderRow, IPropertyTypeCustomizationUtils& CustomizationUtils) override;
+
+	virtual void CustomizeChildren(TSharedRef<IPropertyHandle> PropertyHandle, class IDetailChildrenBuilder& ChildBuilder, IPropertyTypeCustomizationUtils& CustomizationUtils) override;
+	
+private:
+	void FillOutChildren(TSharedRef<IPropertyHandle> PropertyHandle, class IDetailChildrenBuilder& ChildBuilder, IPropertyTypeCustomizationUtils& CustomizationUtils);
+
+	FReply OnCustomizeNavigation(TWeakPtr<IPropertyHandle> PropertyHandle);
+
+private:
+	UWidgetBlueprint* Blueprint;
+};
diff --git a/Engine/Source/Editor/UMGEditor/Private/Details/SWidgetDetailsView.cpp b/Engine/Source/Editor/UMGEditor/Private/Details/SWidgetDetailsView.cpp
--- a/Engine/Source/Editor/UMGEditor/Private/Details/SWidgetDetailsView.cpp
+++ b/Engine/Source/Editor/UMGEditor/Private/Details/SWidgetDetailsView.cpp
@@ -13,6 +13,7 @@
 
 #include "Customizations/SlateBrushCustomization.h"
 
+#include "WidgetNavigationCustomization.h"
 #include "CanvasSlotCustomization.h"
 #include "HorizontalAlignmentCustomization.h"
 #include "VerticalAlignmentCustomization.h"
@@ -154,6 +155,7 @@ SWidgetDetailsView::~SWidgetDetailsView()
 	static FName PropertyEditor("PropertyEditor");
 	FPropertyEditorModule& PropertyModule = FModuleManager::GetModuleChecked<FPropertyEditorModule>(PropertyEditor);
 
+	PropertyModule.UnregisterCustomPropertyTypeLayout(TEXT("WidgetNavigation"), nullptr, PropertyView);
 	PropertyModule.UnregisterCustomPropertyTypeLayout(TEXT("PanelSlot"), nullptr, PropertyView);
 	PropertyModule.UnregisterCustomPropertyTypeLayout(TEXT("EHorizontalAlignment"), nullptr, PropertyView);
 	PropertyModule.UnregisterCustomPropertyTypeLayout(TEXT("EVerticalAlignment"), nullptr, PropertyView);
@@ -169,6 +171,7 @@ void SWidgetDetailsView::RegisterCustomizations()
 	static FName PropertyEditor("PropertyEditor");
 	FPropertyEditorModule& PropertyModule = FModuleManager::GetModuleChecked<FPropertyEditorModule>(PropertyEditor);
 
+	PropertyModule.RegisterCustomPropertyTypeLayout(TEXT("WidgetNavigation"), FOnGetPropertyTypeCustomizationInstance::CreateStatic(&FWidgetNavigationCustomization::MakeInstance, BlueprintEditor.Pin()->GetBlueprintObj()), nullptr, PropertyView);
 	PropertyModule.RegisterCustomPropertyTypeLayout(TEXT("PanelSlot"), FOnGetPropertyTypeCustomizationInstance::CreateStatic(&FCanvasSlotCustomization::MakeInstance, BlueprintEditor.Pin()->GetBlueprintObj()), nullptr, PropertyView);
 	PropertyModule.RegisterCustomPropertyTypeLayout(TEXT("EHorizontalAlignment"), FOnGetPropertyTypeCustomizationInstance::CreateStatic(&FHorizontalAlignmentCustomization::MakeInstance), nullptr, PropertyView);
 	PropertyModule.RegisterCustomPropertyTypeLayout(TEXT("EVerticalAlignment"), FOnGetPropertyTypeCustomizationInstance::CreateStatic(&FVerticalAlignmentCustomization::MakeInstance), nullptr, PropertyView);
diff --git a/Engine/Source/Runtime/UMG/Private/Components/Widget.cpp b/Engine/Source/Runtime/UMG/Private/Components/Widget.cpp
--- a/Engine/Source/Runtime/UMG/Private/Components/Widget.cpp
+++ b/Engine/Source/Runtime/UMG/Private/Components/Widget.cpp
@@ -13,6 +13,9 @@ UWidget::UWidget(const FObjectInitializer& ObjectInitializer)
 	bDesignTime = false;
 	Visiblity = ESlateVisibility::Visible;
 	RenderTransformPivot = FVector2D(0.5f, 0.5f);
+
+	//TODO UMG ToolTipWidget
+	//TODO UMG Cursor doesn't work yet, the underlying slate version needs it to be TOptional.
 }
 
 void UWidget::SetRenderTransform(FWidgetTransform Transform)
@@ -311,7 +314,7 @@ TSharedRef<SWidget> UWidget::BuildDesignTimeWidget(TSharedRef<SWidget> WrapWidge
 bool UWidget::IsGeneratedName() const
 {
 	FString Name = GetName();
-	FString BaseName = GetClass()->GetName() + "_";
+	FString BaseName = GetClass()->GetName() + TEXT("_");
 
 	if ( Name.StartsWith(BaseName) )
 	{
@@ -323,14 +326,14 @@ bool UWidget::IsGeneratedName() const
 
 FString UWidget::GetLabelMetadata() const
 {
-	return "";
+	return TEXT("");
 }
 
 FString UWidget::GetLabel() const
 {
 	if ( IsGeneratedName() && !bIsVariable )
 	{
-		return "[" + GetClass()->GetName() + "]" + GetLabelMetadata();
+		return TEXT("[") + GetClass()->GetName() + TEXT("]") + GetLabelMetadata();
 	}
 	else
 	{
diff --git a/Engine/Source/Runtime/UMG/Public/Blueprint/WidgetNavigation.h b/Engine/Source/Runtime/UMG/Public/Blueprint/WidgetNavigation.h
new file mode 100644
--- /dev/null
+++ b/Engine/Source/Runtime/UMG/Public/Blueprint/WidgetNavigation.h
@@ -0,0 +1,47 @@
+// Copyright 1998-2014 Epic Games, Inc. All Rights Reserved.
+
+#pragma once
+
+#include "WidgetNavigation.generated.h"
+
+/**
+ *
+ */
+USTRUCT()
+struct FWidgetNavigationData
+{
+	GENERATED_USTRUCT_BODY()
+
+public:
+	//UPROPERTY(EditAnywhere, BlueprintReadOnly, Category="Navigation")
+	//TEnumAsByte<> Rule;
+
+	UPROPERTY(EditAnywhere, BlueprintReadOnly, Category="Navigation")
+	FName WidgetToFocus;
+};
+
+/**
+ * 
+ */
+UCLASS()
+class UMG_API UWidgetNavigation : public UObject
+{
+	GENERATED_BODY()
+	
+public:
+	/**  */
+	UPROPERTY(EditAnywhere, BlueprintReadOnly, Category="Navigation")
+	FWidgetNavigationData Up;
+
+	/**  */
+	UPROPERTY(EditAnywhere, BlueprintReadOnly, Category="Navigation")
+	FWidgetNavigationData Down;
+
+	/**  */
+	UPROPERTY(EditAnywhere, BlueprintReadOnly, Category="Navigation")
+	FWidgetNavigationData Left;
+
+	/**  */
+	UPROPERTY(EditAnywhere, BlueprintReadOnly, Category="Navigation")
+	FWidgetNavigationData Right;
+};
diff --git a/Engine/Source/Runtime/UMG/Public/Components/Widget.h b/Engine/Source/Runtime/UMG/Public/Components/Widget.h
--- a/Engine/Source/Runtime/UMG/Public/Components/Widget.h
+++ b/Engine/Source/Runtime/UMG/Public/Components/Widget.h
@@ -80,8 +80,6 @@ public:
 	UPROPERTY()
 	FGetBool bIsEnabledDelegate;
 
-	//TODO UMG ToolTipWidget
-
 	/** Tooltip text to show when the user hovers over the widget with the mouse */
 	UPROPERTY(EditDefaultsOnly, Category=Behavior)
 	FText ToolTipText;
@@ -98,8 +96,6 @@ public:
 	UPROPERTY()
 	FGetSlateVisibility VisiblityDelegate;
 
-	//TODO UMG Cursor doesn't work yet, the underlying slate version needs it to be TOptional.
-
 	/** The cursor to show when the mouse is over the widget */
 	UPROPERTY(EditDefaultsOnly, Category=Behavior, AdvancedDisplay)
 	TEnumAsByte<EMouseCursor::Type> Cursor;
@@ -112,10 +108,17 @@ public:
 	UPROPERTY(EditDefaultsOnly, BlueprintReadOnly, Category="Render Transform", meta=( DisplayName="Transform" ))
 	FWidgetTransform RenderTransform;
 
-	/** The render transform pivot controls the location about which transforms are applied.  This value is a normalized coordinate about which things like rotations will occur. */
+	/**
+	 * The render transform pivot controls the location about which transforms are applied.  
+	 * This value is a normalized coordinate about which things like rotations will occur.
+	 */
 	UPROPERTY(EditDefaultsOnly, BlueprintReadOnly, Category="Render Transform", meta=( DisplayName="Pivot" ))
 	FVector2D RenderTransformPivot;
 
+	/**  */
+	UPROPERTY(Instanced, EditAnywhere, BlueprintReadOnly, Category="Navigation")
+	class UWidgetNavigation* Navigation;
+
 #if WITH_EDITORONLY_DATA
 
 	/** Stores the design time flag setting if the widget is hidden inside the designer */
diff --git a/Engine/Source/Runtime/UMG/Public/UMG.h b/Engine/Source/Runtime/UMG/Public/UMG.h
--- a/Engine/Source/Runtime/UMG/Public/UMG.h
+++ b/Engine/Source/Runtime/UMG/Public/UMG.h
@@ -30,6 +30,7 @@
 
 #include "PanelSlot.h"
 
+#include "WidgetNavigation.h"
 #include "Widget.h"
 
 #include "PanelWidget.h"
