Prepare UMGEditor headers for private PCH cleanup.

diff --git a/Engine/Source/Editor/DetailCustomizations/Public/Customizations/SlateBrushCustomization.h b/Engine/Source/Editor/DetailCustomizations/Public/Customizations/SlateBrushCustomization.h
--- a/Engine/Source/Editor/DetailCustomizations/Public/Customizations/SlateBrushCustomization.h
+++ b/Engine/Source/Editor/DetailCustomizations/Public/Customizations/SlateBrushCustomization.h
@@ -1,7 +1,7 @@
 // Copyright 1998-2014 Epic Games, Inc. All Rights Reserved.
 
 #pragma once
-
+#include "IPropertyTypeCustomization.h"
 
 class DETAILCUSTOMIZATIONS_API FSlateBrushStructCustomization : public IPropertyTypeCustomization
 {
diff --git a/Engine/Source/Editor/Kismet/Public/SBlueprintEditorToolbar.h b/Engine/Source/Editor/Kismet/Public/SBlueprintEditorToolbar.h
--- a/Engine/Source/Editor/Kismet/Public/SBlueprintEditorToolbar.h
+++ b/Engine/Source/Editor/Kismet/Public/SBlueprintEditorToolbar.h
@@ -4,6 +4,8 @@
 
 #include "WorkflowOrientedApp/SModeWidget.h"
 
+class FBlueprintEditor;
+
 /**
  * Kismet menu
  */
diff --git a/Engine/Source/Editor/Sequencer/Public/ISequencerObjectBindingManager.h b/Engine/Source/Editor/Sequencer/Public/ISequencerObjectBindingManager.h
--- a/Engine/Source/Editor/Sequencer/Public/ISequencerObjectBindingManager.h
+++ b/Engine/Source/Editor/Sequencer/Public/ISequencerObjectBindingManager.h
@@ -3,6 +3,7 @@
 #pragma once
 
 class FMovieSceneInstance;
+class UMovieScene;
 
 class ISequencerObjectBindingManager
 {
diff --git a/Engine/Source/Editor/UMGEditor/Private/Animation/Sequencer2DTransformTrackEditor.cpp b/Engine/Source/Editor/UMGEditor/Private/Animation/Sequencer2DTransformTrackEditor.cpp
--- a/Engine/Source/Editor/UMGEditor/Private/Animation/Sequencer2DTransformTrackEditor.cpp
+++ b/Engine/Source/Editor/UMGEditor/Private/Animation/Sequencer2DTransformTrackEditor.cpp
@@ -9,6 +9,7 @@
 #include "Editor/Sequencer/Public/ISectionLayoutBuilder.h"
 #include "Editor/Sequencer/Public/ISequencerObjectChangeListener.h"
 #include "Editor/PropertyEditor/Public/PropertyHandle.h"
+#include "Slate/WidgetTransform.h"
 
 class F2DTransformSection : public FPropertySection
 {
diff --git a/Engine/Source/Editor/UMGEditor/Private/Animation/UMGDetailKeyframeHandler.cpp b/Engine/Source/Editor/UMGEditor/Private/Animation/UMGDetailKeyframeHandler.cpp
--- a/Engine/Source/Editor/UMGEditor/Private/Animation/UMGDetailKeyframeHandler.cpp
+++ b/Engine/Source/Editor/UMGEditor/Private/Animation/UMGDetailKeyframeHandler.cpp
@@ -4,6 +4,7 @@
 #include "UMGDetailKeyframeHandler.h"
 #include "ISequencer.h"
 #include "PropertyHandle.h"
+#include "WidgetBlueprintEditor.h"
 
 FUMGDetailKeyframeHandler::FUMGDetailKeyframeHandler(TSharedPtr<FWidgetBlueprintEditor> InBlueprintEditor)
 	: BlueprintEditor( InBlueprintEditor )
diff --git a/Engine/Source/Editor/UMGEditor/Private/Animation/UMGSequencerObjectBindingManager.cpp b/Engine/Source/Editor/UMGEditor/Private/Animation/UMGSequencerObjectBindingManager.cpp
--- a/Engine/Source/Editor/UMGEditor/Private/Animation/UMGSequencerObjectBindingManager.cpp
+++ b/Engine/Source/Editor/UMGEditor/Private/Animation/UMGSequencerObjectBindingManager.cpp
@@ -5,6 +5,12 @@
 #include "WidgetBlueprintEditor.h"
 #include "ISequencer.h"
 #include "MovieScene.h"
+#include "WidgetBlueprintEditor.h"
+#include "Components/PanelSlot.h"
+#include "Components/Visual.h"
+#include "Blueprint/WidgetTree.h"
+#include "Animation/WidgetAnimation.h"
+#include "WidgetBlueprint.h"
 
 FUMGSequencerObjectBindingManager::FUMGSequencerObjectBindingManager( FWidgetBlueprintEditor& InWidgetBlueprintEditor, UWidgetAnimation& InAnimation )
 	: WidgetAnimation( &InAnimation )
diff --git a/Engine/Source/Editor/UMGEditor/Private/Animation/UMGSequencerObjectBindingManager.h b/Engine/Source/Editor/UMGEditor/Private/Animation/UMGSequencerObjectBindingManager.h
--- a/Engine/Source/Editor/UMGEditor/Private/Animation/UMGSequencerObjectBindingManager.h
+++ b/Engine/Source/Editor/UMGEditor/Private/Animation/UMGSequencerObjectBindingManager.h
@@ -4,6 +4,9 @@
 
 #include "ISequencerObjectBindingManager.h"
 
+class FWidgetBlueprintEditor;
+class UWidgetAnimation;
+
 class FUMGSequencerObjectBindingManager : public ISequencerObjectBindingManager
 {
 public:
diff --git a/Engine/Source/Editor/UMGEditor/Private/BlueprintModes/WidgetBlueprintApplicationModes.cpp b/Engine/Source/Editor/UMGEditor/Private/BlueprintModes/WidgetBlueprintApplicationModes.cpp
--- a/Engine/Source/Editor/UMGEditor/Private/BlueprintModes/WidgetBlueprintApplicationModes.cpp
+++ b/Engine/Source/Editor/UMGEditor/Private/BlueprintModes/WidgetBlueprintApplicationModes.cpp
@@ -1,6 +1,7 @@
 // Copyright 1998-2014 Epic Games, Inc. All Rights Reserved.
 
 #include "UMGEditorPrivatePCH.h"
+#include "BlueprintModes/WidgetBlueprintApplicationModes.h"
 
 // Mode constants
 const FName FWidgetBlueprintApplicationModes::DesignerMode("DesignerName");
diff --git a/Engine/Source/Editor/UMGEditor/Private/BlueprintModes/WidgetDesignerApplicationMode.cpp b/Engine/Source/Editor/UMGEditor/Private/BlueprintModes/WidgetDesignerApplicationMode.cpp
--- a/Engine/Source/Editor/UMGEditor/Private/BlueprintModes/WidgetDesignerApplicationMode.cpp
+++ b/Engine/Source/Editor/UMGEditor/Private/BlueprintModes/WidgetDesignerApplicationMode.cpp
@@ -21,6 +21,7 @@
 #include "SequencerTabSummoner.h"
 #include "DetailsTabSummoner.h"
 #include "AnimationTabSummoner.h"
+#include "BlueprintModes/WidgetBlueprintApplicationModes.h"
 
 #define LOCTEXT_NAMESPACE "WidgetDesignerMode"
 
diff --git a/Engine/Source/Editor/UMGEditor/Private/BlueprintModes/WidgetGraphApplicationMode.cpp b/Engine/Source/Editor/UMGEditor/Private/BlueprintModes/WidgetGraphApplicationMode.cpp
--- a/Engine/Source/Editor/UMGEditor/Private/BlueprintModes/WidgetGraphApplicationMode.cpp
+++ b/Engine/Source/Editor/UMGEditor/Private/BlueprintModes/WidgetGraphApplicationMode.cpp
@@ -10,6 +10,7 @@
 #include "WidgetBlueprintEditor.h"
 #include "WidgetGraphApplicationMode.h"
 #include "WidgetBlueprintEditorToolbar.h"
+#include "BlueprintModes/WidgetBlueprintApplicationModes.h"
 
 /////////////////////////////////////////////////////
 // FWidgetGraphApplicationMode
diff --git a/Engine/Source/Editor/UMGEditor/Private/Customizations/CanvasSlotCustomization.h b/Engine/Source/Editor/UMGEditor/Private/Customizations/CanvasSlotCustomization.h
--- a/Engine/Source/Editor/UMGEditor/Private/Customizations/CanvasSlotCustomization.h
+++ b/Engine/Source/Editor/UMGEditor/Private/Customizations/CanvasSlotCustomization.h
@@ -1,6 +1,10 @@
 // Copyright 1998-2014 Epic Games, Inc. All Rights Reserved.
 
 #pragma once
+#include "IPropertyTypeCustomization.h"
+#include "WidgetBlueprint.h"
+
+class IDetailPropertyRow;
 
 class FCanvasSlotCustomization : public IPropertyTypeCustomization
 {
diff --git a/Engine/Source/Editor/UMGEditor/Private/Customizations/DetailCustomizations.cpp b/Engine/Source/Editor/UMGEditor/Private/Customizations/DetailCustomizations.cpp
--- a/Engine/Source/Editor/UMGEditor/Private/Customizations/DetailCustomizations.cpp
+++ b/Engine/Source/Editor/UMGEditor/Private/Customizations/DetailCustomizations.cpp
@@ -3,13 +3,20 @@
 #include "UMGEditorPrivatePCH.h"
 
 #include "DetailCustomizations.h"
-
+#include "BlueprintModes/WidgetBlueprintApplicationModes.h"
 #include "PropertyEditing.h"
 #include "ObjectEditorUtils.h"
 #include "WidgetGraphSchema.h"
 #include "ScopedTransaction.h"
 #include "BlueprintEditorUtils.h"
 #include "WidgetGraphSchema.h"
+#include "Components/Widget.h"
+#include "WidgetBlueprint.h"
+#include "K2Node_ComponentBoundEvent.h"
+#include "Kismet2/KismetEditorUtilities.h"
+#include "EdGraphSchema_K2_Actions.h"
+#include "WidgetBlueprintEditor.h"
+#include "Components/PanelSlot.h"
 
 #define LOCTEXT_NAMESPACE "UMG"
 
diff --git a/Engine/Source/Editor/UMGEditor/Private/Customizations/DetailCustomizations.h b/Engine/Source/Editor/UMGEditor/Private/Customizations/DetailCustomizations.h
--- a/Engine/Source/Editor/UMGEditor/Private/Customizations/DetailCustomizations.h
+++ b/Engine/Source/Editor/UMGEditor/Private/Customizations/DetailCustomizations.h
@@ -1,6 +1,12 @@
 // Copyright 1998-2014 Epic Games, Inc. All Rights Reserved.
 
 #pragma once
+#include "IDetailCustomization.h"
+
+class FWidgetBlueprintEditor;
+class UWidget;
+class IPropertyHandle;
+class UWidgetBlueprint;
 
 /**
  * Provides the customization for all UWidgets.  Bindings, style disabling...etc.
diff --git a/Engine/Source/Editor/UMGEditor/Private/Customizations/HorizontalAlignmentCustomization.h b/Engine/Source/Editor/UMGEditor/Private/Customizations/HorizontalAlignmentCustomization.h
--- a/Engine/Source/Editor/UMGEditor/Private/Customizations/HorizontalAlignmentCustomization.h
+++ b/Engine/Source/Editor/UMGEditor/Private/Customizations/HorizontalAlignmentCustomization.h
@@ -1,6 +1,7 @@
 // Copyright 1998-2014 Epic Games, Inc. All Rights Reserved.
 
 #pragma once
+#include "IPropertyTypeCustomization.h"
 
 class FHorizontalAlignmentCustomization : public IPropertyTypeCustomization
 {
diff --git a/Engine/Source/Editor/UMGEditor/Private/Customizations/SlateChildSizeCustomization.cpp b/Engine/Source/Editor/UMGEditor/Private/Customizations/SlateChildSizeCustomization.cpp
--- a/Engine/Source/Editor/UMGEditor/Private/Customizations/SlateChildSizeCustomization.cpp
+++ b/Engine/Source/Editor/UMGEditor/Private/Customizations/SlateChildSizeCustomization.cpp
@@ -10,6 +10,8 @@
 
 #include "SlateChildSizeCustomization.h"
 #include "SNumericEntryBox.h"
+#include "Components/SlateWrapperTypes.h"
+#include "Widgets/Input/SCheckBox.h"
 
 #define LOCTEXT_NAMESPACE "UMG"
 
diff --git a/Engine/Source/Editor/UMGEditor/Private/Customizations/SlateChildSizeCustomization.h b/Engine/Source/Editor/UMGEditor/Private/Customizations/SlateChildSizeCustomization.h
--- a/Engine/Source/Editor/UMGEditor/Private/Customizations/SlateChildSizeCustomization.h
+++ b/Engine/Source/Editor/UMGEditor/Private/Customizations/SlateChildSizeCustomization.h
@@ -1,6 +1,9 @@
 // Copyright 1998-2014 Epic Games, Inc. All Rights Reserved.
 
 #pragma once
+#include "Components/SlateWrapperTypes.h"
+#include "Widgets/Input/SCheckBox.h"
+#include "IPropertyTypeCustomization.h"
 
 class FSlateChildSizeCustomization : public IPropertyTypeCustomization
 {
diff --git a/Engine/Source/Editor/UMGEditor/Private/Designer/SDesignerView.cpp b/Engine/Source/Editor/UMGEditor/Private/Designer/SDesignerView.cpp
--- a/Engine/Source/Editor/UMGEditor/Private/Designer/SDesignerView.cpp
+++ b/Engine/Source/Editor/UMGEditor/Private/Designer/SDesignerView.cpp
@@ -1,6 +1,9 @@
 // Copyright 1998-2014 Epic Games, Inc. All Rights Reserved.
 
 #include "UMGEditorPrivatePCH.h"
+
+#include "Designer/DesignTimeUtils.h"
+
 #include "Extensions/CanvasSlotExtension.h"
 #include "Extensions/GridSlotExtension.h"
 #include "Extensions/HorizontalSlotExtension.h"
@@ -22,6 +25,19 @@
 #include "SDPIScaler.h"
 #include "SNumericEntryBox.h"
 
+#include "Components/PanelWidget.h"
+#include "Components/Widget.h"
+#include "WidgetBlueprint.h"
+#include "WidgetBlueprintCompiler.h"
+#include "WidgetBlueprintEditor.h"
+#include "WidgetBlueprintEditorUtils.h"
+
+#include "ObjectEditorUtils.h"
+#include "Blueprint/WidgetTree.h"
+#include "ScopedTransaction.h"
+#include "Settings/WidgetDesignerSettings.h"
+#include "Components/CanvasPanelSlot.h"
+
 #define LOCTEXT_NAMESPACE "UMG"
 
 const float HoveredAnimationTime = 0.150f;
diff --git a/Engine/Source/Editor/UMGEditor/Private/Designer/SDesignerView.h b/Engine/Source/Editor/UMGEditor/Private/Designer/SDesignerView.h
--- a/Engine/Source/Editor/UMGEditor/Private/Designer/SDesignerView.h
+++ b/Engine/Source/Editor/UMGEditor/Private/Designer/SDesignerView.h
@@ -20,6 +20,8 @@ namespace EDesignerMessage
 }
 
 class FDesignerExtension;
+class UPanelWidget;
+class UUserWidget;
 
 /**
  * The designer for widgets.  Allows for laying out widgets in a drag and drop environment.
diff --git a/Engine/Source/Editor/UMGEditor/Private/Designer/STransformHandle.cpp b/Engine/Source/Editor/UMGEditor/Private/Designer/STransformHandle.cpp
--- a/Engine/Source/Editor/UMGEditor/Private/Designer/STransformHandle.cpp
+++ b/Engine/Source/Editor/UMGEditor/Private/Designer/STransformHandle.cpp
@@ -1,9 +1,12 @@
 // Copyright 1998-2014 Epic Games, Inc. All Rights Reserved.
 
 #include "UMGEditorPrivatePCH.h"
-
+#include "IUMGDesigner.h"
 #include "ObjectEditorUtils.h"
 #include "STransformHandle.h"
+#include "WidgetReference.h"
+#include "Widget.h"
+#include "Components/PanelSlot.h"
 
 #define LOCTEXT_NAMESPACE "UMG"
 
diff --git a/Engine/Source/Editor/UMGEditor/Private/Details/SWidgetDetailsView.cpp b/Engine/Source/Editor/UMGEditor/Private/Details/SWidgetDetailsView.cpp
--- a/Engine/Source/Editor/UMGEditor/Private/Details/SWidgetDetailsView.cpp
+++ b/Engine/Source/Editor/UMGEditor/Private/Details/SWidgetDetailsView.cpp
@@ -12,6 +12,7 @@
 #include "EditorClassUtils.h"
 
 #include "Customizations/SlateBrushCustomization.h"
+#include "Customizations/DetailCustomizations.h"
 
 #include "WidgetNavigationCustomization.h"
 #include "CanvasSlotCustomization.h"
@@ -19,6 +20,12 @@
 #include "VerticalAlignmentCustomization.h"
 #include "SlateChildSizeCustomization.h"
 #include "TextJustifyCustomization.h"
+#include "PropertyEditorModule.h"
+#include "WidgetBlueprintEditor.h"
+#include "Kismet2/BlueprintEditorUtils.h"
+#include "Blueprint/WidgetTree.h"
+#include "WidgetBlueprintEditorUtils.h"
+#include "ScopedTransaction.h"
 
 #define LOCTEXT_NAMESPACE "UMG"
 
diff --git a/Engine/Source/Editor/UMGEditor/Private/Details/SWidgetDetailsView.h b/Engine/Source/Editor/UMGEditor/Private/Details/SWidgetDetailsView.h
--- a/Engine/Source/Editor/UMGEditor/Private/Details/SWidgetDetailsView.h
+++ b/Engine/Source/Editor/UMGEditor/Private/Details/SWidgetDetailsView.h
@@ -7,6 +7,8 @@
 
 #include "WidgetTemplate.h"
 
+class FWidgetBlueprintEditor;
+
 /**
  * The details view used in the designer section of the widget blueprint editor.
  */
diff --git a/Engine/Source/Editor/UMGEditor/Private/Extensions/CanvasSlotExtension.cpp b/Engine/Source/Editor/UMGEditor/Private/Extensions/CanvasSlotExtension.cpp
--- a/Engine/Source/Editor/UMGEditor/Private/Extensions/CanvasSlotExtension.cpp
+++ b/Engine/Source/Editor/UMGEditor/Private/Extensions/CanvasSlotExtension.cpp
@@ -1,7 +1,11 @@
 // Copyright 1998-2014 Epic Games, Inc. All Rights Reserved.
 
 #include "UMGEditorPrivatePCH.h"
+#include "Components/CanvasPanel.h"
+#include "Components/CanvasPanelSlot.h"
+#include "Components/Widget.h"
 #include "Extensions/CanvasSlotExtension.h"
+#include "ObjectEditorUtils.h"
 
 #define LOCTEXT_NAMESPACE "UMG"
 
diff --git a/Engine/Source/Editor/UMGEditor/Private/Extensions/CanvasSlotExtension.h b/Engine/Source/Editor/UMGEditor/Private/Extensions/CanvasSlotExtension.h
--- a/Engine/Source/Editor/UMGEditor/Private/Extensions/CanvasSlotExtension.h
+++ b/Engine/Source/Editor/UMGEditor/Private/Extensions/CanvasSlotExtension.h
@@ -3,6 +3,10 @@
 #pragma once
 
 #include "DesignerExtension.h"
+#include "Slate/Anchors.h"
+
+class UCanvasPanel;
+class UCanvasPanelSlot;
 
 /** Set of anchor widget types */
 namespace EAnchorWidget
diff --git a/Engine/Source/Editor/UMGEditor/Private/Extensions/DesignerExtension.cpp b/Engine/Source/Editor/UMGEditor/Private/Extensions/DesignerExtension.cpp
--- a/Engine/Source/Editor/UMGEditor/Private/Extensions/DesignerExtension.cpp
+++ b/Engine/Source/Editor/UMGEditor/Private/Extensions/DesignerExtension.cpp
@@ -1,7 +1,9 @@
 // Copyright 1998-2014 Epic Games, Inc. All Rights Reserved.
 
 #include "UMGEditorPrivatePCH.h"
-
+#include "Components/Widget.h"
+#include "DesignerExtension.h"
+#include "ScopedTransaction.h"
 #include "WidgetBlueprintEditor.h"
 
 #define LOCTEXT_NAMESPACE "UMG"
diff --git a/Engine/Source/Editor/UMGEditor/Private/Extensions/GridSlotExtension.cpp b/Engine/Source/Editor/UMGEditor/Private/Extensions/GridSlotExtension.cpp
--- a/Engine/Source/Editor/UMGEditor/Private/Extensions/GridSlotExtension.cpp
+++ b/Engine/Source/Editor/UMGEditor/Private/Extensions/GridSlotExtension.cpp
@@ -2,6 +2,10 @@
 
 #include "UMGEditorPrivatePCH.h"
 #include "Components/GridSlot.h"
+#include "Components/Widget.h"
+#include "Extensions/GridSlotExtension.h"
+#include "Kismet2/BlueprintEditorUtils.h"
+#include "WidgetBlueprint.h"
 
 #define LOCTEXT_NAMESPACE "UMG"
 
diff --git a/Engine/Source/Editor/UMGEditor/Private/Extensions/HorizontalSlotExtension.cpp b/Engine/Source/Editor/UMGEditor/Private/Extensions/HorizontalSlotExtension.cpp
--- a/Engine/Source/Editor/UMGEditor/Private/Extensions/HorizontalSlotExtension.cpp
+++ b/Engine/Source/Editor/UMGEditor/Private/Extensions/HorizontalSlotExtension.cpp
@@ -3,7 +3,9 @@
 #include "UMGEditorPrivatePCH.h"
 #include "Components/HorizontalBox.h"
 #include "Components/HorizontalBoxSlot.h"
+#include "Extensions/HorizontalSlotExtension.h"
 #include "Kismet2/BlueprintEditorUtils.h"
+#include "WidgetBlueprint.h"
 
 #define LOCTEXT_NAMESPACE "UMG"
 
diff --git a/Engine/Source/Editor/UMGEditor/Private/Extensions/UniformGridSlotExtension.cpp b/Engine/Source/Editor/UMGEditor/Private/Extensions/UniformGridSlotExtension.cpp
--- a/Engine/Source/Editor/UMGEditor/Private/Extensions/UniformGridSlotExtension.cpp
+++ b/Engine/Source/Editor/UMGEditor/Private/Extensions/UniformGridSlotExtension.cpp
@@ -2,6 +2,11 @@
 
 #include "UMGEditorPrivatePCH.h"
 #include "Components/UniformGridSlot.h"
+#include "Extensions/UniformGridSlotExtension.h"
+#include "Components/Widget.h"
+#include "Kismet2/BlueprintEditorUtils.h"
+#include "WidgetBlueprint.h"
+
 #define LOCTEXT_NAMESPACE "UMG"
 
 FUniformGridSlotExtension::FUniformGridSlotExtension()
diff --git a/Engine/Source/Editor/UMGEditor/Private/Extensions/VerticalSlotExtension.cpp b/Engine/Source/Editor/UMGEditor/Private/Extensions/VerticalSlotExtension.cpp
--- a/Engine/Source/Editor/UMGEditor/Private/Extensions/VerticalSlotExtension.cpp
+++ b/Engine/Source/Editor/UMGEditor/Private/Extensions/VerticalSlotExtension.cpp
@@ -3,7 +3,9 @@
 #include "UMGEditorPrivatePCH.h"
 #include "Components/VerticalBox.h"
 #include "Components/VerticalBoxSlot.h"
+#include "Extensions/VerticalSlotExtension.h"
 #include "Kismet2/BlueprintEditorUtils.h"
+#include "WidgetBlueprint.h"
 
 #define LOCTEXT_NAMESPACE "UMG"
 
diff --git a/Engine/Source/Editor/UMGEditor/Private/Hierarchy/SHierarchyView.cpp b/Engine/Source/Editor/UMGEditor/Private/Hierarchy/SHierarchyView.cpp
--- a/Engine/Source/Editor/UMGEditor/Private/Hierarchy/SHierarchyView.cpp
+++ b/Engine/Source/Editor/UMGEditor/Private/Hierarchy/SHierarchyView.cpp
@@ -15,6 +15,10 @@
 #include "WidgetTemplateClass.h"
 #include "WidgetBlueprintEditor.h"
 #include "SSearchBox.h"
+#include "WidgetBlueprint.h"
+#include "WidgetBlueprintEditorUtils.h"
+
+#include "Components/PanelWidget.h"
 
 #define LOCTEXT_NAMESPACE "UMG"
 
diff --git a/Engine/Source/Editor/UMGEditor/Private/Hierarchy/SHierarchyView.h b/Engine/Source/Editor/UMGEditor/Private/Hierarchy/SHierarchyView.h
--- a/Engine/Source/Editor/UMGEditor/Private/Hierarchy/SHierarchyView.h
+++ b/Engine/Source/Editor/UMGEditor/Private/Hierarchy/SHierarchyView.h
@@ -7,6 +7,12 @@
 #include "SCompoundWidget.h"
 #include "SHierarchyViewItem.h"
 #include "TreeFilterHandler.h"
+#include "WidgetBlueprintEditorUtils.h"
+
+class FWidgetBlueprintEditor;
+class UWidget;
+class UWidgetBlueprint;
+
 //TODO rename SUMGEditorHierarchy
 
 /**
diff --git a/Engine/Source/Editor/UMGEditor/Private/Hierarchy/SHierarchyViewItem.cpp b/Engine/Source/Editor/UMGEditor/Private/Hierarchy/SHierarchyViewItem.cpp
--- a/Engine/Source/Editor/UMGEditor/Private/Hierarchy/SHierarchyViewItem.cpp
+++ b/Engine/Source/Editor/UMGEditor/Private/Hierarchy/SHierarchyViewItem.cpp
@@ -17,6 +17,14 @@
 
 #include "WidgetBlueprintEditor.h"
 #include "SInlineEditableTextBlock.h"
+#include "Components/Widget.h"
+
+#include "WidgetBlueprintEditorUtils.h"
+#include "Components/PanelSlot.h"
+#include "ScopedTransaction.h"
+#include "WidgetBlueprint.h"
+#include "Blueprint/WidgetTree.h"
+#include "Components/PanelWidget.h"
 
 #define LOCTEXT_NAMESPACE "UMG"
 
diff --git a/Engine/Source/Editor/UMGEditor/Private/Hierarchy/SHierarchyViewItem.h b/Engine/Source/Editor/UMGEditor/Private/Hierarchy/SHierarchyViewItem.h
--- a/Engine/Source/Editor/UMGEditor/Private/Hierarchy/SHierarchyViewItem.h
+++ b/Engine/Source/Editor/UMGEditor/Private/Hierarchy/SHierarchyViewItem.h
@@ -4,6 +4,12 @@
 
 #include "SCompoundWidget.h"
 #include "BlueprintEditor.h"
+#include "WidgetReference.h"
+#include "Components/Widget.h"
+#include "Kismet2/BlueprintEditorUtils.h"
+#include "WidgetBlueprintEditor.h"
+
+class FWidgetBlueprintEditor;
 
 class FHierarchyModel : public TSharedFromThis < FHierarchyModel >
 {
diff --git a/Engine/Source/Editor/UMGEditor/Private/Nodes/K2Node_CreateWidget.cpp b/Engine/Source/Editor/UMGEditor/Private/Nodes/K2Node_CreateWidget.cpp
--- a/Engine/Source/Editor/UMGEditor/Private/Nodes/K2Node_CreateWidget.cpp
+++ b/Engine/Source/Editor/UMGEditor/Private/Nodes/K2Node_CreateWidget.cpp
@@ -1,6 +1,7 @@
 // Copyright 1998-2014 Epic Games, Inc. All Rights Reserved.
 
 #include "UMGEditorPrivatePCH.h"
+#include "Blueprint/UserWidget.h"
 #include "Blueprint/WidgetBlueprintLibrary.h"
 #include "KismetCompiler.h"
 #include "BlueprintNodeSpawner.h"
diff --git a/Engine/Source/Editor/UMGEditor/Private/Palette/SPaletteView.cpp b/Engine/Source/Editor/UMGEditor/Private/Palette/SPaletteView.cpp
--- a/Engine/Source/Editor/UMGEditor/Private/Palette/SPaletteView.cpp
+++ b/Engine/Source/Editor/UMGEditor/Private/Palette/SPaletteView.cpp
@@ -22,6 +22,12 @@
 #include "AssetRegistryModule.h"
 #include "SSearchBox.h"
 
+#include "WidgetBlueprintCompiler.h"
+
+#include "Blueprint/UserWidget.h"
+#include "WidgetBlueprint.h"
+#include "ObjectEditorUtils.h"
+
 #define LOCTEXT_NAMESPACE "UMG"
 
 class SPaletteViewItem : public SCompoundWidget
diff --git a/Engine/Source/Editor/UMGEditor/Private/Palette/SPaletteView.h b/Engine/Source/Editor/UMGEditor/Private/Palette/SPaletteView.h
--- a/Engine/Source/Editor/UMGEditor/Private/Palette/SPaletteView.h
+++ b/Engine/Source/Editor/UMGEditor/Private/Palette/SPaletteView.h
@@ -8,6 +8,8 @@
 
 #include "WidgetTemplate.h"
 
+class UWidgetBlueprint;
+
 /** View model for the items in the widget template list */
 class FWidgetViewModel : public TSharedFromThis<FWidgetViewModel>
 {
diff --git a/Engine/Source/Editor/UMGEditor/Private/Settings/WidgetDesignerSettings.cpp b/Engine/Source/Editor/UMGEditor/Private/Settings/WidgetDesignerSettings.cpp
--- a/Engine/Source/Editor/UMGEditor/Private/Settings/WidgetDesignerSettings.cpp
+++ b/Engine/Source/Editor/UMGEditor/Private/Settings/WidgetDesignerSettings.cpp
@@ -1,6 +1,7 @@
 // Copyright 1998-2014 Epic Games, Inc. All Rights Reserved.
 
 #include "UMGEditorPrivatePCH.h"
+#include "Settings/WidgetDesignerSettings.h"
 
 UWidgetDesignerSettings::UWidgetDesignerSettings(const FObjectInitializer& ObjectInitializer)
 	: Super(ObjectInitializer)
diff --git a/Engine/Source/Editor/UMGEditor/Private/TabFactory/AnimationTabSummoner.cpp b/Engine/Source/Editor/UMGEditor/Private/TabFactory/AnimationTabSummoner.cpp
--- a/Engine/Source/Editor/UMGEditor/Private/TabFactory/AnimationTabSummoner.cpp
+++ b/Engine/Source/Editor/UMGEditor/Private/TabFactory/AnimationTabSummoner.cpp
@@ -11,6 +11,10 @@
 #include "SSearchBox.h"
 #include "SInlineEditableTextBlock.h"
 #include "GenericCommands.h"
+#include "ScopedTransaction.h"
+#include "WidgetBlueprint.h"
+#include "Kismet2/BlueprintEditorUtils.h"
+#include "Misc/TextFilter.h"
 
 #define LOCTEXT_NAMESPACE "UMG"
 
diff --git a/Engine/Source/Editor/UMGEditor/Private/TabFactory/AnimationTabSummoner.h b/Engine/Source/Editor/UMGEditor/Private/TabFactory/AnimationTabSummoner.h
--- a/Engine/Source/Editor/UMGEditor/Private/TabFactory/AnimationTabSummoner.h
+++ b/Engine/Source/Editor/UMGEditor/Private/TabFactory/AnimationTabSummoner.h
@@ -1,6 +1,7 @@
 // Copyright 1998-2014 Epic Games, Inc. All Rights Reserved.
 
 #pragma once
+#include "WorkflowOrientedApp/WorkflowTabFactory.h"
 
 struct FAnimationTabSummoner : public FWorkflowTabFactory
 {
diff --git a/Engine/Source/Editor/UMGEditor/Private/TabFactory/DesignerTabSummoner.cpp b/Engine/Source/Editor/UMGEditor/Private/TabFactory/DesignerTabSummoner.cpp
--- a/Engine/Source/Editor/UMGEditor/Private/TabFactory/DesignerTabSummoner.cpp
+++ b/Engine/Source/Editor/UMGEditor/Private/TabFactory/DesignerTabSummoner.cpp
@@ -6,7 +6,7 @@
 #include "SDesignerView.h"
 #include "SScissorRectBox.h"
 #include "UMGStyle.h"
-
+#include "WidgetBlueprintEditor.h"
 #define LOCTEXT_NAMESPACE "UMG"
 
 const FName FDesignerTabSummoner::TabID(TEXT("SlatePreview"));
diff --git a/Engine/Source/Editor/UMGEditor/Private/TabFactory/DesignerTabSummoner.h b/Engine/Source/Editor/UMGEditor/Private/TabFactory/DesignerTabSummoner.h
--- a/Engine/Source/Editor/UMGEditor/Private/TabFactory/DesignerTabSummoner.h
+++ b/Engine/Source/Editor/UMGEditor/Private/TabFactory/DesignerTabSummoner.h
@@ -1,6 +1,7 @@
 // Copyright 1998-2014 Epic Games, Inc. All Rights Reserved.
 
 #pragma once
+#include "WorkflowOrientedApp/WorkflowTabFactory.h"
 
 struct FDesignerTabSummoner : public FWorkflowTabFactory
 {
diff --git a/Engine/Source/Editor/UMGEditor/Private/TabFactory/DetailsTabSummoner.cpp b/Engine/Source/Editor/UMGEditor/Private/TabFactory/DetailsTabSummoner.cpp
--- a/Engine/Source/Editor/UMGEditor/Private/TabFactory/DetailsTabSummoner.cpp
+++ b/Engine/Source/Editor/UMGEditor/Private/TabFactory/DetailsTabSummoner.cpp
@@ -4,7 +4,7 @@
 
 #include "DetailsTabSummoner.h"
 #include "SWidgetDetailsView.h"
-
+#include "WidgetBlueprintEditor.h"
 #define LOCTEXT_NAMESPACE "UMG"
 
 const FName FDetailsTabSummoner::TabID(TEXT("WidgetDetails"));
diff --git a/Engine/Source/Editor/UMGEditor/Private/TabFactory/DetailsTabSummoner.h b/Engine/Source/Editor/UMGEditor/Private/TabFactory/DetailsTabSummoner.h
--- a/Engine/Source/Editor/UMGEditor/Private/TabFactory/DetailsTabSummoner.h
+++ b/Engine/Source/Editor/UMGEditor/Private/TabFactory/DetailsTabSummoner.h
@@ -1,6 +1,7 @@
 // Copyright 1998-2014 Epic Games, Inc. All Rights Reserved.
 
 #pragma once
+#include "WorkflowOrientedApp/WorkflowTabFactory.h"
 
 struct FDetailsTabSummoner : public FWorkflowTabFactory
 {
diff --git a/Engine/Source/Editor/UMGEditor/Private/TabFactory/HierarchyTabSummoner.cpp b/Engine/Source/Editor/UMGEditor/Private/TabFactory/HierarchyTabSummoner.cpp
--- a/Engine/Source/Editor/UMGEditor/Private/TabFactory/HierarchyTabSummoner.cpp
+++ b/Engine/Source/Editor/UMGEditor/Private/TabFactory/HierarchyTabSummoner.cpp
@@ -1,9 +1,9 @@
 // Copyright 1998-2014 Epic Games, Inc. All Rights Reserved.
 
 #include "UMGEditorPrivatePCH.h"
-
 #include "HierarchyTabSummoner.h"
 #include "SHierarchyView.h"
+#include "WidgetBlueprintEditor.h"
 
 #define LOCTEXT_NAMESPACE "UMG"
 
diff --git a/Engine/Source/Editor/UMGEditor/Private/TabFactory/HierarchyTabSummoner.h b/Engine/Source/Editor/UMGEditor/Private/TabFactory/HierarchyTabSummoner.h
--- a/Engine/Source/Editor/UMGEditor/Private/TabFactory/HierarchyTabSummoner.h
+++ b/Engine/Source/Editor/UMGEditor/Private/TabFactory/HierarchyTabSummoner.h
@@ -1,6 +1,7 @@
 // Copyright 1998-2014 Epic Games, Inc. All Rights Reserved.
 
 #pragma once
+#include "WorkflowOrientedApp/WorkflowTabFactory.h"
 
 struct FHierarchyTabSummoner : public FWorkflowTabFactory
 {
diff --git a/Engine/Source/Editor/UMGEditor/Private/TabFactory/PaletteTabSummoner.cpp b/Engine/Source/Editor/UMGEditor/Private/TabFactory/PaletteTabSummoner.cpp
--- a/Engine/Source/Editor/UMGEditor/Private/TabFactory/PaletteTabSummoner.cpp
+++ b/Engine/Source/Editor/UMGEditor/Private/TabFactory/PaletteTabSummoner.cpp
@@ -5,6 +5,7 @@
 #include "PaletteTabSummoner.h"
 #include "SPaletteView.h"
 #include "UMGStyle.h"
+#include "WidgetBlueprintEditor.h"
 
 #define LOCTEXT_NAMESPACE "UMG"
 
diff --git a/Engine/Source/Editor/UMGEditor/Private/TabFactory/PaletteTabSummoner.h b/Engine/Source/Editor/UMGEditor/Private/TabFactory/PaletteTabSummoner.h
--- a/Engine/Source/Editor/UMGEditor/Private/TabFactory/PaletteTabSummoner.h
+++ b/Engine/Source/Editor/UMGEditor/Private/TabFactory/PaletteTabSummoner.h
@@ -1,6 +1,7 @@
 // Copyright 1998-2014 Epic Games, Inc. All Rights Reserved.
 
 #pragma once
+#include "WorkflowOrientedApp/WorkflowTabFactory.h"
 
 struct FPaletteTabSummoner : public FWorkflowTabFactory
 {
diff --git a/Engine/Source/Editor/UMGEditor/Private/TabFactory/SequencerTabSummoner.cpp b/Engine/Source/Editor/UMGEditor/Private/TabFactory/SequencerTabSummoner.cpp
--- a/Engine/Source/Editor/UMGEditor/Private/TabFactory/SequencerTabSummoner.cpp
+++ b/Engine/Source/Editor/UMGEditor/Private/TabFactory/SequencerTabSummoner.cpp
@@ -5,6 +5,7 @@
 #include "SequencerTabSummoner.h"
 #include "WidgetBlueprintEditor.h"
 #include "ISequencer.h"
+#include "UMGStyle.h"
 
 #define LOCTEXT_NAMESPACE "UMG"
 
diff --git a/Engine/Source/Editor/UMGEditor/Private/TabFactory/SequencerTabSummoner.h b/Engine/Source/Editor/UMGEditor/Private/TabFactory/SequencerTabSummoner.h
--- a/Engine/Source/Editor/UMGEditor/Private/TabFactory/SequencerTabSummoner.h
+++ b/Engine/Source/Editor/UMGEditor/Private/TabFactory/SequencerTabSummoner.h
@@ -1,6 +1,7 @@
 // Copyright 1998-2014 Epic Games, Inc. All Rights Reserved.
 
 #pragma once
+#include "WorkflowOrientedApp/WorkflowTabFactory.h"
 
 struct FSequencerTabSummoner : public FWorkflowTabFactory
 {
diff --git a/Engine/Source/Editor/UMGEditor/Private/Templates/WidgetTemplateBlueprintClass.cpp b/Engine/Source/Editor/UMGEditor/Private/Templates/WidgetTemplateBlueprintClass.cpp
--- a/Engine/Source/Editor/UMGEditor/Private/Templates/WidgetTemplateBlueprintClass.cpp
+++ b/Engine/Source/Editor/UMGEditor/Private/Templates/WidgetTemplateBlueprintClass.cpp
@@ -1,9 +1,10 @@
 // Copyright 1998-2014 Epic Games, Inc. All Rights Reserved.
 
 #include "UMGEditorPrivatePCH.h"
-
+#include "Blueprint/UserWidget.h"
 #include "WidgetTemplateBlueprintClass.h"
 #include "IDocumentation.h"
+#include "WidgetBlueprint.h"
 
 #define LOCTEXT_NAMESPACE "UMGEditor"
 
diff --git a/Engine/Source/Editor/UMGEditor/Private/Templates/WidgetTemplateBlueprintClass.h b/Engine/Source/Editor/UMGEditor/Private/Templates/WidgetTemplateBlueprintClass.h
--- a/Engine/Source/Editor/UMGEditor/Private/Templates/WidgetTemplateBlueprintClass.h
+++ b/Engine/Source/Editor/UMGEditor/Private/Templates/WidgetTemplateBlueprintClass.h
@@ -4,6 +4,8 @@
 
 #include "WidgetTemplateClass.h"
 
+class UUserWidget;
+
 /**
 * A template for classes generated by UWidgetBlueprint assets
 * Accounts for the possibility that a the blueprint may not be loaded when the palette is created
diff --git a/Engine/Source/Editor/UMGEditor/Private/Templates/WidgetTemplateClass.cpp b/Engine/Source/Editor/UMGEditor/Private/Templates/WidgetTemplateClass.cpp
--- a/Engine/Source/Editor/UMGEditor/Private/Templates/WidgetTemplateClass.cpp
+++ b/Engine/Source/Editor/UMGEditor/Private/Templates/WidgetTemplateClass.cpp
@@ -1,9 +1,10 @@
 // Copyright 1998-2014 Epic Games, Inc. All Rights Reserved.
 
 #include "UMGEditorPrivatePCH.h"
-
+#include "Components/Widget.h"
 #include "WidgetTemplateClass.h"
 #include "IDocumentation.h"
+#include "Blueprint/WidgetTree.h"
 
 #define LOCTEXT_NAMESPACE "UMGEditor"
 
diff --git a/Engine/Source/Editor/UMGEditor/Private/Templates/WidgetTemplateClass.h b/Engine/Source/Editor/UMGEditor/Private/Templates/WidgetTemplateClass.h
--- a/Engine/Source/Editor/UMGEditor/Private/Templates/WidgetTemplateClass.h
+++ b/Engine/Source/Editor/UMGEditor/Private/Templates/WidgetTemplateClass.h
@@ -4,6 +4,8 @@
 
 #include "WidgetTemplate.h"
 
+class UWidget;
+
 /**
  * A template that can spawn any widget derived from the UWidget class.
  */
diff --git a/Engine/Source/Editor/UMGEditor/Private/UMGEditorModule.cpp b/Engine/Source/Editor/UMGEditor/Private/UMGEditorModule.cpp
--- a/Engine/Source/Editor/UMGEditor/Private/UMGEditorModule.cpp
+++ b/Engine/Source/Editor/UMGEditor/Private/UMGEditorModule.cpp
@@ -15,6 +15,7 @@
 #include "Animation/Sequencer2DTransformTrackEditor.h"
 #include "IUMGModule.h"
 #include "DesignerCommands.h"
+#include "WidgetBlueprint.h"
 
 const FName UMGEditorAppIdentifier = FName(TEXT("UMGEditorApp"));
 
diff --git a/Engine/Source/Editor/UMGEditor/Private/WidgetBlueprint.cpp b/Engine/Source/Editor/UMGEditor/Private/WidgetBlueprint.cpp
--- a/Engine/Source/Editor/UMGEditor/Private/WidgetBlueprint.cpp
+++ b/Engine/Source/Editor/UMGEditor/Private/WidgetBlueprint.cpp
@@ -4,6 +4,8 @@
 #include "Blueprint/WidgetTree.h"
 #include "Runtime/MovieSceneCore/Classes/MovieScene.h"
 #include "PropertyTag.h"
+#include "Blueprint/WidgetBlueprintGeneratedClass.h"
+#include "WidgetBlueprint.h"
 
 #define LOCTEXT_NAMESPACE "UMG"
 
diff --git a/Engine/Source/Editor/UMGEditor/Private/WidgetBlueprintCompiler.cpp b/Engine/Source/Editor/UMGEditor/Private/WidgetBlueprintCompiler.cpp
--- a/Engine/Source/Editor/UMGEditor/Private/WidgetBlueprintCompiler.cpp
+++ b/Engine/Source/Editor/UMGEditor/Private/WidgetBlueprintCompiler.cpp
@@ -9,6 +9,8 @@
 #include "MovieScene.h"
 #include "K2Node_FunctionEntry.h"
 #include "Blueprint/WidgetTree.h"
+#include "WidgetBlueprint.h"
+#include "Kismet2/BlueprintEditorUtils.h"
 
 FOnWidgetBlueprintCompiled FWidgetBlueprintCompiler::OnWidgetBlueprintCompiled = FOnWidgetBlueprintCompiled();
 
diff --git a/Engine/Source/Editor/UMGEditor/Private/WidgetBlueprintCompiler.h b/Engine/Source/Editor/UMGEditor/Private/WidgetBlueprintCompiler.h
--- a/Engine/Source/Editor/UMGEditor/Private/WidgetBlueprintCompiler.h
+++ b/Engine/Source/Editor/UMGEditor/Private/WidgetBlueprintCompiler.h
@@ -4,6 +4,9 @@
 
 #include "KismetCompiler.h"
 
+class UWidgetBlueprint;
+class UWidgetBlueprintGeneratedClass;
+
 //////////////////////////////////////////////////////////////////////////
 // FWidgetBlueprintCompiler
 
diff --git a/Engine/Source/Editor/UMGEditor/Private/WidgetBlueprintEditor.cpp b/Engine/Source/Editor/UMGEditor/Private/WidgetBlueprintEditor.cpp
--- a/Engine/Source/Editor/UMGEditor/Private/WidgetBlueprintEditor.cpp
+++ b/Engine/Source/Editor/UMGEditor/Private/WidgetBlueprintEditor.cpp
@@ -20,6 +20,7 @@
 #include "Blueprint/WidgetTree.h"
 #include "Components/CanvasPanel.h"
 #include "GenericCommands.h"
+#include "WidgetBlueprint.h"
 
 #define LOCTEXT_NAMESPACE "UMG"
 
diff --git a/Engine/Source/Editor/UMGEditor/Private/WidgetBlueprintEditor.h b/Engine/Source/Editor/UMGEditor/Private/WidgetBlueprintEditor.h
--- a/Engine/Source/Editor/UMGEditor/Private/WidgetBlueprintEditor.h
+++ b/Engine/Source/Editor/UMGEditor/Private/WidgetBlueprintEditor.h
@@ -8,6 +8,8 @@
 #include "WidgetReference.h"
 
 class ISequencer;
+class UWidgetAnimation;
+class UUserWidget;
 
 /**
  * widget blueprint editor (extends Blueprint editor)
diff --git a/Engine/Source/Editor/UMGEditor/Private/WidgetBlueprintEditorToolbar.cpp b/Engine/Source/Editor/UMGEditor/Private/WidgetBlueprintEditorToolbar.cpp
--- a/Engine/Source/Editor/UMGEditor/Private/WidgetBlueprintEditorToolbar.cpp
+++ b/Engine/Source/Editor/UMGEditor/Private/WidgetBlueprintEditorToolbar.cpp
@@ -6,6 +6,8 @@
 
 #include "WidgetBlueprintEditor.h"
 #include "SModeWidget.h"
+#include "Kismet2/BlueprintEditorUtils.h"
+#include "WidgetBlueprintApplicationModes.h"
 
 #define LOCTEXT_NAMESPACE "UMG"
 
diff --git a/Engine/Source/Editor/UMGEditor/Private/WidgetBlueprintEditorToolbar.h b/Engine/Source/Editor/UMGEditor/Private/WidgetBlueprintEditorToolbar.h
--- a/Engine/Source/Editor/UMGEditor/Private/WidgetBlueprintEditorToolbar.h
+++ b/Engine/Source/Editor/UMGEditor/Private/WidgetBlueprintEditorToolbar.h
@@ -2,6 +2,8 @@
 
 #pragma once
 
+class FWidgetBlueprintEditor;
+
 /**
  * Handles all of the toolbar related construction for the widget blueprint editor.
  */
diff --git a/Engine/Source/Editor/UMGEditor/Private/WidgetBlueprintEditorUtils.cpp b/Engine/Source/Editor/UMGEditor/Private/WidgetBlueprintEditorUtils.cpp
--- a/Engine/Source/Editor/UMGEditor/Private/WidgetBlueprintEditorUtils.cpp
+++ b/Engine/Source/Editor/UMGEditor/Private/WidgetBlueprintEditorUtils.cpp
@@ -12,6 +12,11 @@
 #include "Factories.h"
 #include "UnrealExporter.h"
 #include "GenericCommands.h"
+#include "WidgetBlueprint.h"
+#include "Blueprint/WidgetTree.h"
+#include "ScopedTransaction.h"
+#include "K2Node_ComponentBoundEvent.h"
+#include "CanvasPanel.h"
 
 #define LOCTEXT_NAMESPACE "UMG"
 
diff --git a/Engine/Source/Editor/UMGEditor/Private/WidgetBlueprintEditorUtils.h b/Engine/Source/Editor/UMGEditor/Private/WidgetBlueprintEditorUtils.h
--- a/Engine/Source/Editor/UMGEditor/Private/WidgetBlueprintEditorUtils.h
+++ b/Engine/Source/Editor/UMGEditor/Private/WidgetBlueprintEditorUtils.h
@@ -2,6 +2,10 @@
 
 #pragma once
 
+class UWidgetTree;
+class UWidgetBlueprint;
+struct FWidgetReference;
+
 //////////////////////////////////////////////////////////////////////////
 // FWidgetBlueprintEditorUtils
 
diff --git a/Engine/Source/Editor/UMGEditor/Private/WidgetBlueprintFactory.cpp b/Engine/Source/Editor/UMGEditor/Private/WidgetBlueprintFactory.cpp
--- a/Engine/Source/Editor/UMGEditor/Private/WidgetBlueprintFactory.cpp
+++ b/Engine/Source/Editor/UMGEditor/Private/WidgetBlueprintFactory.cpp
@@ -6,6 +6,8 @@
 #include "Kismet2/KismetEditorUtilities.h"
 
 #include "WidgetBlueprintFactory.h"
+#include "WidgetBlueprint.h"
+
 #define LOCTEXT_NAMESPACE "UWidgetBlueprintFactory"
 
 /*------------------------------------------------------------------------------
diff --git a/Engine/Source/Editor/UMGEditor/Private/WidgetReference.cpp b/Engine/Source/Editor/UMGEditor/Private/WidgetReference.cpp
--- a/Engine/Source/Editor/UMGEditor/Private/WidgetReference.cpp
+++ b/Engine/Source/Editor/UMGEditor/Private/WidgetReference.cpp
@@ -1,6 +1,9 @@
 // Copyright 1998-2014 Epic Games, Inc. All Rights Reserved.
 
 #include "UMGEditorPrivatePCH.h"
+#include "WidgetReference.h"
+#include "Blueprint/UserWidget.h"
+#include "WidgetBlueprintEditor.h"
 
 #define LOCTEXT_NAMESPACE "UMG"
 
diff --git a/Engine/Source/Editor/UMGEditor/Public/IUMGDesigner.h b/Engine/Source/Editor/UMGEditor/Public/IUMGDesigner.h
--- a/Engine/Source/Editor/UMGEditor/Public/IUMGDesigner.h
+++ b/Engine/Source/Editor/UMGEditor/Public/IUMGDesigner.h
@@ -2,6 +2,8 @@
 
 #pragma once
 
+struct FWidgetReference;
+
 /**
  * The logical type of transform that can be applied to a widget.
  */
diff --git a/Engine/Source/Editor/UMGEditor/Public/WidgetTemplate.h b/Engine/Source/Editor/UMGEditor/Public/WidgetTemplate.h
--- a/Engine/Source/Editor/UMGEditor/Public/WidgetTemplate.h
+++ b/Engine/Source/Editor/UMGEditor/Public/WidgetTemplate.h
@@ -1,6 +1,8 @@
 // Copyright 1998-2014 Epic Games, Inc. All Rights Reserved.
 
 #pragma once
+struct FSlateBrush;
+class UWidget;
 
 /**
  * The widget template represents a widget or a set of widgets to create and spawn into the widget tree.
diff --git a/Engine/Source/Runtime/SlateCore/Public/Fonts/SlateFontInfo.h b/Engine/Source/Runtime/SlateCore/Public/Fonts/SlateFontInfo.h
--- a/Engine/Source/Runtime/SlateCore/Public/Fonts/SlateFontInfo.h
+++ b/Engine/Source/Runtime/SlateCore/Public/Fonts/SlateFontInfo.h
@@ -1,7 +1,7 @@
 // Copyright 1998-2014 Epic Games, Inc. All Rights Reserved.
 
 #pragma once
-
+#include "Templates/TypeHash.h"
 #include "SlateFontInfo.generated.h"
 
 UENUM()
diff --git a/Engine/Source/Runtime/SlateCore/Public/Layout/Margin.h b/Engine/Source/Runtime/SlateCore/Public/Layout/Margin.h
--- a/Engine/Source/Runtime/SlateCore/Public/Layout/Margin.h
+++ b/Engine/Source/Runtime/SlateCore/Public/Layout/Margin.h
@@ -1,6 +1,7 @@
 // Copyright 1998-2014 Epic Games, Inc. All Rights Reserved.
 
 #pragma once
+#include "Math/Vector2D.h"
 #include "Types/SlateEnums.h"
 #include "Margin.generated.h"
 
diff --git a/Engine/Source/Runtime/SlateCore/Public/Styling/ISlateStyle.h b/Engine/Source/Runtime/SlateCore/Public/Styling/ISlateStyle.h
--- a/Engine/Source/Runtime/SlateCore/Public/Styling/ISlateStyle.h
+++ b/Engine/Source/Runtime/SlateCore/Public/Styling/ISlateStyle.h
@@ -4,10 +4,10 @@
 
 struct FMargin;
 struct FSlateColor;
+struct FSlateDynamicImageBrush;
 struct FSlateFontInfo;
 struct FSlateSound;
 struct FSlateWidgetStyle;
-
 /**
  * 
  */
diff --git a/Engine/Source/Runtime/SlateCore/Public/Styling/StyleDefaults.h b/Engine/Source/Runtime/SlateCore/Public/Styling/StyleDefaults.h
--- a/Engine/Source/Runtime/SlateCore/Public/Styling/StyleDefaults.h
+++ b/Engine/Source/Runtime/SlateCore/Public/Styling/StyleDefaults.h
@@ -2,6 +2,7 @@
 #pragma once
 #include "Brushes/SlateNoResource.h"
 #include "Fonts/SlateFontInfo.h"
+#include "Misc/Paths.h"
 
 struct FMargin;
 struct FSlateSound;
diff --git a/Engine/Source/Runtime/SlateCore/Public/Styling/WidgetStyle.h b/Engine/Source/Runtime/SlateCore/Public/Styling/WidgetStyle.h
--- a/Engine/Source/Runtime/SlateCore/Public/Styling/WidgetStyle.h
+++ b/Engine/Source/Runtime/SlateCore/Public/Styling/WidgetStyle.h
@@ -1,7 +1,11 @@
 // Copyright 1998-2014 Epic Games, Inc. All Rights Reserved.
 
 #pragma once
+#include "HAL/Platform.h"
+#include "Math/Color.h"
 
+template< typename ObjectType > class TAttribute;
+struct FSlateColor;
 
 /**
  * Contains info about those aspects of widget appearance that should be propagated hierarchically.
@@ -57,7 +61,7 @@ public:
 	 * @return This instance (for method chaining).
 	 * @see GetForegroundColor
 	 */
-	FWidgetStyle& SetForegroundColor( const TAttribute<struct FSlateColor>& InForeground );
+	FWidgetStyle& SetForegroundColor( const TAttribute<FSlateColor>& InForeground );
 
 public:
 
diff --git a/Engine/Source/Runtime/UMG/Public/Animation/MovieScene2DTransformSection.h b/Engine/Source/Runtime/UMG/Public/Animation/MovieScene2DTransformSection.h
--- a/Engine/Source/Runtime/UMG/Public/Animation/MovieScene2DTransformSection.h
+++ b/Engine/Source/Runtime/UMG/Public/Animation/MovieScene2DTransformSection.h
@@ -6,6 +6,8 @@
 
 #include "MovieScene2DTransformSection.generated.h"
 
+struct FWidgetTransform;
+
 /**
  * A transform section
  */
diff --git a/Engine/Source/Runtime/UMG/Public/Animation/MovieScene2DTransformTrack.h b/Engine/Source/Runtime/UMG/Public/Animation/MovieScene2DTransformTrack.h
--- a/Engine/Source/Runtime/UMG/Public/Animation/MovieScene2DTransformTrack.h
+++ b/Engine/Source/Runtime/UMG/Public/Animation/MovieScene2DTransformTrack.h
@@ -3,6 +3,7 @@
 #pragma once
 
 #include "MovieScenePropertyTrack.h"
+#include "Slate/WidgetTransform.h"
 #include "MovieScene2DTransformTrack.generated.h"
 
 struct F2DTransformKey
diff --git a/Engine/Source/Runtime/UMG/Public/Animation/WidgetAnimation.h b/Engine/Source/Runtime/UMG/Public/Animation/WidgetAnimation.h
--- a/Engine/Source/Runtime/UMG/Public/Animation/WidgetAnimation.h
+++ b/Engine/Source/Runtime/UMG/Public/Animation/WidgetAnimation.h
@@ -6,7 +6,7 @@
 #include "WidgetAnimation.generated.h"
 
 class UMovieScene;
-
+class UWidgetTree;
 
 /** A single object bound to a umg sequence */
 USTRUCT()
diff --git a/Engine/Source/Runtime/UMG/Public/Blueprint/WidgetBlueprintLibrary.h b/Engine/Source/Runtime/UMG/Public/Blueprint/WidgetBlueprintLibrary.h
--- a/Engine/Source/Runtime/UMG/Public/Blueprint/WidgetBlueprintLibrary.h
+++ b/Engine/Source/Runtime/UMG/Public/Blueprint/WidgetBlueprintLibrary.h
@@ -1,9 +1,15 @@
 // Copyright 1998-2014 Epic Games, Inc. All Rights Reserved.
 
 #pragma once
-
+#include "Kismet/BlueprintFunctionLibrary.h"
+#include "Components/SlateWrapperTypes.h"
 #include "WidgetBlueprintLibrary.generated.h"
 
+class UDragDropOperation;
+class USlateBrushAsset;
+
+struct FPaintContext;
+
 UCLASS(MinimalAPI)
 class UWidgetBlueprintLibrary : public UBlueprintFunctionLibrary
 {
diff --git a/Engine/Source/Runtime/UMG/Public/Components/CanvasPanelSlot.h b/Engine/Source/Runtime/UMG/Public/Components/CanvasPanelSlot.h
--- a/Engine/Source/Runtime/UMG/Public/Components/CanvasPanelSlot.h
+++ b/Engine/Source/Runtime/UMG/Public/Components/CanvasPanelSlot.h
@@ -3,6 +3,7 @@
 #pragma once
 
 #include "Anchors.h"
+#include "Components/PanelSlot.h"
 #include "Slate/SConstraintCanvas.h"
 #include "CanvasPanelSlot.generated.h"
 
diff --git a/Engine/Source/Runtime/UMG/Public/Components/PanelSlot.h b/Engine/Source/Runtime/UMG/Public/Components/PanelSlot.h
--- a/Engine/Source/Runtime/UMG/Public/Components/PanelSlot.h
+++ b/Engine/Source/Runtime/UMG/Public/Components/PanelSlot.h
@@ -1,7 +1,7 @@
 // Copyright 1998-2014 Epic Games, Inc. All Rights Reserved.
 
 #pragma once
-
+#include "Components/Visual.h"
 #include "PanelSlot.generated.h"
 
 /** The base class for all Slots in UMG. */
