UMGEditor PCH cleanup prep.

diff --git a/Engine/Source/Editor/UMGEditor/Private/Customizations/WidgetNavigationCustomization.cpp b/Engine/Source/Editor/UMGEditor/Private/Customizations/WidgetNavigationCustomization.cpp
--- a/Engine/Source/Editor/UMGEditor/Private/Customizations/WidgetNavigationCustomization.cpp
+++ b/Engine/Source/Editor/UMGEditor/Private/Customizations/WidgetNavigationCustomization.cpp
@@ -7,7 +7,7 @@
 #include "WidgetGraphSchema.h"
 #include "ScopedTransaction.h"
 #include "BlueprintEditorUtils.h"
-
+#include "Blueprint/WidgetNavigation.h"
 #include "WidgetNavigationCustomization.h"
 
 #define LOCTEXT_NAMESPACE "UMG"
