UMG - Adding a step during widget blueprint compiling that verifies all bindings are valid before adding them to the runtime bindings list, this will prevent any invalidly bound signatures or impure functions being linked against.

diff --git a/Engine/Source/Editor/UMGEditor/Private/Details/SPropertyBinding.cpp b/Engine/Source/Editor/UMGEditor/Private/Details/SPropertyBinding.cpp
--- a/Engine/Source/Editor/UMGEditor/Private/Details/SPropertyBinding.cpp
+++ b/Engine/Source/Editor/UMGEditor/Private/Details/SPropertyBinding.cpp
@@ -72,7 +72,7 @@ void SPropertyBinding::Construct(const FArguments& InArgs, TSharedRef<FWidgetBlu
 	];
 }
 
-void SPropertyBinding::RefreshBlueprintMemberCache(const UFunction* DelegateSignature)
+void SPropertyBinding::RefreshBlueprintMemberCache(const UFunction* DelegateSignature, bool bIsPure)
 {
 	const UWidgetGraphSchema* Schema = GetDefault<UWidgetGraphSchema>();
 	const FSlateFontInfo DetailFontInfo = IDetailLayoutBuilder::GetDetailFont();
@@ -98,7 +98,7 @@ void SPropertyBinding::RefreshBlueprintMemberCache(const UFunction* DelegateSign
 		FName FunctionFName = Graph->GetFName();
 
 		UFunction* Function = SkeletonClass->FindFunctionByName(FunctionFName, EIncludeSuperFlag::IncludeSuper);
-		if ( Function == NULL )
+		if ( Function == nullptr )
 		{
 			continue;
 		}
@@ -106,6 +106,12 @@ void SPropertyBinding::RefreshBlueprintMemberCache(const UFunction* DelegateSign
 		// We ignore CPF_ReturnParm because all that matters for binding to script functions is that the number of out parameters match.
 		if ( Function->IsSignatureCompatibleWith(DelegateSignature, UFunction::GetDefaultIgnoredSignatureCompatibilityFlags() | CPF_ReturnParm) )
 		{
+			// Only allow binding pure functions if we're limited to pure function bindings.
+			if ( bIsPure && !Function->HasAllFunctionFlags(FUNC_BlueprintPure) )
+			{
+				continue;
+			}
+			
 			FGraphDisplayInfo DisplayInfo;
 			Graph->GetSchema()->GetGraphDisplayInformation(*Graph, DisplayInfo);
 
@@ -139,7 +145,7 @@ void SPropertyBinding::RefreshBlueprintMemberCache(const UFunction* DelegateSign
 
 TSharedRef<SWidget> SPropertyBinding::OnGenerateDelegateMenu(UWidget* Widget, TSharedRef<IPropertyHandle> PropertyHandle, UFunction* DelegateSignature, bool bIsPure)
 {
-	RefreshBlueprintMemberCache(DelegateSignature);
+	RefreshBlueprintMemberCache(DelegateSignature, bIsPure);
 
 	const bool bInShouldCloseWindowAfterMenuSelection = true;
 	FMenuBuilder MenuBuilder(bInShouldCloseWindowAfterMenuSelection, NULL);
diff --git a/Engine/Source/Editor/UMGEditor/Private/Details/SPropertyBinding.h b/Engine/Source/Editor/UMGEditor/Private/Details/SPropertyBinding.h
--- a/Engine/Source/Editor/UMGEditor/Private/Details/SPropertyBinding.h
+++ b/Engine/Source/Editor/UMGEditor/Private/Details/SPropertyBinding.h
@@ -26,7 +26,7 @@ protected:
 		UEdGraph* EdGraph;
 	};
 
-	void RefreshBlueprintMemberCache(const UFunction* DelegateSignature);
+	void RefreshBlueprintMemberCache(const UFunction* DelegateSignature, bool bIsPure);
 
 	TSharedRef<SWidget> OnGenerateDelegateMenu(UWidget* Widget, TSharedRef<IPropertyHandle> PropertyHandle, UFunction* DelegateSignature, bool bIsPure);
 
diff --git a/Engine/Source/Editor/UMGEditor/Private/WidgetBlueprint.cpp b/Engine/Source/Editor/UMGEditor/Private/WidgetBlueprint.cpp
--- a/Engine/Source/Editor/UMGEditor/Private/WidgetBlueprint.cpp
+++ b/Engine/Source/Editor/UMGEditor/Private/WidgetBlueprint.cpp
@@ -9,6 +9,48 @@
 
 #define LOCTEXT_NAMESPACE "UMG"
 
+bool FDelegateEditorBinding::IsBindingValid(UClass* BlueprintGeneratedClass, UWidgetBlueprint* Blueprint) const
+{
+	FDelegateRuntimeBinding RuntimeBinding = ToRuntimeBinding(Blueprint);
+
+	// First find the target widget we'll be attaching the binding to.
+	if ( UWidget* TargetWidget = Blueprint->WidgetTree->FindWidget(FName(*ObjectName)) )
+	{
+		// Next find the underlying delegate we're actually binding to, if it's an event the name will be the same,
+		// for properties we need to lookup the delegate property we're actually going to be binding to.
+		UDelegateProperty* BindableProperty = FindField<UDelegateProperty>(TargetWidget->GetClass(), FName(*( PropertyName.ToString() + TEXT("Delegate") )));
+		UDelegateProperty* EventProperty = FindField<UDelegateProperty>(TargetWidget->GetClass(), PropertyName);
+
+		bool bNeedsToBePure = BindableProperty ? true : false;
+		UDelegateProperty* DelegateProperty = BindableProperty ? BindableProperty : EventProperty;
+
+		// Locate the delegate property on the widget that's a delegate for a property we want to bind.
+		if ( DelegateProperty )
+		{
+			// On our incoming blueprint generated class, try and find the function we claim exists that users
+			// are binding their property too.
+			if ( UFunction* Function = BlueprintGeneratedClass->FindFunctionByName(RuntimeBinding.FunctionName, EIncludeSuperFlag::IncludeSuper) )
+			{
+				// Check the signatures to ensure these functions match.
+				if ( Function->IsSignatureCompatibleWith(DelegateProperty->SignatureFunction, UFunction::GetDefaultIgnoredSignatureCompatibilityFlags() | CPF_ReturnParm) )
+				{
+					if ( bNeedsToBePure )
+					{
+						// Only allow binding pure functions to property bindings.
+						return Function->HasAllFunctionFlags(FUNC_BlueprintPure);
+					}
+					else
+					{
+						return true;
+					}
+				}
+			}
+		}
+	}
+
+	return false;
+}
+
 FDelegateRuntimeBinding FDelegateEditorBinding::ToRuntimeBinding(UWidgetBlueprint* Blueprint) const
 {
 	FDelegateRuntimeBinding Binding;
@@ -20,7 +62,7 @@ FDelegateRuntimeBinding FDelegateEditorBinding::ToRuntimeBinding(UWidgetBlueprin
 	}
 	else
 	{
-		Binding.FunctionName = FunctionName;// Blueprint->GetFieldNameFromClassByGuid<UProperty>(Blueprint->SkeletonGeneratedClass, MemberGuid);
+		Binding.FunctionName = FunctionName;
 	}
 	Binding.Kind = Kind;
 
diff --git a/Engine/Source/Editor/UMGEditor/Private/WidgetBlueprintCompiler.cpp b/Engine/Source/Editor/UMGEditor/Private/WidgetBlueprintCompiler.cpp
--- a/Engine/Source/Editor/UMGEditor/Private/WidgetBlueprintCompiler.cpp
+++ b/Engine/Source/Editor/UMGEditor/Private/WidgetBlueprintCompiler.cpp
@@ -265,7 +265,10 @@ void FWidgetBlueprintCompiler::FinishCompilingClass(UClass* Class)
 	// @todo UMG Possibly need duplication here 
 	for ( const FDelegateEditorBinding& EditorBinding : Blueprint->Bindings )
 	{
-		BPGClass->Bindings.Add(EditorBinding.ToRuntimeBinding(Blueprint));
+		if ( EditorBinding.IsBindingValid(Class, Blueprint) )
+		{
+			BPGClass->Bindings.Add(EditorBinding.ToRuntimeBinding(Blueprint));
+		}
 	}
 
 	Super::FinishCompilingClass(Class);
diff --git a/Engine/Source/Editor/UMGEditor/Public/WidgetBlueprint.h b/Engine/Source/Editor/UMGEditor/Public/WidgetBlueprint.h
--- a/Engine/Source/Editor/UMGEditor/Public/WidgetBlueprint.h
+++ b/Engine/Source/Editor/UMGEditor/Public/WidgetBlueprint.h
@@ -40,6 +40,8 @@ struct UMGEDITOR_API FDelegateEditorBinding
 		return ObjectName == Other.ObjectName && PropertyName == Other.PropertyName;
 	}
 
+	bool IsBindingValid(UClass* Class, class UWidgetBlueprint* Blueprint) const;
+
 	FDelegateRuntimeBinding ToRuntimeBinding(class UWidgetBlueprint* Blueprint) const;
 };
 
