UMG - Instead of the slot related properties appearing under the "Layout (SLOT TYPE)" category they now appear under "Slot (SLOT TYPE)".

diff --git a/Engine/Source/Editor/UMGEditor/Private/Customizations/DetailCustomizations.cpp b/Engine/Source/Editor/UMGEditor/Private/Customizations/DetailCustomizations.cpp
--- a/Engine/Source/Editor/UMGEditor/Private/Customizations/DetailCustomizations.cpp
+++ b/Engine/Source/Editor/UMGEditor/Private/Customizations/DetailCustomizations.cpp
@@ -339,7 +339,7 @@ void FBlueprintWidgetCustomization::CustomizeDetails( IDetailLayoutBuilder& Deta
 			if ( Widget->Slot )
 			{
 				UClass* SlotClass = Widget->Slot->GetClass();
-				FString LayoutCatName = FString("Layout (") + SlotClass->GetDisplayNameText().ToString() + FString(")");
+				FString LayoutCatName = FString(TEXT("Slot (")) + SlotClass->GetDisplayNameText().ToString() + FString(TEXT(")"));
 
 				DetailLayout.EditCategory(LayoutCategoryKey, LayoutCatName, ECategoryPriority::TypeSpecific);
 			}
