UMG - Fixing property binding, to correctly apply the pure flag to new function bindings.

diff --git a/Engine/Source/Editor/UMGEditor/Private/Details/SPropertyBinding.cpp b/Engine/Source/Editor/UMGEditor/Private/Details/SPropertyBinding.cpp
--- a/Engine/Source/Editor/UMGEditor/Private/Details/SPropertyBinding.cpp
+++ b/Engine/Source/Editor/UMGEditor/Private/Details/SPropertyBinding.cpp
@@ -421,15 +421,12 @@ void SPropertyBinding::HandleCreateAndAddBinding(UWidget* Widget, TSharedRef<IPr
 
 	// Create the function graph.
 	FString FunctionName = Pre + WidgetName + Post;
-	UEdGraph* FunctionGraph = FBlueprintEditorUtils::CreateNewGraph(Blueprint, FBlueprintEditorUtils::FindUniqueKismetName(Blueprint, FunctionName), UEdGraph::StaticClass(), UEdGraphSchema_K2::StaticClass());
-
-	// Only mark bindings as pure that need to be.
-	if ( bIsPure )
-	{
-		const UEdGraphSchema_K2* Schema_K2 = Cast<UEdGraphSchema_K2>(FunctionGraph->GetSchema());
-		Schema_K2->AddExtraFunctionFlags(FunctionGraph, FUNC_BlueprintPure);
-	}
-
+	UEdGraph* FunctionGraph = FBlueprintEditorUtils::CreateNewGraph(
+		Blueprint, 
+		FBlueprintEditorUtils::FindUniqueKismetName(Blueprint, FunctionName),
+		UEdGraph::StaticClass(),
+		UEdGraphSchema_K2::StaticClass());
+	
 	// Add the binding to the blueprint
 	TSharedPtr<FunctionInfo> SelectedFunction = MakeShareable(new FunctionInfo());
 	SelectedFunction->FuncName = FunctionGraph->GetFName();
@@ -437,9 +434,16 @@ void SPropertyBinding::HandleCreateAndAddBinding(UWidget* Widget, TSharedRef<IPr
 
 	HandleAddFunctionBinding(PropertyHandle, SelectedFunction);
 
-	bool bUserCreated = true;
+	const bool bUserCreated = true;
 	FBlueprintEditorUtils::AddFunctionGraph(Blueprint, FunctionGraph, bUserCreated, DelegateSignature);
 
+	// Only mark bindings as pure that need to be.
+	if ( bIsPure )
+	{
+		const UEdGraphSchema_K2* Schema_K2 = Cast<UEdGraphSchema_K2>(FunctionGraph->GetSchema());
+		Schema_K2->AddExtraFunctionFlags(FunctionGraph, FUNC_BlueprintPure);
+	}
+
 	GotoFunction(FunctionGraph);
 }
 
