UMG - Adding errors for function binding compiling errors.  Fixing an issue with auto generated getters for property bindings not generating pure functions.

diff --git a/Engine/Source/Editor/UMGEditor/Private/WidgetBlueprint.cpp b/Engine/Source/Editor/UMGEditor/Private/WidgetBlueprint.cpp
--- a/Engine/Source/Editor/UMGEditor/Private/WidgetBlueprint.cpp
+++ b/Engine/Source/Editor/UMGEditor/Private/WidgetBlueprint.cpp
@@ -9,7 +9,7 @@
 
 #define LOCTEXT_NAMESPACE "UMG"
 
-bool FDelegateEditorBinding::IsBindingValid(UClass* BlueprintGeneratedClass, UWidgetBlueprint* Blueprint) const
+bool FDelegateEditorBinding::IsBindingValid(UClass* BlueprintGeneratedClass, UWidgetBlueprint* Blueprint, FCompilerResultsLog& MessageLog) const
 {
 	FDelegateRuntimeBinding RuntimeBinding = ToRuntimeBinding(Blueprint);
 
@@ -34,19 +34,37 @@ bool FDelegateEditorBinding::IsBindingValid(UClass* BlueprintGeneratedClass, UWi
 				// Check the signatures to ensure these functions match.
 				if ( Function->IsSignatureCompatibleWith(DelegateProperty->SignatureFunction, UFunction::GetDefaultIgnoredSignatureCompatibilityFlags() | CPF_ReturnParm) )
 				{
-					if ( bNeedsToBePure )
+					// Only allow binding pure functions to property bindings.
+					if ( bNeedsToBePure && !Function->HasAllFunctionFlags(FUNC_BlueprintPure) )
 					{
-						// Only allow binding pure functions to property bindings.
-						return Function->HasAllFunctionFlags(FUNC_BlueprintPure);
-					}
-					else
-					{
-						return true;
+						FText const ErrorFormat = LOCTEXT("BindingNotBoundToPure", "Binding: property '@@' on widget '@@' needs to be bound to a pure function, '@@' is not pure.");
+						MessageLog.Error(*ErrorFormat.ToString(), DelegateProperty, TargetWidget, Function);
+
+						return false;
 					}
+
+					return true;
+				}
+				else
+				{
+					FText const ErrorFormat = LOCTEXT("BindingFunctionSigDontMatch", "Binding: property '@@' on widget '@@' bound to function '@@', but the sigatnures don't match.  The function must return the same type as the property and have no parameters.");
+					MessageLog.Error(*ErrorFormat.ToString(), DelegateProperty, TargetWidget, Function);
 				}
 			}
+			else
+			{
+				//TODO Bindable property removed.
+			}
+		}
+		else
+		{
+			// TODO Bindable Property Removed
 		}
 	}
+	else
+	{
+		//TODO Ignore missing widgets
+	}
 
 	return false;
 }
diff --git a/Engine/Source/Editor/UMGEditor/Private/WidgetBlueprintCompiler.cpp b/Engine/Source/Editor/UMGEditor/Private/WidgetBlueprintCompiler.cpp
--- a/Engine/Source/Editor/UMGEditor/Private/WidgetBlueprintCompiler.cpp
+++ b/Engine/Source/Editor/UMGEditor/Private/WidgetBlueprintCompiler.cpp
@@ -47,13 +47,6 @@ void FWidgetBlueprintCompiler::CreateFunctionList()
 
 			Schema->CreateDefaultNodesForGraph(*FunctionGraph);
 
-			// We need to flag the entry node to make sure that the compiled function is callable from Kismet2
-			int32 ExtraFunctionFlags = ( FUNC_BlueprintCallable | FUNC_BlueprintEvent | FUNC_Public );
-			if ( BPTYPE_FunctionLibrary == Blueprint->BlueprintType )
-			{
-				ExtraFunctionFlags |= FUNC_Static;
-			}
-			K2Schema->AddExtraFunctionFlags(FunctionGraph, ExtraFunctionFlags);
 			K2Schema->MarkFunctionEntryAsEditable(FunctionGraph, true);
 
 			// Create a function entry node
@@ -88,6 +81,10 @@ void FWidgetBlueprintCompiler::CreateFunctionList()
 
 			ReturnPin->MakeLinkTo(VarNode->GetValuePin());
 
+			// We need to flag the entry node to make sure that the compiled function is callable from Kismet2
+			int32 ExtraFunctionFlags = ( FUNC_BlueprintCallable | FUNC_BlueprintEvent | FUNC_Public | FUNC_BlueprintPure );
+			K2Schema->AddExtraFunctionFlags(FunctionGraph, ExtraFunctionFlags);
+
 			//Blueprint->FunctionGraphs.Add(FunctionGraph);
 
 			ProcessOneFunctionGraph(FunctionGraph);
@@ -262,10 +259,11 @@ void FWidgetBlueprintCompiler::FinishCompilingClass(UClass* Class)
 
 	BPGClass->Bindings.Reset();
 
-	// @todo UMG Possibly need duplication here 
+	// Convert all editor time property bindings into a list of bindings
+	// that will be applied at runtime.  Ensure all bindings are still valid.
 	for ( const FDelegateEditorBinding& EditorBinding : Blueprint->Bindings )
 	{
-		if ( EditorBinding.IsBindingValid(Class, Blueprint) )
+		if ( EditorBinding.IsBindingValid(Class, Blueprint, MessageLog) )
 		{
 			BPGClass->Bindings.Add(EditorBinding.ToRuntimeBinding(Blueprint));
 		}
diff --git a/Engine/Source/Editor/UMGEditor/Public/WidgetBlueprint.h b/Engine/Source/Editor/UMGEditor/Public/WidgetBlueprint.h
--- a/Engine/Source/Editor/UMGEditor/Public/WidgetBlueprint.h
+++ b/Engine/Source/Editor/UMGEditor/Public/WidgetBlueprint.h
@@ -40,7 +40,7 @@ struct UMGEDITOR_API FDelegateEditorBinding
 		return ObjectName == Other.ObjectName && PropertyName == Other.PropertyName;
 	}
 
-	bool IsBindingValid(UClass* Class, class UWidgetBlueprint* Blueprint) const;
+	bool IsBindingValid(UClass* Class, class UWidgetBlueprint* Blueprint, FCompilerResultsLog& MessageLog) const;
 
 	FDelegateRuntimeBinding ToRuntimeBinding(class UWidgetBlueprint* Blueprint) const;
 };
