UMG - Now displaying the currently applied DPI Scale in the UMG designer, along with a button that brings up the settings for the DPI Curve.  Moved the UI related settings out of URendererSettings, they're now in thier own settings class UUserInterfaceSettings.

diff --git a/Engine/Config/BaseEngine.ini b/Engine/Config/BaseEngine.ini
--- a/Engine/Config/BaseEngine.ini
+++ b/Engine/Config/BaseEngine.ini
@@ -1364,7 +1364,7 @@ InGameAdManagerClassName=
 RenderPriorityPS3=1001
 SuspendGameIO=True
 
-[/Script/Engine.RendererSettings]
+[/Script/Engine.UserInterfaceSettings]
 UIScaleRule=ShortestSide
 UIScaleCurve=(EditorCurveData=(Keys=((Time=480,Value=0.444),(Time=720,Value=0.666),(Time=1080,Value=1.0),(Time=8640,Value=8.0))),ExternalCurve=None)
 
diff --git a/Engine/Source/Editor/EditorStyle/Private/SlateEditorStyle.cpp b/Engine/Source/Editor/EditorStyle/Private/SlateEditorStyle.cpp
--- a/Engine/Source/Editor/EditorStyle/Private/SlateEditorStyle.cpp
+++ b/Engine/Source/Editor/EditorStyle/Private/SlateEditorStyle.cpp
@@ -6052,6 +6052,8 @@ void FSlateEditorStyle::FStyle::SetupUMGEditorStyles()
 	Set("UMGEditor.SelectionOutline", new BORDER_BRUSH("UMG/SelectionBorder", FMargin(.2f) ));
 
 	Set("UMGEditor.AnchorGrid", new IMAGE_BRUSH("Icons/UMG/AnchorGrid", Icon10x10, FLinearColor::White, ESlateBrushTileType::Both ));
+
+	Set("UMGEditor.DPISettings", new IMAGE_BRUSH("Icons/UMG/SettingsButton", Icon16x16));
 }
 
 void FSlateEditorStyle::FStyle::SetupTranslationEditorStyles()
diff --git a/Engine/Source/Editor/LevelEditor/Private/SLevelViewport.cpp b/Engine/Source/Editor/LevelEditor/Private/SLevelViewport.cpp
--- a/Engine/Source/Editor/LevelEditor/Private/SLevelViewport.cpp
+++ b/Engine/Source/Editor/LevelEditor/Private/SLevelViewport.cpp
@@ -29,6 +29,7 @@
 #include "SceneOutlinerTreeItems.h"
 #include "BufferVisualizationData.h"
 #include "EditorViewportCommands.h"
+#include "Runtime/Engine/Classes/Engine/UserInterfaceSettings.h"
 #include "Runtime/Engine/Classes/Engine/RendererSettings.h"
 #include "SScissorRectBox.h"
 #include "SDPIScaler.h"
@@ -374,7 +375,7 @@ float SLevelViewport::GetGameViewportDPIScale() const
 {
 	if ( HasPlayInEditorViewport() || LevelViewportClient->IsSimulateInEditorViewport() )
 	{
-		return GetDefault<URendererSettings>(URendererSettings::StaticClass())->GetDPIScaleBasedOnSize( ActiveViewport->GetSize() );
+		return GetDefault<UUserInterfaceSettings>(UUserInterfaceSettings::StaticClass())->GetDPIScaleBasedOnSize(ActiveViewport->GetSize());
 	}
 
 	return 1.0f;
diff --git a/Engine/Source/Editor/ProjectSettingsViewer/Private/ProjectSettingsViewerModule.cpp b/Engine/Source/Editor/ProjectSettingsViewer/Private/ProjectSettingsViewerModule.cpp
--- a/Engine/Source/Editor/ProjectSettingsViewer/Private/ProjectSettingsViewerModule.cpp
+++ b/Engine/Source/Editor/ProjectSettingsViewer/Private/ProjectSettingsViewerModule.cpp
@@ -4,6 +4,7 @@
 #include "Engine/Console.h"
 #include "ProjectTargetPlatformEditor.h"
 #include "CookerSettings.h"
+#include "Runtime/Engine/Classes/Engine/UserInterfaceSettings.h"
 #include "Runtime/Engine/Classes/Engine/RendererSettings.h"
 #include "Runtime/Engine/Classes/Sound/AudioSettings.h"
 #include "SDockTab.h"
@@ -156,7 +157,14 @@ protected:
 			GetMutableDefault<URendererSettings>()
 		);
 
-		// Rendering settings
+		// UI settings
+		SettingsModule.RegisterSettings("Project", "Engine", "UI",
+			LOCTEXT("EngineUISettingsName", "User Interface"),
+			LOCTEXT("ProjectUISettingsDescription", "User Interface settings that control Slate and UMG."),
+			GetMutableDefault<UUserInterfaceSettings>()
+		);
+
+		// Cooker settings
 		SettingsModule.RegisterSettings("Project", "Engine", "Cooker",
 			LOCTEXT("CookerSettingsName", "Cooker"),
 			LOCTEXT("CookerSettingsDescription", "Various cooker settings."),
diff --git a/Engine/Source/Editor/UMGEditor/Private/Designer/SDesignerView.cpp b/Engine/Source/Editor/UMGEditor/Private/Designer/SDesignerView.cpp
--- a/Engine/Source/Editor/UMGEditor/Private/Designer/SDesignerView.cpp
+++ b/Engine/Source/Editor/UMGEditor/Private/Designer/SDesignerView.cpp
@@ -2,6 +2,8 @@
 
 #include "UMGEditorPrivatePCH.h"
 
+#include "Settings.h"
+
 #include "Designer/DesignTimeUtils.h"
 
 #include "Extensions/CanvasSlotExtension.h"
@@ -21,6 +23,7 @@
 #include "SDesignerToolBar.h"
 #include "DesignerCommands.h"
 #include "STransformHandle.h"
+#include "Runtime/Engine/Classes/Engine/UserInterfaceSettings.h"
 #include "Runtime/Engine/Classes/Engine/RendererSettings.h"
 #include "SDPIScaler.h"
 #include "SNumericEntryBox.h"
@@ -135,11 +138,11 @@ const FString SDesignerView::DefaultAspectRatio = "16:9";
 
 void SDesignerView::Construct(const FArguments& InArgs, TSharedPtr<FWidgetBlueprintEditor> InBlueprintEditor)
 {
-	ScopedTransaction = NULL;
+	ScopedTransaction = nullptr;
 
-	PreviewWidget = NULL;
-	DropPreviewWidget = NULL;
-	DropPreviewParent = NULL;
+	PreviewWidget = nullptr;
+	DropPreviewWidget = nullptr;
+	DropPreviewParent = nullptr;
 	BlueprintEditor = InBlueprintEditor;
 
 	DesignerMessage = EDesignerMessage::None;
@@ -298,7 +301,7 @@ void SDesignerView::Construct(const FArguments& InArgs, TSharedPtr<FWidgetBluepr
 				[
 					SNew(SNumericEntryBox<int32>)
 					.AllowSpin(true)
-					.Delta(5)
+					.Delta(1)
 					.MinSliderValue(1)
 					.MinValue(1)
 					.MaxSliderValue(TOptional<int32>(1000))
@@ -320,7 +323,7 @@ void SDesignerView::Construct(const FArguments& InArgs, TSharedPtr<FWidgetBluepr
 				[
 					SNew(SNumericEntryBox<int32>)
 					.AllowSpin(true)
-					.Delta(5)
+					.Delta(1)
 					.MinSliderValue(1)
 					.MaxSliderValue(TOptional<int32>(1000))
 					.MinValue(1)
@@ -336,12 +339,35 @@ void SDesignerView::Construct(const FArguments& InArgs, TSharedPtr<FWidgetBluepr
 				]
 			]
 
+			// Info Bar, displays heads up information about some actions.
+			+ SOverlay::Slot()
+			.HAlign(HAlign_Fill)
+			.VAlign(VAlign_Bottom)
+			[
+				SNew(SDisappearingBar)
+				[
+					SNew(SBorder)
+					.BorderImage(FEditorStyle::GetBrush("WhiteBrush"))
+					.BorderBackgroundColor(FLinearColor(0.10, 0.10, 0.10, 0.75))
+					.HAlign(HAlign_Center)
+					.VAlign(VAlign_Center)
+					.Padding(FMargin(0, 5))
+					.Visibility(this, &SDesignerView::GetInfoBarVisibility)
+					[
+						SNew(STextBlock)
+						.TextStyle(FEditorStyle::Get(), "Graph.ZoomText")
+						.Text(this, &SDesignerView::GetInfoBarText)
+					]
+				]
+			]
+
 			// Bottom bar to show current resolution & AR
 			+ SOverlay::Slot()
 			.HAlign(HAlign_Fill)
 			.VAlign(VAlign_Bottom)
 			[
 				SNew(SHorizontalBox)
+
 				+ SHorizontalBox::Slot()
 				.AutoWidth()
 				.Padding(6, 0, 0, 2)
@@ -351,26 +377,38 @@ void SDesignerView::Construct(const FArguments& InArgs, TSharedPtr<FWidgetBluepr
 					.Text(this, &SDesignerView::GetCurrentResolutionText)
 					.ColorAndOpacity(this, &SDesignerView::GetResolutionTextColorAndOpacity)
 				]
-			]
 
-			// Info Bar, displays heads up information about some actions.
-			+ SOverlay::Slot()
-			.HAlign(HAlign_Fill)
-			.VAlign(VAlign_Bottom)
-			[
-				SNew(SDisappearingBar)
+				+ SHorizontalBox::Slot()
+				.FillWidth(1.0f)
+				.HAlign(HAlign_Right)
+				.Padding(0, 0, 6, 2)
 				[
-					SNew(SBorder)
-					.BorderImage(FEditorStyle::GetBrush("WhiteBrush"))
-					.BorderBackgroundColor(FLinearColor(0.10, 0.10, 0.10, 0.75))
-					.HAlign(HAlign_Center)
-					.VAlign(VAlign_Center)
-					.Padding(FMargin(0, 5))
-					.Visibility(this, &SDesignerView::GetInfoBarVisibility)
+					SNew(SHorizontalBox)
+
+					+ SHorizontalBox::Slot()
+					.AutoWidth()
 					[
 						SNew(STextBlock)
 						.TextStyle(FEditorStyle::Get(), "Graph.ZoomText")
-						.Text(this, &SDesignerView::GetInfoBarText)
+						.Text(this, &SDesignerView::GetCurrentDPIScaleText)
+						.ColorAndOpacity(FLinearColor(1, 1, 1, 0.25f))
+					]
+
+					+ SHorizontalBox::Slot()
+					.AutoWidth()
+					.Padding(6, 0, 0, 0)
+					[
+						SNew(SButton)
+						.ButtonStyle(FEditorStyle::Get(), "HoverHintOnly")
+						.ContentPadding(FMargin(3, 1))
+						.OnClicked(this, &SDesignerView::HandleDPISettingsClicked)
+						.ToolTipText(LOCTEXT("DPISettingsTooltip", "Configure the UI Scale Curve to control how the UI is scaled on different resolutions."))
+						.HAlign(HAlign_Center)
+						.VAlign(VAlign_Center)
+						[
+							SNew(SImage)
+							.Image(FEditorStyle::GetBrush("UMGEditor.DPISettings"))
+						]
 					]
 				]
 			]
@@ -518,7 +556,7 @@ float SDesignerView::GetPreviewDPIScale() const
 		}
 	}
 
-	return GetDefault<URendererSettings>(URendererSettings::StaticClass())->GetDPIScaleBasedOnSize(FIntPoint(PreviewWidth, PreviewHeight));
+	return GetDefault<UUserInterfaceSettings>(UUserInterfaceSettings::StaticClass())->GetDPIScaleBasedOnSize(FIntPoint(PreviewWidth, PreviewHeight));
 }
 
 FSlateRect SDesignerView::ComputeAreaBounds() const
@@ -703,7 +741,7 @@ FVector2D SDesignerView::GetExtensionPosition(TSharedRef<FDesignerSurfaceElement
 		FWidgetReference ParentRef = BlueprintEditor.Pin()->GetReferenceFromTemplate(SelectedWidget.GetTemplate()->GetParent());
 
 		UWidget* Preview = ParentRef.GetPreview();
-		TSharedPtr<SWidget> CachedPreviewSlateWidget = Preview ? Preview->GetCachedWidget() : NULL;
+		TSharedPtr<SWidget> CachedPreviewSlateWidget = Preview ? Preview->GetCachedWidget() : nullptr;
 		if ( CachedPreviewSlateWidget.IsValid() )
 		{
 			FWidgetPath WidgetPath;
@@ -777,7 +815,7 @@ UWidgetBlueprint* SDesignerView::GetBlueprint() const
 		return Cast<UWidgetBlueprint>(BP);
 	}
 
-	return NULL;
+	return nullptr;
 }
 
 void SDesignerView::Register(TSharedRef<FDesignerExtension> Extension)
@@ -793,7 +831,7 @@ void SDesignerView::OnBlueprintCompiled(UBlueprint* InBlueprint)
 	// having slate widgets that still may reference data in their owner UWidget that has been garbage collected.
 	CachedWidgetGeometry.Reset();
 
-	PreviewWidget = NULL;
+	PreviewWidget = nullptr;
 	PreviewSurface->SetContent(SNullWidget::NullWidget);
 }
 
@@ -812,7 +850,7 @@ FWidgetReference SDesignerView::GetWidgetAtCursor(const FGeometry& MyGeometry, c
 	UUserWidget* WidgetActor = BlueprintEditor.Pin()->GetPreview();
 	if ( WidgetActor )
 	{
-		UWidget* Preview = NULL;
+		UWidget* Preview = nullptr;
 
 		for ( int32 ChildIndex = Children.Num() - 1; ChildIndex >= 0; ChildIndex-- )
 		{
@@ -822,7 +860,7 @@ FWidgetReference SDesignerView::GetWidgetAtCursor(const FGeometry& MyGeometry, c
 			// Ignore the drop preview widget when doing widget picking
 			if (Preview == DropPreviewWidget)
 			{
-				Preview = NULL;
+				Preview = nullptr;
 				continue;
 			}
 			
@@ -1014,7 +1052,7 @@ FReply SDesignerView::OnKeyDown(const FGeometry& MyGeometry, const FKeyboardEven
 
 void SDesignerView::ShowContextMenu(const FGeometry& MyGeometry, const FPointerEvent& MouseEvent)
 {
-	FMenuBuilder MenuBuilder(true, NULL);
+	FMenuBuilder MenuBuilder(true, nullptr);
 
 	FWidgetBlueprintEditorUtils::CreateWidgetContextMenu(MenuBuilder, BlueprintEditor.Pin().ToSharedRef(), SelectedWidgetContextMenuLocation);
 
@@ -1266,7 +1304,7 @@ void SDesignerView::OnDragLeave(const FDragDropEvent& DragDropEvent)
 
 		UWidgetBlueprint* BP = GetBlueprint();
 		BP->WidgetTree->RemoveWidget(DropPreviewWidget);
-		DropPreviewWidget = NULL;
+		DropPreviewWidget = nullptr;
 	}
 }
 
@@ -1282,7 +1320,7 @@ FReply SDesignerView::OnDragOver(const FGeometry& MyGeometry, const FDragDropEve
 		}
 		
 		BP->WidgetTree->RemoveWidget(DropPreviewWidget);
-		DropPreviewWidget = NULL;
+		DropPreviewWidget = nullptr;
 	}
 	
 	const bool bIsPreview = true;
@@ -1316,7 +1354,7 @@ UWidget* SDesignerView::ProcessDropAndAddWidget(const FGeometry& MyGeometry, con
 		}
 
 		BP->WidgetTree->RemoveWidget(DropPreviewWidget);
-		DropPreviewWidget = NULL;
+		DropPreviewWidget = nullptr;
 	}
 
 	FArrangedWidget ArrangedWidget(SNullWidget::NullWidget, FGeometry());
@@ -1324,7 +1362,7 @@ UWidget* SDesignerView::ProcessDropAndAddWidget(const FGeometry& MyGeometry, con
 
 	FGeometry WidgetUnderCursorGeometry = ArrangedWidget.Geometry;
 
-	UWidget* Target = NULL;
+	UWidget* Target = nullptr;
 	if ( WidgetUnderCursor.IsValid() )
 	{
 		Target = bIsPreview ? WidgetUnderCursor.GetPreview() : WidgetUnderCursor.GetTemplate();
@@ -1336,7 +1374,7 @@ UWidget* SDesignerView::ProcessDropAndAddWidget(const FGeometry& MyGeometry, con
 		TemplateDragDropOp->SetCursorOverride(TOptional<EMouseCursor::Type>());
 
 		// If there's no root widget go ahead and add the widget into the root slot.
-		if ( BP->WidgetTree->RootWidget == NULL )
+		if ( BP->WidgetTree->RootWidget == nullptr )
 		{
 			FScopedTransaction Transaction(LOCTEXT("Designer_AddWidget", "Add Widget"));
 
@@ -1356,7 +1394,7 @@ UWidget* SDesignerView::ProcessDropAndAddWidget(const FGeometry& MyGeometry, con
 
 			SelectedWidget = BlueprintEditor.Pin()->GetReferenceFromTemplate(Widget);
 
-			DropPreviewParent = NULL;
+			DropPreviewParent = nullptr;
 
 			if ( bIsPreview )
 			{
@@ -1597,7 +1635,7 @@ UWidget* SDesignerView::ProcessDropAndAddWidget(const FGeometry& MyGeometry, con
 		}
 	}
 	
-	return NULL;
+	return nullptr;
 }
 
 FReply SDesignerView::OnDrop(const FGeometry& MyGeometry, const FDragDropEvent& DragDropEvent)
@@ -1614,7 +1652,7 @@ FReply SDesignerView::OnDrop(const FGeometry& MyGeometry, const FDragDropEvent&
 		}
 		
 		BP->WidgetTree->RemoveWidget(DropPreviewWidget);
-		DropPreviewWidget = NULL;
+		DropPreviewWidget = nullptr;
 	}
 	
 	const bool bIsPreview = false;
@@ -1652,11 +1690,31 @@ FText SDesignerView::GetCurrentResolutionText() const
 	return GetResolutionText(PreviewWidth, PreviewHeight, PreviewAspectRatio);
 }
 
+FText SDesignerView::GetCurrentDPIScaleText() const
+{
+	FInternationalization& I18N = FInternationalization::Get();
+
+	FNumberFormattingOptions Options;
+	Options.MinimumIntegralDigits = 1;
+	Options.MaximumFractionalDigits = 2;
+	Options.MinimumFractionalDigits = 1;
+
+	FText DPIString = FText::AsNumber(GetPreviewDPIScale(), &Options, I18N.GetInvariantCulture());
+	return FText::Format(LOCTEXT("CurrentDPIScaleFormat", "DPI Scale {0}"), DPIString);
+}
+
 FSlateColor SDesignerView::GetResolutionTextColorAndOpacity() const
 {
 	return FLinearColor(1, 1, 1, 1.25f - ResolutionTextFade.GetLerp());
 }
 
+FReply SDesignerView::HandleDPISettingsClicked()
+{
+	FModuleManager::LoadModuleChecked<ISettingsModule>("Settings").ShowViewer("Project", "Engine", "UI");
+
+	return FReply::Handled();
+}
+
 void SDesignerView::HandleOnCommonResolutionSelected(int32 Width, int32 Height, FString AspectRatio)
 {
 	PreviewWidth = Width;
diff --git a/Engine/Source/Editor/UMGEditor/Private/Designer/SDesignerView.h b/Engine/Source/Editor/UMGEditor/Private/Designer/SDesignerView.h
--- a/Engine/Source/Editor/UMGEditor/Private/Designer/SDesignerView.h
+++ b/Engine/Source/Editor/UMGEditor/Private/Designer/SDesignerView.h
@@ -119,6 +119,7 @@ private:
 	FText GetResolutionText(int32 Width, int32 Height, const FString& AspectRatio) const;
 
 	FText GetCurrentResolutionText() const;
+	FText GetCurrentDPIScaleText() const;
 	FSlateColor GetResolutionTextColorAndOpacity() const;
 
 	TOptional<int32> GetCustomResolutionWidth() const;
@@ -137,6 +138,7 @@ private:
 
 	// Handles drawing selection and other effects a SPaintSurface widget injected into the hierarchy.
 	int32 HandleEffectsPainting(const FOnPaintHandlerParams& PaintArgs);
+	FReply HandleDPISettingsClicked();
 
 	UUserWidget* GetDefaultWidget() const;
 
diff --git a/Engine/Source/Editor/UMGEditor/UMGEditor.Build.cs b/Engine/Source/Editor/UMGEditor/UMGEditor.Build.cs
--- a/Engine/Source/Editor/UMGEditor/UMGEditor.Build.cs
+++ b/Engine/Source/Editor/UMGEditor/UMGEditor.Build.cs
@@ -48,7 +48,8 @@ public class UMGEditor : ModuleRules
 				"SlateCore",
 				"MovieSceneCore",
 				"Sequencer",
-				"DetailCustomizations"
+				"DetailCustomizations",
+                "Settings"
 			}
 			);
 
diff --git a/Engine/Source/Editor/UnrealEd/Private/PlayLevel.cpp b/Engine/Source/Editor/UnrealEd/Private/PlayLevel.cpp
--- a/Engine/Source/Editor/UnrealEd/Private/PlayLevel.cpp
+++ b/Engine/Source/Editor/UnrealEd/Private/PlayLevel.cpp
@@ -31,6 +31,7 @@
 #include "EngineAnalytics.h"
 #include "Engine/GameInstance.h"
 #include "EditorAnalytics.h"
+#include "Runtime/Engine/Classes/Engine/UserInterfaceSettings.h"
 #include "Runtime/Engine/Classes/Engine/RendererSettings.h"
 #include "SScissorRectBox.h"
 #include "Online.h"
@@ -2682,7 +2683,7 @@ float UEditorEngine::GetGameViewportDPIScale(UGameViewportClient* ViewportClient
 {
 	FVector2D ViewportSize;
 	ViewportClient->GetViewportSize(ViewportSize);
-	return GetDefault<URendererSettings>(URendererSettings::StaticClass())->GetDPIScaleBasedOnSize(FIntPoint(ViewportSize.X, ViewportSize.Y));
+	return GetDefault<UUserInterfaceSettings>(UUserInterfaceSettings::StaticClass())->GetDPIScaleBasedOnSize(FIntPoint(ViewportSize.X, ViewportSize.Y));
 }
 
 FViewport* UEditorEngine::GetActiveViewport()
diff --git a/Engine/Source/Runtime/Engine/Classes/Engine/RendererSettings.h b/Engine/Source/Runtime/Engine/Classes/Engine/RendererSettings.h
--- a/Engine/Source/Runtime/Engine/Classes/Engine/RendererSettings.h
+++ b/Engine/Source/Runtime/Engine/Classes/Engine/RendererSettings.h
@@ -2,6 +2,8 @@
 
 #pragma once
 
+#include "UserInterfaceSettings.h"
+
 #include "RendererSettings.generated.h"
 
 
@@ -81,26 +83,6 @@ namespace EAntiAliasingMethodUI
 }
 
 
-/** The Side to use when scaling the UI. */
-UENUM()
-namespace EUIScalingRule
-{
-	enum Type
-	{
-		/** Evaluates the scale curve based on the shortest side of the viewport */
-		ShortestSide,
-		/** Evaluates the scale curve based on the longest side of the viewport */
-		LongestSide,
-		/** Evaluates the scale curve based on the X axis of the viewport */
-		Horizontal,
-		/** Evaluates the scale curve based on the Y axis of the viewport */
-		Vertical,
-		/** Custom - Allows custom rule interpretation */
-		//Custom
-	};
-}
-
-
 /**
  * Implements project settings for the Rendering sub-system.
  */
@@ -256,27 +238,23 @@ class ENGINE_API URendererSettings
 		ToolTip="Screen radius at which wireframe objects are culled. Larger values can improve performance when viewing a scene in wireframe."))
 	float WireframeCullThreshold;
 
-	UPROPERTY(config, EditAnywhere, Category=UI, meta=(
-		DisplayName="DPI Scale Rule",
-		ToolTip="The rule used when trying to decide what scale to apply." ))
-	TEnumAsByte<EUIScalingRule::Type> UIScaleRule;
-
-	UPROPERTY(config, EditAnywhere, Category=UI, meta=(
-		DisplayName="DPI Curve",
-		ToolTip="Controls how the UI is scaled at different resolutions based on the DPI Scale Rule" ))
-	FRuntimeFloatCurve UIScaleCurve;
-
 public:
 
 	// Begin UObject interface
 
 	virtual void PostInitProperties() override;
+
 #if WITH_EDITOR
 	virtual void PostEditChangeProperty(FPropertyChangedEvent& PropertyChangedEvent) override;
 #endif
 
 	// End UObject interface
 
-	/** Gets the current scale of the UI based on the size */
-	float GetDPIScaleBasedOnSize(FIntPoint Size) const;
+private:
+	
+	UPROPERTY(config)
+	TEnumAsByte<EUIScalingRule::Type> UIScaleRule_DEPRECATED;
+
+	UPROPERTY(config)
+	FRuntimeFloatCurve UIScaleCurve_DEPRECATED;
 };
diff --git a/Engine/Source/Runtime/Engine/Classes/Engine/UserInterfaceSettings.h b/Engine/Source/Runtime/Engine/Classes/Engine/UserInterfaceSettings.h
new file mode 100644
--- /dev/null
+++ b/Engine/Source/Runtime/Engine/Classes/Engine/UserInterfaceSettings.h
@@ -0,0 +1,57 @@
+// Copyright 1998-2014 Epic Games, Inc. All Rights Reserved.
+
+#pragma once
+
+#include "UserInterfaceSettings.generated.h"
+
+
+/** The Side to use when scaling the UI. */
+UENUM()
+namespace EUIScalingRule
+{
+	enum Type
+	{
+		/** Evaluates the scale curve based on the shortest side of the viewport */
+		ShortestSide,
+		/** Evaluates the scale curve based on the longest side of the viewport */
+		LongestSide,
+		/** Evaluates the scale curve based on the X axis of the viewport */
+		Horizontal,
+		/** Evaluates the scale curve based on the Y axis of the viewport */
+		Vertical,
+		/** Custom - Allows custom rule interpretation */
+		//Custom
+	};
+}
+
+
+/**
+ * Implements user interface related settings.
+ */
+UCLASS(config=Engine, defaultconfig)
+class ENGINE_API UUserInterfaceSettings
+	: public UObject
+{
+	GENERATED_UCLASS_BODY()
+
+public:
+
+	UPROPERTY(config, EditAnywhere, Category="DPI Scaling", meta=(
+		DisplayName="DPI Scale Rule",
+		ToolTip="The rule used when trying to decide what scale to apply." ))
+	TEnumAsByte<EUIScalingRule::Type> UIScaleRule;
+
+	UPROPERTY(config, EditAnywhere, Category="DPI Scaling", meta=(
+		DisplayName="DPI Curve",
+		ToolTip="Controls how the UI is scaled at different resolutions based on the DPI Scale Rule" ))
+	FRuntimeFloatCurve UIScaleCurve;
+
+public:
+
+	/** Gets the current scale of the UI based on the size */
+	float GetDPIScaleBasedOnSize(FIntPoint Size) const;
+
+#if WITH_EDITOR
+	void PostEditChangeProperty(FPropertyChangedEvent& PropertyChangedEvent);
+#endif
+};
diff --git a/Engine/Source/Runtime/Engine/Private/GameEngine.cpp b/Engine/Source/Runtime/Engine/Private/GameEngine.cpp
--- a/Engine/Source/Runtime/Engine/Private/GameEngine.cpp
+++ b/Engine/Source/Runtime/Engine/Private/GameEngine.cpp
@@ -16,6 +16,7 @@
 #include "EngineModule.h"
 #include "Engine/GameInstance.h"
 #include "Engine/RendererSettings.h"
+#include "Engine/UserInterfaceSettings.h"
 #include "AVIWriter.h"
 
 #include "SlateBasics.h"
@@ -169,7 +170,7 @@ float UGameEngine::GetGameViewportDPIScale(UGameViewportClient* ViewportClient)
 {
 	FVector2D ViewportSize;
 	ViewportClient->GetViewportSize(ViewportSize);
-	return GetDefault<URendererSettings>(URendererSettings::StaticClass())->GetDPIScaleBasedOnSize(FIntPoint(ViewportSize.X, ViewportSize.Y));
+	return GetDefault<UUserInterfaceSettings>(UUserInterfaceSettings::StaticClass())->GetDPIScaleBasedOnSize(FIntPoint(ViewportSize.X, ViewportSize.Y));
 }
 
 void UGameEngine::ConditionallyOverrideSettings(int32& ResolutionX, int32& ResolutionY, EWindowMode::Type& WindowMode)
diff --git a/Engine/Source/Runtime/Engine/Private/RenderingSettings.cpp b/Engine/Source/Runtime/Engine/Private/RenderingSettings.cpp
--- a/Engine/Source/Runtime/Engine/Private/RenderingSettings.cpp
+++ b/Engine/Source/Runtime/Engine/Private/RenderingSettings.cpp
@@ -14,6 +14,24 @@ static FName ToolTipFName(TEXT("ConsoleVariable"));
 void URendererSettings::PostInitProperties()
 {
 	Super::PostInitProperties();
+
+	if ( IsTemplate() )
+	{
+		// If we have UI scale curve data in the renderer settings config, move it to the user interface settings.
+		if ( FRichCurve* OldCurve = UIScaleCurve_DEPRECATED.GetRichCurve() )
+		{
+			if ( OldCurve->GetNumKeys() != 0 )
+			{
+				UUserInterfaceSettings* UISettings = GetMutableDefault<UUserInterfaceSettings>(UUserInterfaceSettings::StaticClass());
+				UISettings->UIScaleRule = UIScaleRule_DEPRECATED;
+				UISettings->UIScaleCurve = UIScaleCurve_DEPRECATED;
+
+				// Remove all old keys so that we don't attempt to load from renderer settings again.
+				OldCurve->Reset();
+			}
+		}
+	}
+
 #if WITH_EDITOR
 	if (IsTemplate())
 	{
@@ -49,6 +67,7 @@ void URendererSettings::PostInitProperties()
 void URendererSettings::PostEditChangeProperty(FPropertyChangedEvent& PropertyChangedEvent)
 {
 	Super::PostEditChangeProperty(PropertyChangedEvent);
+
 	if (PropertyChangedEvent.Property)
 	{
 		FString CVarName = PropertyChangedEvent.Property->GetMetaData(ConsoleVariableFName);
@@ -80,40 +99,14 @@ void URendererSettings::PostEditChangeProperty(FPropertyChangedEvent& PropertyCh
 				UE_LOG(LogInit, Warning, TEXT("CVar named '%s' marked up in URenderingSettings was not found or is set to read-only"), *CVarName);
 			}
 		}
-	}
-}
-#endif // #if WITH_EDITOR
-
-float URendererSettings::GetDPIScaleBasedOnSize(FIntPoint Size) const
-{
-	float Scale = 1;
 
-	//if ( UIScaleRule == EUIScalingRule::Custom )
-	//{
-
-	//}
-	//else
-	{
-		int32 EvalPoint = 0;
-		switch ( UIScaleRule )
+		// If the renderer settings are updated and saved, we need to update the user interface settings too, to prevent data loss
+		// from old curve data for UI settings.
+		if ( IsTemplate() )
 		{
-		case EUIScalingRule::ShortestSide:
-			EvalPoint = FMath::Min(Size.X, Size.Y);
-			break;
-		case EUIScalingRule::LongestSide:
-			EvalPoint = FMath::Max(Size.X, Size.Y);
-			break;
-		case EUIScalingRule::Horizontal:
-			EvalPoint = Size.X;
-			break;
-		case EUIScalingRule::Vertical:
-			EvalPoint = Size.Y;
-			break;
+			UUserInterfaceSettings* UISettings = GetMutableDefault<UUserInterfaceSettings>(UUserInterfaceSettings::StaticClass());
+			UISettings->UpdateDefaultConfigFile();
 		}
-
-		const FRichCurve* DPICurve = UIScaleCurve.GetRichCurveConst();
-		Scale = DPICurve->Eval((float)EvalPoint, 1.0f);
 	}
-
-	return FMath::Max(Scale, 0.01f);
 }
+#endif // #if WITH_EDITOR
diff --git a/Engine/Source/Runtime/Engine/Private/UserInterfaceSettings.cpp b/Engine/Source/Runtime/Engine/Private/UserInterfaceSettings.cpp
new file mode 100644
--- /dev/null
+++ b/Engine/Source/Runtime/Engine/Private/UserInterfaceSettings.cpp
@@ -0,0 +1,61 @@
+// Copyright 1998-2014 Epic Games, Inc. All Rights Reserved.
+
+/*=============================================================================
+	PlayerInput.cpp: Unreal input system.
+=============================================================================*/
+
+#include "EnginePrivate.h"
+#include "Engine/UserInterfaceSettings.h"
+
+
+UUserInterfaceSettings::UUserInterfaceSettings(const FObjectInitializer& ObjectInitializer)
+: Super(ObjectInitializer)
+{
+}
+
+float UUserInterfaceSettings::GetDPIScaleBasedOnSize(FIntPoint Size) const
+{
+	float Scale = 1;
+
+	//if ( UIScaleRule == EUIScalingRule::Custom )
+	//{
+
+	//}
+	//else
+	{
+		int32 EvalPoint = 0;
+		switch ( UIScaleRule )
+		{
+		case EUIScalingRule::ShortestSide:
+			EvalPoint = FMath::Min(Size.X, Size.Y);
+			break;
+		case EUIScalingRule::LongestSide:
+			EvalPoint = FMath::Max(Size.X, Size.Y);
+			break;
+		case EUIScalingRule::Horizontal:
+			EvalPoint = Size.X;
+			break;
+		case EUIScalingRule::Vertical:
+			EvalPoint = Size.Y;
+			break;
+		}
+
+		const FRichCurve* DPICurve = UIScaleCurve.GetRichCurveConst();
+		Scale = DPICurve->Eval((float)EvalPoint, 1.0f);
+	}
+
+	return FMath::Max(Scale, 0.01f);
+}
+
+#if WITH_EDITOR
+void UUserInterfaceSettings::PostEditChangeProperty(FPropertyChangedEvent& PropertyChangedEvent)
+{
+	Super::PostEditChangeProperty(PropertyChangedEvent);
+
+	if ( PropertyChangedEvent.Property )
+	{
+		URendererSettings* RendererSettings = GetMutableDefault<URendererSettings>(URendererSettings::StaticClass());
+		RendererSettings->UpdateDefaultConfigFile();
+	}
+}
+#endif // #if WITH_EDITOR
diff --git a/Engine/Source/Runtime/UMG/Private/WidgetLayoutLibrary.cpp b/Engine/Source/Runtime/UMG/Private/WidgetLayoutLibrary.cpp
--- a/Engine/Source/Runtime/UMG/Private/WidgetLayoutLibrary.cpp
+++ b/Engine/Source/Runtime/UMG/Private/WidgetLayoutLibrary.cpp
@@ -3,6 +3,7 @@
 #include "UMGPrivatePCH.h"
 #include "Slate/SlateBrushAsset.h"
 #include "WidgetLayoutLibrary.h"
+#include "Runtime/Engine/Classes/Engine/UserInterfaceSettings.h"
 #include "Runtime/Engine/Classes/Engine/RendererSettings.h"
 
 #define LOCTEXT_NAMESPACE "UMG"
@@ -24,7 +25,7 @@ float UWidgetLayoutLibrary::GetViewportScale(UObject* WorldContextObject)
 		{
 			FVector2D ViewportSize;
 			ViewportClient->GetViewportSize(ViewportSize);
-			return GetDefault<URendererSettings>(URendererSettings::StaticClass())->GetDPIScaleBasedOnSize(FIntPoint(ViewportSize.X, ViewportSize.Y));
+			return GetDefault<UUserInterfaceSettings>(UUserInterfaceSettings::StaticClass())->GetDPIScaleBasedOnSize(FIntPoint(ViewportSize.X, ViewportSize.Y));
 		}
 	}
 
