UMG - Work continues on named slots, upgrading the way selection is synced in the hierarchy view to work better with named slots.

diff --git a/Engine/Source/Editor/UMGEditor/Private/Hierarchy/SHierarchyView.cpp b/Engine/Source/Editor/UMGEditor/Private/Hierarchy/SHierarchyView.cpp
--- a/Engine/Source/Editor/UMGEditor/Private/Hierarchy/SHierarchyView.cpp
+++ b/Engine/Source/Editor/UMGEditor/Private/Hierarchy/SHierarchyView.cpp
@@ -145,54 +145,12 @@ void SHierarchyView::OnEditorSelectionChanged()
 {
 	WidgetTreeView->ClearSelection();
 
-	bool bFirst = true;
-
-	// Update the selection and expansion in the tree to match the new selection
-	const TSet<FWidgetReference>& SelectedWidgets = BlueprintEditor.Pin()->GetSelectedWidgets();
-	for ( const FWidgetReference& WidgetRef : SelectedWidgets )
+	if ( RootWidgets.Num() > 0 )
 	{
-		UWidget* TemplateWidget = WidgetRef.GetTemplate();
-
-		if ( TemplateWidget )
-		{
-			TArray< UWidget* > Widgets;
-			Widgets.Add(TemplateWidget);
-
-			UWidget* Parent = TemplateWidget->GetParent();
-			while ( Parent != NULL )
-			{
-				Widgets.Add(Parent);
-				Parent = Parent->GetParent();
-			}
-
-			TSharedPtr<FHierarchyModel> CurrentItem = RootWidgets[0];
-			WidgetTreeView->SetItemExpansion(CurrentItem, true);
-
-			for ( int32 WidgetIndex = Widgets.Num() - 1; WidgetIndex >= 0; WidgetIndex-- )
-			{
-				UWidget* WidgetItem = Widgets[WidgetIndex];
-
-				TArray< TSharedPtr<FHierarchyModel> > Children;
-				CurrentItem->GatherChildren(Children);
-
-				for ( auto C : Children )
-				{
-					if ( C->IsModel(WidgetItem) )
-					{
-						WidgetTreeView->SetItemExpansion(C, true);
-						CurrentItem = C;
-
-						if ( WidgetIndex == 0 )
-						{
-							WidgetTreeView->SetItemSelection(CurrentItem, true, ESelectInfo::Direct);
-						}
-					}
-				}
-			}
-
-			WidgetTreeView->RequestScrollIntoView(CurrentItem);
-		}
+		RootWidgets[0]->RefreshSelection();
 	}
+
+	RestoreSelectedItems();
 }
 
 UWidgetBlueprint* SHierarchyView::GetBlueprint() const
@@ -331,6 +289,38 @@ void SHierarchyView::RecursiveExpand(TSharedPtr<FHierarchyModel>& Model)
 	}
 }
 
+void SHierarchyView::RestoreSelectedItems()
+{
+	for ( TSharedPtr<FHierarchyModel>& Model : RootWidgets )
+	{
+		RecursiveSelection(Model);
+	}
+}
+
+void SHierarchyView::RecursiveSelection(TSharedPtr<FHierarchyModel>& Model)
+{
+	if ( Model->ContainsSelection() )
+	{
+		// Expand items that contain selection.
+		WidgetTreeView->SetItemExpansion(Model, true);
+
+		TArray< TSharedPtr<FHierarchyModel> > Children;
+		Model->GatherChildren(Children);
+
+		for ( TSharedPtr<FHierarchyModel>& ChildModel : Children )
+		{
+			RecursiveSelection(ChildModel);
+		}
+	}
+
+	if ( Model->IsSelected() )
+	{
+		WidgetTreeView->SetItemSelection(Model, true, ESelectInfo::Direct);
+		WidgetTreeView->RequestScrollIntoView(Model);
+	}
+}
+
+
 //@TODO UMG Drop widgets onto the tree, when nothing is present, if there is a root node present, what happens then, let the root node attempt to place it?
 
 #undef LOCTEXT_NAMESPACE
diff --git a/Engine/Source/Editor/UMGEditor/Private/Hierarchy/SHierarchyView.h b/Engine/Source/Editor/UMGEditor/Private/Hierarchy/SHierarchyView.h
--- a/Engine/Source/Editor/UMGEditor/Private/Hierarchy/SHierarchyView.h
+++ b/Engine/Source/Editor/UMGEditor/Private/Hierarchy/SHierarchyView.h
@@ -94,6 +94,12 @@ private:
 	/** Recursively expands the models based on the expansion set. */
 	void RecursiveExpand(TSharedPtr<FHierarchyModel>& Model);
 
+	/**  */
+	void RestoreSelectedItems();
+
+	/**  */
+	void RecursiveSelection(TSharedPtr<FHierarchyModel>& Model);
+
 private:
 
 	/** Cached pointer to the blueprint editor that owns this tree. */
diff --git a/Engine/Source/Editor/UMGEditor/Private/Hierarchy/SHierarchyViewItem.cpp b/Engine/Source/Editor/UMGEditor/Private/Hierarchy/SHierarchyViewItem.cpp
--- a/Engine/Source/Editor/UMGEditor/Private/Hierarchy/SHierarchyViewItem.cpp
+++ b/Engine/Source/Editor/UMGEditor/Private/Hierarchy/SHierarchyViewItem.cpp
@@ -206,6 +206,7 @@ FReply ProcessHierarchyDragDrop(const FDragDropEvent& DragDropEvent, bool bIsDro
 
 FHierarchyModel::FHierarchyModel()
 	: bInitialized(false)
+	, bIsSelected(false)
 {
 
 }
@@ -250,17 +251,54 @@ void FHierarchyModel::OnNameTextCommited(const FText& InText, ETextCommit::Type
 
 }
 
-void FHierarchyModel::GatherChildren(TArray< TSharedPtr<FHierarchyModel> >& Children)
+void FHierarchyModel::InitializeChildren()
 {
 	if ( !bInitialized )
 	{
 		bInitialized = true;
 		GetChildren(Models);
 	}
+}
+
+void FHierarchyModel::GatherChildren(TArray< TSharedPtr<FHierarchyModel> >& Children)
+{
+	InitializeChildren();
 
 	Children.Append(Models);
 }
 
+bool FHierarchyModel::ContainsSelection()
+{
+	InitializeChildren();
+
+	for ( TSharedPtr<FHierarchyModel>& Model : Models )
+	{
+		if ( Model->IsSelected() || Model->ContainsSelection() )
+		{
+			return true;
+		}
+	}
+
+	return false;
+}
+
+void FHierarchyModel::RefreshSelection()
+{
+	InitializeChildren();
+
+	UpdateSelection();
+
+	for ( TSharedPtr<FHierarchyModel>& Model : Models )
+	{
+		Model->RefreshSelection();
+	}
+}
+
+bool FHierarchyModel::IsSelected() const
+{
+	return bIsSelected;
+}
+
 //////////////////////////////////////////////////////////////////////////
 
 FHierarchyRoot::FHierarchyRoot(TSharedPtr<FWidgetBlueprintEditor> InBlueprintEditor)
@@ -281,7 +319,7 @@ FText FHierarchyRoot::GetText() const
 
 const FSlateBrush* FHierarchyRoot::GetImage() const
 {
-	return NULL;
+	return nullptr;
 }
 
 FSlateFontInfo FHierarchyRoot::GetFont() const
@@ -312,6 +350,20 @@ void FHierarchyRoot::OnSelection()
 	}
 }
 
+void FHierarchyRoot::UpdateSelection()
+{
+	TSharedPtr<FWidgetBlueprintEditor> BPEd = BlueprintEditor.Pin();
+	if ( UWidget* Default = BPEd->GetWidgetBlueprintObj()->GeneratedClass->GetDefaultObject<UWidget>() )
+	{
+		const TSet< TWeakObjectPtr<UObject> >& SelectedObjects = BlueprintEditor.Pin()->GetSelectedObjects();
+		bIsSelected = SelectedObjects.Contains(Default);
+	}
+	else
+	{
+		bIsSelected = false;
+	}
+}
+
 FReply FHierarchyRoot::OnDragOver(const FGeometry& MyGeometry, const FDragDropEvent& DragDropEvent)
 {
 	UWidgetBlueprint* Blueprint = BlueprintEditor.Pin()->GetWidgetBlueprintObj();
@@ -373,36 +425,65 @@ FSlateFontInfo FNamedSlotModel::GetFont() const
 
 void FNamedSlotModel::GetChildren(TArray< TSharedPtr<FHierarchyModel> >& Children)
 {
-	//TSharedPtr<FWidgetBlueprintEditor> BPEd = BlueprintEditor.Pin();
-	//UWidgetBlueprint* Blueprint = BPEd->GetWidgetBlueprintObj();
-
-	//if ( Blueprint->WidgetTree->RootWidget )
-	//{
-	//	TSharedPtr<FHierarchyWidget> RootChild = MakeShareable(new FHierarchyWidget(BPEd->GetReferenceFromTemplate(Blueprint->WidgetTree->RootWidget), BPEd));
-	//	Children.Add(RootChild);
-	//}
+	TSharedPtr<FWidgetBlueprintEditor> BPEd = BlueprintEditor.Pin();
+	if ( INamedSlotInterface* NamedSlotHost = Cast<INamedSlotInterface>(Item.GetTemplate()) )
+	{
+		TSet<FWidgetReference> SelectedWidgets;
+		if ( UWidget* TemplateSlotContent = NamedSlotHost->GetContentForSlot(SlotName) )
+		{
+			TSharedPtr<FHierarchyWidget> RootChild = MakeShareable(new FHierarchyWidget(BPEd->GetReferenceFromTemplate(TemplateSlotContent), BPEd));
+			Children.Add(RootChild);
+		}
+	}
 }
 
 void FNamedSlotModel::OnSelection()
 {
 	TSharedPtr<FWidgetBlueprintEditor> BPEd = BlueprintEditor.Pin();
-	if ( INamedSlotInterface* NamedSlotHost = Cast<INamedSlotInterface>(Item.GetPreview()) )
+	if ( INamedSlotInterface* NamedSlotHost = Cast<INamedSlotInterface>(Item.GetTemplate()) )
 	{
 		TSet<FWidgetReference> SelectedWidgets;
-		if ( UWidget* SlotContent = NamedSlotHost->GetContentForSlot(SlotName) )
+		if ( UWidget* TemplateSlotContent = NamedSlotHost->GetContentForSlot(SlotName) )
 		{
-			SelectedWidgets.Add( BPEd->GetReferenceFromPreview(SlotContent) );
+			SelectedWidgets.Add(BPEd->GetReferenceFromTemplate(TemplateSlotContent));
 		}
 
 		BPEd->SelectWidgets(SelectedWidgets);
 	}
 }
 
+void FNamedSlotModel::UpdateSelection()
+{
+	//bIsSelected = false;
+
+	//const TSet<FWidgetReference>& SelectedWidgets = BlueprintEditor.Pin()->GetSelectedWidgets();
+
+	//TSharedPtr<FWidgetBlueprintEditor> BPEd = BlueprintEditor.Pin();
+	//if ( INamedSlotInterface* NamedSlotHost = Cast<INamedSlotInterface>(Item.GetTemplate()) )
+	//{
+	//	if ( UWidget* TemplateSlotContent = NamedSlotHost->GetContentForSlot(SlotName) )
+	//	{
+	//		bIsSelected = SelectedWidgets.Contains(BPEd->GetReferenceFromTemplate(TemplateSlotContent));
+	//	}
+	//}
+}
+
 FReply FNamedSlotModel::OnDragOver(const FGeometry& MyGeometry, const FDragDropEvent& DragDropEvent)
 {
-	//UWidgetBlueprint* Blueprint = BlueprintEditor.Pin()->GetWidgetBlueprintObj();
-	//bool bIsDrop = false;
-	//return ProcessHierarchyDragDrop(DragDropEvent, bIsDrop, Blueprint, FWidgetReference());
+	TSharedPtr<FWidgetTemplateDragDropOp> TemplateDragDropOp = DragDropEvent.GetOperationAs<FWidgetTemplateDragDropOp>();
+	if ( TemplateDragDropOp.IsValid() )
+	{
+		if ( INamedSlotInterface* NamedSlotHost = Cast<INamedSlotInterface>(Item.GetTemplate()) )
+		{
+			// Only assign content to the named slot if it is null.
+			if ( NamedSlotHost->GetContentForSlot(SlotName) != nullptr )
+			{
+				TSharedPtr<FDecoratedDragDropOp> DecoratedDragDropOp = DragDropEvent.GetOperationAs<FDecoratedDragDropOp>();
+				TemplateDragDropOp->SetCursorOverride(EMouseCursor::SlashedCircle);
+				return FReply::Handled();
+			}
+		}
+	}
 
 	return FReply::Unhandled();
 }
@@ -419,11 +500,13 @@ FReply FNamedSlotModel::HandleDrop(FDragDropEvent const& DragDropEvent)
 	{
 		if ( INamedSlotInterface* NamedSlotHost = Cast<INamedSlotInterface>(Item.GetTemplate()) )
 		{
-			UWidget* Widget = TemplateDragDropOp->Template->Create(Blueprint->WidgetTree);
-
-			NamedSlotHost->SetContentForSlot(SlotName, Widget);
-
-			FBlueprintEditorUtils::MarkBlueprintAsStructurallyModified(Blueprint);
+			// Only assign content to the named slot if it is null.
+			if ( NamedSlotHost->GetContentForSlot(SlotName) == nullptr )
+			{
+				UWidget* Widget = TemplateDragDropOp->Template->Create(Blueprint->WidgetTree);
+				NamedSlotHost->SetContentForSlot(SlotName, Widget);
+				FBlueprintEditorUtils::MarkBlueprintAsStructurallyModified(Blueprint);
+			}
 
 			return FReply::Handled();
 		}
@@ -611,6 +694,12 @@ void FHierarchyWidget::OnSelection()
 	BlueprintEditor.Pin()->SelectWidgets(SelectedWidgets);
 }
 
+void FHierarchyWidget::UpdateSelection()
+{
+	const TSet<FWidgetReference>& SelectedWidgets = BlueprintEditor.Pin()->GetSelectedWidgets();
+	bIsSelected = SelectedWidgets.Contains(Item);
+}
+
 //////////////////////////////////////////////////////////////////////////
 
 void SHierarchyViewItem::Construct(const FArguments& InArgs, const TSharedRef< STableViewBase >& InOwnerTableView, TSharedPtr<FHierarchyModel> InModel)
diff --git a/Engine/Source/Editor/UMGEditor/Private/Hierarchy/SHierarchyViewItem.h b/Engine/Source/Editor/UMGEditor/Private/Hierarchy/SHierarchyViewItem.h
--- a/Engine/Source/Editor/UMGEditor/Private/Hierarchy/SHierarchyViewItem.h
+++ b/Engine/Source/Editor/UMGEditor/Private/Hierarchy/SHierarchyViewItem.h
@@ -40,8 +40,10 @@ public:
 	virtual void GatherChildren(TArray< TSharedPtr<FHierarchyModel> >& Children);
 
 	virtual void OnSelection() = 0;
-
-	virtual bool IsModel(UObject* PossibleModel) const { return false; }
+	
+	void RefreshSelection();
+	bool ContainsSelection();
+	bool IsSelected() const;
 
 	virtual bool IsVisible() const { return true; }
 	virtual bool CanControlVisibility() const { return false; }
@@ -49,10 +51,15 @@ public:
 
 protected:
 	virtual void GetChildren(TArray< TSharedPtr<FHierarchyModel> >& Children) = 0;
+	virtual void UpdateSelection() = 0;
+
+private:
+	void InitializeChildren();
 
 protected:
 
 	bool bInitialized;
+	bool bIsSelected;
 	TArray< TSharedPtr<FHierarchyModel> > Models;
 };
 
@@ -80,6 +87,7 @@ public:
 
 protected:
 	virtual void GetChildren(TArray< TSharedPtr<FHierarchyModel> >& Children) override;
+	virtual void UpdateSelection();
 
 private:
 
@@ -103,13 +111,14 @@ public:
 	virtual FSlateFontInfo GetFont() const override;
 
 	virtual void OnSelection();
-
+	
 	virtual FReply OnDragOver(const FGeometry& MyGeometry, const FDragDropEvent& DragDropEvent);
 
 	virtual FReply HandleDrop(FDragDropEvent const& DragDropEvent);
 
 protected:
 	virtual void GetChildren(TArray< TSharedPtr<FHierarchyModel> >& Children) override;
+	virtual void UpdateSelection();
 
 private:
 
@@ -147,11 +156,6 @@ public:
 
 	virtual void OnSelection();
 
-	virtual bool IsModel(UObject* PossibleModel) const override
-	{
-		return Item.GetTemplate() == PossibleModel;
-	}
-
 	virtual bool IsVisible() const override
 	{
 		return !Item.GetTemplate()->bHiddenInDesigner;
@@ -177,6 +181,7 @@ public:
 
 protected:
 	virtual void GetChildren(TArray< TSharedPtr<FHierarchyModel> >& Children) override;
+	virtual void UpdateSelection();
 
 private:
 	FWidgetReference Item;
