UMG - Fixing a buidl error

diff --git a/Engine/Source/Editor/UMGEditor/Private/Extensions/CanvasSlotExtension.cpp b/Engine/Source/Editor/UMGEditor/Private/Extensions/CanvasSlotExtension.cpp
--- a/Engine/Source/Editor/UMGEditor/Private/Extensions/CanvasSlotExtension.cpp
+++ b/Engine/Source/Editor/UMGEditor/Private/Extensions/CanvasSlotExtension.cpp
@@ -459,7 +459,7 @@ void FCanvasSlotExtension::PaintCollisionLines(const TSet< FWidgetReference >& S
 							for ( int32 SlotIndex = 0; SlotIndex < Canvas->GetChildrenCount(); SlotIndex++ )
 							{
 								// Ignore the slot being dragged.
-								if ( Canvas->Slots[SlotIndex] == CanvasSlot )
+								if ( Canvas->GetSlots()[SlotIndex] == CanvasSlot )
 								{
 									continue;
 								}
