UMG - Limited support for multiple widget selection in the editor!

diff --git a/Engine/Source/Editor/UMGEditor/Private/Designer/SDesignerView.cpp b/Engine/Source/Editor/UMGEditor/Private/Designer/SDesignerView.cpp
--- a/Engine/Source/Editor/UMGEditor/Private/Designer/SDesignerView.cpp
+++ b/Engine/Source/Editor/UMGEditor/Private/Designer/SDesignerView.cpp
@@ -158,8 +158,6 @@ void SDesignerView::Construct(const FArguments& InArgs, TSharedPtr<FWidgetBluepr
 	ResolutionTextFade = FCurveSequence(0.0f, 1.0f);
 	ResolutionTextFade.Play();
 
-	HoverTime = 0;
-
 	bMovingExistingWidget = false;
 
 	// TODO UMG - Register these with the module through some public interface to allow for new extensions to be registered.
@@ -514,9 +512,25 @@ float SDesignerView::GetPreviewScale() const
 	return GetZoomAmount() * GetPreviewDPIScale();
 }
 
+const TSet<FWidgetReference>& SDesignerView::GetSelectedWidgets() const
+{
+	return BlueprintEditor.Pin()->GetSelectedWidgets();
+}
+
 FWidgetReference SDesignerView::GetSelectedWidget() const
 {
-	return SelectedWidget;
+	const TSet<FWidgetReference>& SelectedWidgets = BlueprintEditor.Pin()->GetSelectedWidgets();
+
+	// Only return a selected widget when we have only a single item selected.
+	if ( SelectedWidgets.Num() == 1 )
+	{
+		for ( TSet<FWidgetReference>::TConstIterator SetIt(SelectedWidgets); SetIt; ++SetIt )
+		{
+			return *SetIt;
+		}
+	}
+
+	return FWidgetReference();
 }
 
 ETransformMode::Type SDesignerView::GetTransformMode() const
@@ -595,7 +609,7 @@ void SDesignerView::OnEditorSelectionChanged()
 	TSet<FWidgetReference> PendingSelectedWidgets = BlueprintEditor.Pin()->GetSelectedWidgets();
 
 	// Notify all widgets that are no longer selected.
-	for ( FWidgetReference& WidgetRef : SelectedWidgets )
+	for ( FWidgetReference& WidgetRef : SelectedWidgetsCache )
 	{
 		if ( WidgetRef.IsValid() && !PendingSelectedWidgets.Contains(WidgetRef) )
 		{
@@ -606,26 +620,13 @@ void SDesignerView::OnEditorSelectionChanged()
 	// Notify all widgets that are now selected.
 	for ( FWidgetReference& WidgetRef : PendingSelectedWidgets )
 	{
-		if ( WidgetRef.IsValid() && !SelectedWidgets.Contains(WidgetRef) )
+		if ( WidgetRef.IsValid() && !SelectedWidgetsCache.Contains(WidgetRef) )
 		{
 			WidgetRef.GetPreview()->Select();
 		}
 	}
 
-	SelectedWidgets = PendingSelectedWidgets;
-
-	if ( SelectedWidgets.Num() > 0 )
-	{
-		for ( FWidgetReference& Widget : SelectedWidgets )
-		{
-			SelectedWidget = Widget;
-			break;
-		}
-	}
-	else
-	{
-		SelectedWidget = FWidgetReference();
-	}
+	SelectedWidgetsCache = PendingSelectedWidgets;
 
 	CreateExtensionWidgetsForSelection();
 }
@@ -691,12 +692,9 @@ void SDesignerView::CreateExtensionWidgetsForSelection()
 	// Remove all the current extension widgets
 	ClearExtensionWidgets();
 
-	TArray<FWidgetReference> Selected;
-	if ( SelectedWidget.IsValid() )
-	{
-		Selected.Add(SelectedWidget);
-	}
-
+	// Get the selected widgets as an array
+	TArray<FWidgetReference> Selected = GetSelectedWidgets().Array();
+	
 	TArray< TSharedRef<FDesignerSurfaceElement> > ExtensionElements;
 
 	if ( Selected.Num() > 0 )
@@ -737,74 +735,73 @@ void SDesignerView::CreateExtensionWidgetsForSelection()
 
 FVector2D SDesignerView::GetExtensionPosition(TSharedRef<FDesignerSurfaceElement> ExtensionElement) const
 {
-	const FVector2D TopLeft = CachedDesignerWidgetLocation;
-	const FVector2D Size = CachedDesignerWidgetSize * GetPreviewScale();
+	FWidgetReference SelectedWidget = GetSelectedWidget();
 
-	// Calculate the parent position and size.  We use this information for calculating offsets.
-	FVector2D ParentPosition, ParentSize;
+	if ( SelectedWidget.IsValid() )
 	{
-		FWidgetReference ParentRef = BlueprintEditor.Pin()->GetReferenceFromTemplate(SelectedWidget.GetTemplate()->GetParent());
+		FGeometry SelectedWidgetGeometry_RelativeToDesigner;
+		FGeometry SelectedWidgetParentGeometry_RelativeToDesigner;
 
-		UWidget* Preview = ParentRef.GetPreview();
-		TSharedPtr<SWidget> CachedPreviewSlateWidget = Preview ? Preview->GetCachedWidget() : nullptr;
-		if ( CachedPreviewSlateWidget.IsValid() )
+		if ( GetWidgetGeometry(SelectedWidget, SelectedWidgetGeometry_RelativeToDesigner) && GetWidgetParentGeometry(SelectedWidget, SelectedWidgetParentGeometry_RelativeToDesigner) )
 		{
-			FWidgetPath WidgetPath;
-			SelectedWidgetPath.ToWidgetPath(WidgetPath);
-		
-			FArrangedWidget ArrangedWidget(SNullWidget::NullWidget, FGeometry());
-			FDesignTimeUtils::GetArrangedWidgetRelativeToParent(WidgetPath, CachedPreviewSlateWidget.ToSharedRef(), AsShared(), ArrangedWidget);
+			SelectedWidgetGeometry_RelativeToDesigner.AppendTransform(FSlateLayoutTransform(Inverse(CachedDesignerGeometry.AbsolutePosition)));
+			SelectedWidgetParentGeometry_RelativeToDesigner.AppendTransform(FSlateLayoutTransform(Inverse(CachedDesignerGeometry.AbsolutePosition)));
 
-			ParentPosition = ArrangedWidget.Geometry.AbsolutePosition;
-			ParentSize = ArrangedWidget.Geometry.Size * GetPreviewScale();
-		}
-	}
+			const FVector2D WidgetPostion_DesignerSpace = SelectedWidgetGeometry_RelativeToDesigner.AbsolutePosition;
+			const FVector2D WidgetSize = SelectedWidgetGeometry_RelativeToDesigner.Size * GetPreviewScale();
 
-	FVector2D FinalPosition(0, 0);
+			const FVector2D ParentPostion_DesignerSpace = SelectedWidgetParentGeometry_RelativeToDesigner.AbsolutePosition;
+			const FVector2D ParentSize = SelectedWidgetParentGeometry_RelativeToDesigner.Size * GetPreviewScale();
 
-	// Get the intial offset based on the location around the selected object.
-	switch ( ExtensionElement->GetLocation() )
-	{
-	case EExtensionLayoutLocation::Absolute:
-	{
-		FinalPosition = ParentPosition;
-		break;
-	}
-	case EExtensionLayoutLocation::TopLeft:
-		FinalPosition = TopLeft;
-		break;
-	case EExtensionLayoutLocation::TopCenter:
-		FinalPosition = TopLeft + FVector2D(Size.X * 0.5f, 0);
-		break;
-	case EExtensionLayoutLocation::TopRight:
-		FinalPosition = TopLeft + FVector2D(Size.X, 0);
-		break;
+			FVector2D FinalPosition(0, 0);
 
-	case EExtensionLayoutLocation::CenterLeft:
-		FinalPosition = TopLeft + FVector2D(0, Size.Y * 0.5f);
-		break;
-	case EExtensionLayoutLocation::CenterCenter:
-		FinalPosition = TopLeft + FVector2D(Size.X * 0.5f, Size.Y * 0.5f);
-		break;
-	case EExtensionLayoutLocation::CenterRight:
-		FinalPosition = TopLeft + FVector2D(Size.X, Size.Y * 0.5f);
-		break;
+			// Get the initial offset based on the location around the selected object.
+			switch ( ExtensionElement->GetLocation() )
+			{
+			case EExtensionLayoutLocation::Absolute:
+			{
+				FinalPosition = ParentPostion_DesignerSpace;
+				break;
+			}
+			case EExtensionLayoutLocation::TopLeft:
+				FinalPosition = WidgetPostion_DesignerSpace;
+				break;
+			case EExtensionLayoutLocation::TopCenter:
+				FinalPosition = WidgetPostion_DesignerSpace + FVector2D(WidgetSize.X * 0.5f, 0);
+				break;
+			case EExtensionLayoutLocation::TopRight:
+				FinalPosition = WidgetPostion_DesignerSpace + FVector2D(WidgetSize.X, 0);
+				break;
 
-	case EExtensionLayoutLocation::BottomLeft:
-		FinalPosition = TopLeft + FVector2D(0, Size.Y);
-		break;
-	case EExtensionLayoutLocation::BottomCenter:
-		FinalPosition = TopLeft + FVector2D(Size.X * 0.5f, Size.Y);
-		break;
-	case EExtensionLayoutLocation::BottomRight:
-		FinalPosition = TopLeft + Size;
-		break;
-	}
+			case EExtensionLayoutLocation::CenterLeft:
+				FinalPosition = WidgetPostion_DesignerSpace + FVector2D(0, WidgetSize.Y * 0.5f);
+				break;
+			case EExtensionLayoutLocation::CenterCenter:
+				FinalPosition = WidgetPostion_DesignerSpace + FVector2D(WidgetSize.X * 0.5f, WidgetSize.Y * 0.5f);
+				break;
+			case EExtensionLayoutLocation::CenterRight:
+				FinalPosition = WidgetPostion_DesignerSpace + FVector2D(WidgetSize.X, WidgetSize.Y * 0.5f);
+				break;
 
-	// Add the alignment offset
-	FinalPosition += ParentSize * ExtensionElement->GetAlignment();
+			case EExtensionLayoutLocation::BottomLeft:
+				FinalPosition = WidgetPostion_DesignerSpace + FVector2D(0, WidgetSize.Y);
+				break;
+			case EExtensionLayoutLocation::BottomCenter:
+				FinalPosition = WidgetPostion_DesignerSpace + FVector2D(WidgetSize.X * 0.5f, WidgetSize.Y);
+				break;
+			case EExtensionLayoutLocation::BottomRight:
+				FinalPosition = WidgetPostion_DesignerSpace + WidgetSize;
+				break;
+			}
+
+			// Add the alignment offset
+			FinalPosition += ParentSize * ExtensionElement->GetAlignment();
 
-	return FinalPosition + ExtensionElement->GetOffset();
+			return FinalPosition + ExtensionElement->GetOffset();
+		}
+	}
+
+	return FVector2D(0, 0);
 }
 
 FVector2D SDesignerView::GetExtensionSize(TSharedRef<FDesignerSurfaceElement> ExtensionElement) const
@@ -891,7 +888,7 @@ void SDesignerView::ResolvePendingSelectedWidgets()
 	{
 		TSet<FWidgetReference> SelectedTemplates;
 		SelectedTemplates.Add(PendingSelectedWidget);
-		BlueprintEditor.Pin()->SelectWidgets(SelectedTemplates);
+		BlueprintEditor.Pin()->SelectWidgets(SelectedTemplates, FSlateApplication::Get().GetModifierKeys().IsControlDown());
 
 		PendingSelectedWidget = FWidgetReference();
 	}
@@ -913,9 +910,9 @@ FReply SDesignerView::OnMouseButtonDown(const FGeometry& MyGeometry, const FPoin
 		if ( MouseEvent.GetEffectingButton() == EKeys::LeftMouseButton )
 		{
 			const bool bResolvePendingSelectionImmediately =
-				!SelectedWidget.IsValid() ||
-				!NewSelectedWidget.GetTemplate()->IsChildOf(SelectedWidget.GetTemplate()) ||
-				SelectedWidget.GetTemplate()->GetParent() == nullptr;
+				!GetSelectedWidget().IsValid() ||
+				!NewSelectedWidget.GetTemplate()->IsChildOf(GetSelectedWidget().GetTemplate()) ||
+				GetSelectedWidget().GetTemplate()->GetParent() == nullptr;
 
 			// If the newly clicked item is a child of the active selection, add it to the pending set of selected 
 			// widgets, if they begin dragging we can just move the parent, but if it's not part of the parent set, 
@@ -975,6 +972,8 @@ FReply SDesignerView::OnMouseMove(const FGeometry& MyGeometry, const FPointerEve
 
 	if ( MouseEvent.IsMouseButtonDown(EKeys::LeftMouseButton) )
 	{
+		FWidgetReference SelectedWidget = GetSelectedWidget();
+
 		if ( SelectedWidget.IsValid() && !bMovingExistingWidget )
 		{
 			if ( TransformMode == ETransformMode::Layout )
@@ -1017,25 +1016,23 @@ FReply SDesignerView::OnMouseMove(const FGeometry& MyGeometry, const FPointerEve
 	// Update the hovered widget under the mouse
 	FArrangedWidget ArrangedWidget(SNullWidget::NullWidget, FGeometry());
 	FWidgetReference NewHoveredWidget = GetWidgetAtCursor(MyGeometry, MouseEvent, ArrangedWidget);
-	if ( !( NewHoveredWidget == HoveredWidget ) )
-	{
-		HoveredWidget = NewHoveredWidget;
-		HoverTime = 0;
-	}
+	BlueprintEditor.Pin()->SetHoveredWidget(NewHoveredWidget);
 
 	return FReply::Unhandled();
 }
 
 void SDesignerView::OnMouseEnter(const FGeometry& MyGeometry, const FPointerEvent& MouseEvent)
 {
-	HoveredWidget = FWidgetReference();
-	HoverTime = 0;
+	SCompoundWidget::OnMouseEnter(MyGeometry, MouseEvent);
+
+	BlueprintEditor.Pin()->ClearHoveredWidget();
 }
 
 void SDesignerView::OnMouseLeave(const FPointerEvent& MouseEvent)
 {
-	HoveredWidget = FWidgetReference();
-	HoverTime = 0;
+	SCompoundWidget::OnMouseLeave(MouseEvent);
+
+	BlueprintEditor.Pin()->ClearHoveredWidget();
 }
 
 FReply SDesignerView::OnKeyDown(const FGeometry& MyGeometry, const FKeyboardEvent& InKeyboardEvent)
@@ -1085,92 +1082,71 @@ void SDesignerView::PopulateWidgetGeometryCache(FArrangedWidget& Root)
 	}
 }
 
-void SDesignerView::CacheSelectedWidgetGeometry()
-{
-	if ( SelectedSlateWidget.IsValid() )
-	{
-		TSharedRef<SWidget> Widget = SelectedSlateWidget.Pin().ToSharedRef();
-
-		FWidgetPath WidgetPath;
-		if ( FSlateApplication::Get().GeneratePathToWidgetUnchecked(Widget, WidgetPath) )
-		{
-			SelectedWidgetPath = FWeakWidgetPath(WidgetPath);
-		}
-		else
-		{
-			SelectedWidgetPath = FWeakWidgetPath();
-		}
-
-		FArrangedWidget ArrangedWidget(SNullWidget::NullWidget, FGeometry());
-		FDesignTimeUtils::GetArrangedWidgetRelativeToParent(WidgetPath, Widget, AsShared(), ArrangedWidget);
-
-		CachedDesignerWidgetLocation = ArrangedWidget.Geometry.AbsolutePosition;
-		CachedDesignerWidgetSize = ArrangedWidget.Geometry.Size;
-	}
-}
-
 int32 SDesignerView::HandleEffectsPainting(const FOnPaintHandlerParams& PaintArgs)
 {
-	TSet<FWidgetReference> Selected;
-	Selected.Add(SelectedWidget);
+	const TSet<FWidgetReference>& SelectedWidgets = GetSelectedWidgets();
 
 	// Allow the extensions to paint anything they want.
 	for ( const TSharedRef<FDesignerExtension>& Ext : DesignerExtensions )
 	{
-		Ext->Paint(Selected, PaintArgs.Geometry, PaintArgs.ClippingRect, PaintArgs.OutDrawElements, PaintArgs.Layer);
+		Ext->Paint(SelectedWidgets, PaintArgs.Geometry, PaintArgs.ClippingRect, PaintArgs.OutDrawElements, PaintArgs.Layer);
 	}
 
 	static const FName SelectionOutlineName("UMGEditor.SelectionOutline");
-	const FSlateBrush* SelectionOutlineBrush = FEditorStyle::Get().GetBrush(SelectionOutlineName);
-	FVector2D SelectionBrushInflationAmount = FVector2D(16, 16) * FVector2D(SelectionOutlineBrush->Margin.Left, SelectionOutlineBrush->Margin.Top) * ( 1.0f / GetPreviewScale() );
-
-	// Don't draw the hovered effect if it's also the selected widget
-	if ( HoveredSlateWidget.IsValid() && HoveredSlateWidget != SelectedSlateWidget )
-	{
-		TSharedRef<SWidget> Widget = HoveredSlateWidget.Pin().ToSharedRef();
-
-		FArrangedWidget ArrangedWidget(SNullWidget::NullWidget, FGeometry());
-		FDesignTimeUtils::GetArrangedWidgetRelativeToWindow(Widget, ArrangedWidget);
-
-		// Draw hovered effect
-		// Azure = 0x007FFF
-		const FLinearColor HoveredTint(0, 0.5, 1, FMath::Clamp(HoverTime / HoveredAnimationTime, 0.0f, 1.0f));
-
-		FPaintGeometry HoveredGeometry = ArrangedWidget.Geometry.ToInflatedPaintGeometry(SelectionBrushInflationAmount);
 
-		FSlateDrawElement::MakeBox(
-			PaintArgs.OutDrawElements,
-			PaintArgs.Layer,
-			HoveredGeometry,
-			SelectionOutlineBrush,
-			PaintArgs.ClippingRect,
-			ESlateDrawEffect::None,
-			HoveredTint
-			);
-	}
+	static const FLinearColor SelectedTint(0, 1, 0);
+	const FLinearColor HoveredTint(0, 0.5, 1, FMath::Clamp(BlueprintEditor.Pin()->GetHoveredWidgetTime() / HoveredAnimationTime, 0.0f, 1.0f)); // Azure = 0x007FFF
 
-	if ( SelectedSlateWidget.IsValid() )
+	for ( const FWidgetReference& SelectedWidget : SelectedWidgets )
 	{
-		TSharedRef<SWidget> Widget = SelectedSlateWidget.Pin().ToSharedRef();
+		const FSlateBrush* SelectionOutlineBrush = FEditorStyle::Get().GetBrush(SelectionOutlineName);
+		FVector2D SelectionBrushInflationAmount = FVector2D(16, 16) * FVector2D(SelectionOutlineBrush->Margin.Left, SelectionOutlineBrush->Margin.Top) * ( 1.0f / GetPreviewScale() );
 
-		FArrangedWidget ArrangedWidget(SNullWidget::NullWidget, FGeometry());
-		FDesignTimeUtils::GetArrangedWidgetRelativeToWindow(Widget, ArrangedWidget);
+		TSharedPtr<SWidget> SelectedSlateWidget = SelectedWidget.GetPreviewSlate();
+		TSharedPtr<SWidget> HoveredSlateWidget = BlueprintEditor.Pin()->GetHoveredWidget().GetPreviewSlate();
 
-		const FLinearColor Tint(0, 1, 0);
+		// Don't draw the hovered effect if it's also the selected widget
+		if ( HoveredSlateWidget.IsValid() && HoveredSlateWidget != SelectedSlateWidget )
+		{
+			TSharedRef<SWidget> Widget = HoveredSlateWidget.ToSharedRef();
 
-		// Draw selection effect
+			FArrangedWidget ArrangedWidget(SNullWidget::NullWidget, FGeometry());
+			FDesignTimeUtils::GetArrangedWidgetRelativeToWindow(Widget, ArrangedWidget);
+
+			// Draw hovered effect
+			FPaintGeometry HoveredGeometry = ArrangedWidget.Geometry.ToInflatedPaintGeometry(SelectionBrushInflationAmount);
+
+			FSlateDrawElement::MakeBox(
+				PaintArgs.OutDrawElements,
+				PaintArgs.Layer,
+				HoveredGeometry,
+				SelectionOutlineBrush,
+				PaintArgs.ClippingRect,
+				ESlateDrawEffect::None,
+				HoveredTint
+				);
+		}
 
-		FPaintGeometry SelectionGeometry = ArrangedWidget.Geometry.ToInflatedPaintGeometry(SelectionBrushInflationAmount);
+		if ( SelectedSlateWidget.IsValid() )
+		{
+			TSharedRef<SWidget> Widget = SelectedSlateWidget.ToSharedRef();
 
-		FSlateDrawElement::MakeBox(
-			PaintArgs.OutDrawElements,
-			PaintArgs.Layer,
-			SelectionGeometry,
-			SelectionOutlineBrush,
-			PaintArgs.ClippingRect,
-			ESlateDrawEffect::None,
-			Tint
-			);
+			FArrangedWidget ArrangedWidget(SNullWidget::NullWidget, FGeometry());
+			FDesignTimeUtils::GetArrangedWidgetRelativeToWindow(Widget, ArrangedWidget);
+
+			// Draw selection effect
+			FPaintGeometry SelectionGeometry = ArrangedWidget.Geometry.ToInflatedPaintGeometry(SelectionBrushInflationAmount);
+
+			FSlateDrawElement::MakeBox(
+				PaintArgs.OutDrawElements,
+				PaintArgs.Layer,
+				SelectionGeometry,
+				SelectionOutlineBrush,
+				PaintArgs.ClippingRect,
+				ESlateDrawEffect::None,
+				SelectedTint
+				);
+		}
 	}
 
 	return PaintArgs.Layer + 1;
@@ -1194,7 +1170,7 @@ void SDesignerView::UpdatePreviewWidget(bool bForceUpdate)
 			// Notify all selected widgets that they are selected, because there are new preview objects
 			// state may have been lost so this will recreate it if the widget does something special when
 			// selected.
-			for ( FWidgetReference& WidgetRef : SelectedWidgets )
+			for ( const FWidgetReference& WidgetRef : GetSelectedWidgets() )
 			{
 				if ( WidgetRef.IsValid() )
 				{
@@ -1222,41 +1198,15 @@ void SDesignerView::UpdatePreviewWidget(bool bForceUpdate)
 void SDesignerView::Tick(const FGeometry& AllottedGeometry, const double InCurrentTime, const float InDeltaTime)
 {
 	CachedDesignerGeometry = AllottedGeometry;
-	HoverTime += InDeltaTime;
 
 	const bool bForceUpdate = false;
 	UpdatePreviewWidget(bForceUpdate);
 
-	// Update the selected widget to match the selected template.
-	if ( PreviewWidget )
-	{
-		if ( SelectedWidget.IsValid() )
-		{
-			// Set the selected widget so that we can draw the highlight
-			SelectedSlateWidget = PreviewWidget->GetSlateWidgetFromName(SelectedWidget.GetTemplate()->GetFName());
-		}
-		else
-		{
-			SelectedSlateWidget.Reset();
-		}
-
-		if ( HoveredWidget.IsValid() )
-		{
-			HoveredSlateWidget = PreviewWidget->GetSlateWidgetFromName(HoveredWidget.GetTemplate()->GetFName());
-		}
-		else
-		{
-			HoveredSlateWidget.Reset();
-		}
-	}
-
 	// Perform an arrange children pass to cache the geometry of all widgets so that we can query it later.
 	CachedWidgetGeometry.Reset();
 	FArrangedWidget WindowWidgetGeometry(PreviewHitTestRoot.ToSharedRef(), AllottedGeometry);
 	PopulateWidgetGeometryCache(WindowWidgetGeometry);
 
-	CacheSelectedWidgetGeometry();
-
 	// Tick all designer extensions in case they need to update widgets
 	for ( const TSharedRef<FDesignerExtension>& Ext : DesignerExtensions )
 	{
@@ -1268,6 +1218,8 @@ void SDesignerView::Tick(const FGeometry& AllottedGeometry, const double InCurre
 
 FReply SDesignerView::OnDragDetected(const FGeometry& MyGeometry, const FPointerEvent& MouseEvent)
 {
+	FWidgetReference SelectedWidget = GetSelectedWidget();
+
 	if ( SelectedWidget.IsValid() )
 	{
 		// Clear any pending selected widgets, the user has already decided what widget they want.
@@ -1288,11 +1240,19 @@ FReply SDesignerView::OnDragDetected(const FGeometry& MyGeometry, const FPointer
 
 void SDesignerView::OnDragEnter(const FGeometry& MyGeometry, const FDragDropEvent& DragDropEvent)
 {
+	SCompoundWidget::OnDragEnter(MyGeometry, DragDropEvent);
+
+	BlueprintEditor.Pin()->ClearHoveredWidget();
+
 	//@TODO UMG Drop Feedback
 }
 
 void SDesignerView::OnDragLeave(const FDragDropEvent& DragDropEvent)
 {
+	SCompoundWidget::OnDragLeave(DragDropEvent);
+
+	BlueprintEditor.Pin()->ClearHoveredWidget();
+
 	TSharedPtr<FDecoratedDragDropOp> DecoratedDragDropOp = DragDropEvent.GetOperationAs<FDecoratedDragDropOp>();
 	if ( DecoratedDragDropOp.IsValid() )
 	{
@@ -1376,6 +1336,8 @@ UWidget* SDesignerView::ProcessDropAndAddWidget(const FGeometry& MyGeometry, con
 	TSharedPtr<FWidgetTemplateDragDropOp> TemplateDragDropOp = DragDropEvent.GetOperationAs<FWidgetTemplateDragDropOp>();
 	if ( TemplateDragDropOp.IsValid() )
 	{
+		BlueprintEditor.Pin()->SetHoveredWidget(WidgetUnderCursor);
+
 		TemplateDragDropOp->SetCursorOverride(TOptional<EMouseCursor::Type>());
 
 		// If there's no root widget go ahead and add the widget into the root slot.
@@ -1397,8 +1359,6 @@ UWidget* SDesignerView::ProcessDropAndAddWidget(const FGeometry& MyGeometry, con
 
 			BP->WidgetTree->RootWidget = Widget;
 
-			SelectedWidget = BlueprintEditor.Pin()->GetReferenceFromTemplate(Widget);
-
 			DropPreviewParent = nullptr;
 
 			if ( bIsPreview )
@@ -1505,6 +1465,9 @@ UWidget* SDesignerView::ProcessDropAndAddWidget(const FGeometry& MyGeometry, con
 			}
 		}
 
+		FWidgetReference TargetReference = bIsPreview ? BlueprintEditor.Pin()->GetReferenceFromPreview(Target) : BlueprintEditor.Pin()->GetReferenceFromTemplate(Target);
+		BlueprintEditor.Pin()->SetHoveredWidget(TargetReference);
+
 		// If the widget being hovered over is a panel, attempt to place it into that panel.
 		if ( Target && Target->IsA(UPanelWidget::StaticClass()) )
 		{
@@ -1880,10 +1843,13 @@ void SDesignerView::BeginTransaction(const FText& SessionName)
 	{
 		ScopedTransaction = new FScopedTransaction(SessionName);
 
-		if ( SelectedWidget.IsValid() )
+		for ( const FWidgetReference& SelectedWidget : GetSelectedWidgets() )
 		{
-			SelectedWidget.GetPreview()->Modify();
-			SelectedWidget.GetTemplate()->Modify();
+			if ( SelectedWidget.IsValid() )
+			{
+				SelectedWidget.GetPreview()->Modify();
+				SelectedWidget.GetTemplate()->Modify();
+			}
 		}
 	}
 }
diff --git a/Engine/Source/Editor/UMGEditor/Private/Designer/SDesignerView.h b/Engine/Source/Editor/UMGEditor/Private/Designer/SDesignerView.h
--- a/Engine/Source/Editor/UMGEditor/Private/Designer/SDesignerView.h
+++ b/Engine/Source/Editor/UMGEditor/Private/Designer/SDesignerView.h
@@ -58,6 +58,7 @@ public:
 
 	// IUMGDesigner interface
 	virtual float GetPreviewScale() const override;
+	virtual const TSet<FWidgetReference>& GetSelectedWidgets() const override;
 	virtual FWidgetReference GetSelectedWidget() const override;
 	virtual ETransformMode::Type GetTransformMode() const override;
 	virtual FGeometry GetDesignerGeometry() const override;
@@ -113,8 +114,6 @@ private:
 
 	void PopulateWidgetGeometryCache(FArrangedWidget& Root);
 
-	void CacheSelectedWidgetGeometry();
-
 	/** @return Formatted text for the given resolution params */
 	FText GetResolutionText(int32 Width, int32 Height, const FString& AspectRatio) const;
 
@@ -195,22 +194,12 @@ private:
 	TSharedPtr<SCanvas> ExtensionWidgetCanvas;
 	TSharedPtr<SPaintSurface> EffectsLayer;
 
-	FVector2D CachedDesignerWidgetLocation;
-	FVector2D CachedDesignerWidgetSize;
-
-	/** The currently selected set of widgets */
-	TSet< FWidgetReference > SelectedWidgets;
-
-	/** TODO UMG Remove, after getting multiselection working. */
-	FWidgetReference SelectedWidget;
+	/** The currently selected preview widgets in the preview GUI, just a cache used to determine changes between selection changes. */
+	TSet< FWidgetReference > SelectedWidgetsCache;
 
 	/** The location in selected widget local space where the context menu was summoned. */
 	FVector2D SelectedWidgetContextMenuLocation;
 
-	/** The currently selected slate widget, this is refreshed every frame in case it changes */
-	TArray< TWeakPtr<SWidget> > SelectedSlateWidgets;
-	TWeakPtr<SWidget> SelectedSlateWidget;
-
 	/**
 	 * Holds onto a temporary widget that the user may be getting ready to select, or may just 
 	 * be the widget that got hit on the initial mouse down before moving the parent.
@@ -223,15 +212,6 @@ private:
 	/** An existing widget is being moved in its current container, or in to a new container. */
 	bool bMovingExistingWidget;
 
-	/** The wall clock time the user has been hovering over a single widget */
-	float HoverTime;
-
-	/** The current widget being hovered */
-	FWidgetReference HoveredWidget;
-
-	/** The current slate widget being hovered, this is refreshed every frame in case it changes */
-	TWeakPtr<SWidget> HoveredSlateWidget;
-
 	/** The configured Width of the preview area, simulates screen size. */
 	int32 PreviewWidth;
 
diff --git a/Engine/Source/Editor/UMGEditor/Private/Designer/STransformHandle.cpp b/Engine/Source/Editor/UMGEditor/Private/Designer/STransformHandle.cpp
--- a/Engine/Source/Editor/UMGEditor/Private/Designer/STransformHandle.cpp
+++ b/Engine/Source/Editor/UMGEditor/Private/Designer/STransformHandle.cpp
@@ -42,23 +42,26 @@ EVisibility STransformHandle::GetHandleVisibility() const
 
 	// Only show the handles for visible elements in the designer.
 	FWidgetReference SelectedWidget = Designer->GetSelectedWidget();
-	if ( !SelectedWidget.GetTemplate()->bHiddenInDesigner )
+	if ( SelectedWidget.IsValid() )
 	{
-		switch ( TransformMode )
+		if ( !SelectedWidget.GetTemplate()->bHiddenInDesigner )
 		{
-		case ETransformMode::Layout:
-		{
-			if ( UPanelSlot* TemplateSlot = SelectedWidget.GetTemplate()->Slot )
+			switch ( TransformMode )
+			{
+			case ETransformMode::Layout:
 			{
-				if ( TemplateSlot->CanResize(DragDirection) )
+				if ( UPanelSlot* TemplateSlot = SelectedWidget.GetTemplate()->Slot )
 				{
-					return EVisibility::Visible;
+					if ( TemplateSlot->CanResize(DragDirection) )
+					{
+						return EVisibility::Visible;
+					}
 				}
+				break;
+			}
+			case ETransformMode::Render:
+				return EVisibility::Visible;
 			}
-			break;
-		}
-		case ETransformMode::Render:
-			return EVisibility::Visible;
 		}
 	}
 
diff --git a/Engine/Source/Editor/UMGEditor/Private/Hierarchy/SHierarchyView.cpp b/Engine/Source/Editor/UMGEditor/Private/Hierarchy/SHierarchyView.cpp
--- a/Engine/Source/Editor/UMGEditor/Private/Hierarchy/SHierarchyView.cpp
+++ b/Engine/Source/Editor/UMGEditor/Private/Hierarchy/SHierarchyView.cpp
@@ -26,6 +26,7 @@ void SHierarchyView::Construct(const FArguments& InArgs, TSharedPtr<FWidgetBluep
 {
 	BlueprintEditor = InBlueprintEditor;
 	bRebuildTreeRequested = false;
+	bIsUpdatingSelection = false;
 
 	// register for any objects replaced
 	GEditor->OnObjectsReplaced().AddRaw(this, &SHierarchyView::OnObjectsReplaced);
@@ -92,6 +93,8 @@ SHierarchyView::~SHierarchyView()
 
 void SHierarchyView::Tick(const FGeometry& AllottedGeometry, const double InCurrentTime, const float InDeltaTime)
 {
+	SCompoundWidget::Tick(AllottedGeometry, InCurrentTime, InDeltaTime);
+
 	if ( bRebuildTreeRequested || bRefreshRequested )
 	{
 		if ( bRebuildTreeRequested )
@@ -112,6 +115,20 @@ void SHierarchyView::Tick(const FGeometry& AllottedGeometry, const double InCurr
 	}
 }
 
+void SHierarchyView::OnMouseEnter(const FGeometry& MyGeometry, const FPointerEvent& MouseEvent)
+{
+	SCompoundWidget::OnMouseEnter(MyGeometry, MouseEvent);
+
+	BlueprintEditor.Pin()->ClearHoveredWidget();
+}
+
+void SHierarchyView::OnMouseLeave(const FPointerEvent& MouseEvent)
+{
+	SCompoundWidget::OnMouseLeave(MouseEvent);
+
+	BlueprintEditor.Pin()->ClearHoveredWidget();
+}
+
 FReply SHierarchyView::OnKeyDown(const FGeometry& MyGeometry, const FKeyboardEvent& InKeyboardEvent)
 {
 	BlueprintEditor.Pin()->PasteDropLocation = FVector2D(0, 0);
@@ -143,14 +160,17 @@ FText SHierarchyView::GetSearchText() const
 
 void SHierarchyView::OnEditorSelectionChanged()
 {
-	WidgetTreeView->ClearSelection();
-
-	if ( RootWidgets.Num() > 0 )
+	if ( !bIsUpdatingSelection )
 	{
-		RootWidgets[0]->RefreshSelection();
-	}
+		WidgetTreeView->ClearSelection();
+
+		if ( RootWidgets.Num() > 0 )
+		{
+			RootWidgets[0]->RefreshSelection();
+		}
 
-	RestoreSelectedItems();
+		RestoreSelectedItems();
+	}
 }
 
 UWidgetBlueprint* SHierarchyView::GetBlueprint() const
@@ -161,7 +181,7 @@ UWidgetBlueprint* SHierarchyView::GetBlueprint() const
 		return CastChecked<UWidgetBlueprint>(BP);
 	}
 
-	return NULL;
+	return nullptr;
 }
 
 void SHierarchyView::OnBlueprintChanged(UBlueprint* InBlueprint)
@@ -174,7 +194,7 @@ void SHierarchyView::OnBlueprintChanged(UBlueprint* InBlueprint)
 
 TSharedPtr<SWidget> SHierarchyView::WidgetHierarchy_OnContextMenuOpening()
 {
-	FMenuBuilder MenuBuilder(true, NULL);
+	FMenuBuilder MenuBuilder(true, nullptr);
 
 	FWidgetBlueprintEditorUtils::CreateWidgetContextMenu(MenuBuilder, BlueprintEditor.Pin().ToSharedRef(), FVector2D(0, 0));
 
@@ -194,9 +214,26 @@ TSharedRef< ITableRow > SHierarchyView::WidgetHierarchy_OnGenerateRow(TSharedPtr
 
 void SHierarchyView::WidgetHierarchy_OnSelectionChanged(TSharedPtr<FHierarchyModel> SelectedItem, ESelectInfo::Type SelectInfo)
 {
-	if ( SelectedItem.IsValid() && SelectInfo != ESelectInfo::Direct )
+	if ( SelectInfo != ESelectInfo::Direct )
 	{
-		SelectedItem->OnSelection();
+		bIsUpdatingSelection = true;
+
+		TArray< TSharedPtr<FHierarchyModel> > SelectedItems = WidgetTreeView->GetSelectedItems();
+
+		TSet<FWidgetReference> Clear;
+		BlueprintEditor.Pin()->SelectWidgets(Clear, false);
+
+		for ( TSharedPtr<FHierarchyModel>& Item : SelectedItems )
+		{
+			Item->OnSelection();
+		}
+
+		if ( RootWidgets.Num() > 0 )
+		{
+			RootWidgets[0]->RefreshSelection();
+		}
+
+		bIsUpdatingSelection = false;
 	}
 }
 
@@ -221,7 +258,7 @@ void SHierarchyView::RebuildTreeView()
 {
 	SAssignNew(WidgetTreeView, STreeView< TSharedPtr<FHierarchyModel> >)
 		.ItemHeight(20.0f)
-		.SelectionMode(ESelectionMode::Single)
+		.SelectionMode(ESelectionMode::Multi)
 		.OnGetChildren(FilterHandler.ToSharedRef(), &TreeFilterHandler< TSharedPtr<FHierarchyModel> >::OnGetFilteredChildren)
 		.OnGenerateRow(this, &SHierarchyView::WidgetHierarchy_OnGenerateRow)
 		.OnSelectionChanged(this, &SHierarchyView::WidgetHierarchy_OnSelectionChanged)
diff --git a/Engine/Source/Editor/UMGEditor/Private/Hierarchy/SHierarchyView.h b/Engine/Source/Editor/UMGEditor/Private/Hierarchy/SHierarchyView.h
--- a/Engine/Source/Editor/UMGEditor/Private/Hierarchy/SHierarchyView.h
+++ b/Engine/Source/Editor/UMGEditor/Private/Hierarchy/SHierarchyView.h
@@ -33,8 +33,9 @@ public:
 	virtual ~SHierarchyView();
 
 	// Begin SWidget
-	virtual void Tick(const FGeometry& AllottedGeometry, const double InCurrentTime, const float InDeltaTime);
-
+	virtual void Tick(const FGeometry& AllottedGeometry, const double InCurrentTime, const float InDeltaTime) override;
+	virtual void OnMouseEnter(const FGeometry& MyGeometry, const FPointerEvent& MouseEvent) override;
+	virtual void OnMouseLeave(const FPointerEvent& MouseEvent) override;
 	virtual FReply OnKeyDown(const FGeometry& MyGeometry, const FKeyboardEvent& InKeyboardEvent) override;
 	// End SWidget
 
@@ -123,12 +124,15 @@ private:
 	/** The filter used by the search box */
 	TSharedPtr<WidgetTextFilter> SearchBoxWidgetFilter;
 
+	/** Temporary expanded item state, used to restore expansion after tree rebuilds. */
+	TArray<FName> ExpandedItems;
+
 	/** Has a full refresh of the tree been requested?  This happens when the user is filtering the tree */
 	bool bRefreshRequested;
 
 	/** Is the tree in such a changed state that the whole widget needs rebuilding? */
 	bool bRebuildTreeRequested;
 
-	/** Temporary expanded item state, used to restore expansion after tree rebuilds. */
-	TArray<FName> ExpandedItems;
+	/** Flag to ignore selections while the hierarchy view is updating the selection. */
+	bool bIsUpdatingSelection;
 };
diff --git a/Engine/Source/Editor/UMGEditor/Private/Hierarchy/SHierarchyViewItem.cpp b/Engine/Source/Editor/UMGEditor/Private/Hierarchy/SHierarchyViewItem.cpp
--- a/Engine/Source/Editor/UMGEditor/Private/Hierarchy/SHierarchyViewItem.cpp
+++ b/Engine/Source/Editor/UMGEditor/Private/Hierarchy/SHierarchyViewItem.cpp
@@ -448,7 +448,7 @@ void FNamedSlotModel::OnSelection()
 			SelectedWidgets.Add(BPEd->GetReferenceFromTemplate(TemplateSlotContent));
 		}
 
-		BPEd->SelectWidgets(SelectedWidgets);
+		BPEd->SelectWidgets(SelectedWidgets, true);
 	}
 }
 
@@ -691,7 +691,17 @@ void FHierarchyWidget::OnSelection()
 	TSet<FWidgetReference> SelectedWidgets;
 	SelectedWidgets.Add(Item);
 
-	BlueprintEditor.Pin()->SelectWidgets(SelectedWidgets);
+	BlueprintEditor.Pin()->SelectWidgets(SelectedWidgets, true);
+}
+
+void FHierarchyWidget::OnMouseEnter()
+{
+	BlueprintEditor.Pin()->SetHoveredWidget(Item);
+}
+
+void FHierarchyWidget::OnMouseLeave()
+{
+	BlueprintEditor.Pin()->ClearHoveredWidget();
 }
 
 void FHierarchyWidget::UpdateSelection()
@@ -764,6 +774,20 @@ void SHierarchyViewItem::Construct(const FArguments& InArgs, const TSharedRef< S
 		InOwnerTableView);
 }
 
+void SHierarchyViewItem::OnMouseEnter(const FGeometry& MyGeometry, const FPointerEvent& MouseEvent)
+{
+	STableRow< TSharedPtr<FHierarchyModel> >::OnMouseEnter(MyGeometry, MouseEvent);
+
+	Model->OnMouseEnter();
+}
+
+void SHierarchyViewItem::OnMouseLeave(const FPointerEvent& MouseEvent)
+{
+	STableRow< TSharedPtr<FHierarchyModel> >::OnMouseLeave(MouseEvent);
+
+	Model->OnMouseLeave();
+}
+
 bool SHierarchyViewItem::OnVerifyNameTextChanged(const FText& InText, FText& OutErrorMessage)
 {
 	return Model->OnVerifyNameTextChanged(InText, OutErrorMessage);
diff --git a/Engine/Source/Editor/UMGEditor/Private/Hierarchy/SHierarchyViewItem.h b/Engine/Source/Editor/UMGEditor/Private/Hierarchy/SHierarchyViewItem.h
--- a/Engine/Source/Editor/UMGEditor/Private/Hierarchy/SHierarchyViewItem.h
+++ b/Engine/Source/Editor/UMGEditor/Private/Hierarchy/SHierarchyViewItem.h
@@ -40,6 +40,9 @@ public:
 	virtual void GatherChildren(TArray< TSharedPtr<FHierarchyModel> >& Children);
 
 	virtual void OnSelection() = 0;
+
+	virtual void OnMouseEnter() {}
+	virtual void OnMouseLeave() {}
 	
 	void RefreshSelection();
 	bool ContainsSelection();
@@ -156,6 +159,9 @@ public:
 
 	virtual void OnSelection();
 
+	virtual void OnMouseEnter() override;
+	virtual void OnMouseLeave() override;
+
 	virtual bool IsVisible() const override
 	{
 		return !Item.GetTemplate()->bHiddenInDesigner;
@@ -205,6 +211,8 @@ public:
 	void Construct(const FArguments& InArgs, const TSharedRef< STableViewBase >& InOwnerTableView, TSharedPtr<FHierarchyModel> InModel);
 
 	// Begin SWidget
+	void OnMouseEnter(const FGeometry& MyGeometry, const FPointerEvent& MouseEvent);
+	void OnMouseLeave(const FPointerEvent& MouseEvent);
 	virtual FReply OnDragOver(const FGeometry& MyGeometry, const FDragDropEvent& DragDropEvent) override;
 	// End SWidget
 
diff --git a/Engine/Source/Editor/UMGEditor/Private/WidgetBlueprintEditor.cpp b/Engine/Source/Editor/UMGEditor/Private/WidgetBlueprintEditor.cpp
--- a/Engine/Source/Editor/UMGEditor/Private/WidgetBlueprintEditor.cpp
+++ b/Engine/Source/Editor/UMGEditor/Private/WidgetBlueprintEditor.cpp
@@ -27,8 +27,8 @@
 FWidgetBlueprintEditor::FWidgetBlueprintEditor()
 	: PreviewScene(FPreviewScene::ConstructionValues().AllowAudioPlayback(true).ShouldSimulatePhysics(true))
 	, PreviewBlueprint(nullptr)
+	, HoverTime(0)
 {
-	
 }
 
 FWidgetBlueprintEditor::~FWidgetBlueprintEditor()
@@ -123,7 +123,7 @@ void FWidgetBlueprintEditor::RegisterApplicationModes(const TArray<UBlueprint*>&
 	}
 }
 
-void FWidgetBlueprintEditor::SelectWidgets(const TSet<FWidgetReference>& Widgets)
+void FWidgetBlueprintEditor::SelectWidgets(const TSet<FWidgetReference>& Widgets, bool bAppendOrToggle)
 {
 	TSet<FWidgetReference> TempSelection;
 	for ( const FWidgetReference& Widget : Widgets )
@@ -139,10 +139,23 @@ void FWidgetBlueprintEditor::SelectWidgets(const TSet<FWidgetReference>& Widgets
 	// Finally change the selected widgets after we've updated the details panel 
 	// to ensure values that are pending are committed on focus loss, and migrated properly
 	// to the old selected widgets.
-	SelectedWidgets.Empty();
+	if ( !bAppendOrToggle )
+	{
+		SelectedWidgets.Empty();
+	}
 	SelectedObjects.Empty();
 
-	SelectedWidgets.Append(TempSelection);
+	for ( const FWidgetReference& Widget : TempSelection )
+	{
+		if ( bAppendOrToggle && SelectedWidgets.Contains(Widget) )
+		{
+			SelectedWidgets.Remove(Widget);
+		}
+		else
+		{
+			SelectedWidgets.Add(Widget);
+		}
+	}
 
 	OnSelectedWidgetsChanged.Broadcast();
 }
@@ -183,7 +196,7 @@ void FWidgetBlueprintEditor::CleanSelection()
 
 	if ( TempSelection.Num() != SelectedWidgets.Num() )
 	{
-		SelectWidgets(TempSelection);
+		SelectWidgets(TempSelection, false);
 	}
 }
 
@@ -247,7 +260,7 @@ void FWidgetBlueprintEditor::DeleteSelectedWidgets()
 
 	// Clear the selection now that the widget has been deleted.
 	TSet<FWidgetReference> Empty;
-	SelectWidgets(Empty);
+	SelectWidgets(Empty, false);
 }
 
 bool FWidgetBlueprintEditor::CanCopySelectedWidgets()
@@ -329,6 +342,8 @@ void FWidgetBlueprintEditor::Tick(float DeltaTime)
 {
 	FBlueprintEditor::Tick(DeltaTime);
 
+	HoverTime += DeltaTime;
+
 	// Tick the preview scene world.
 	if ( !GIntraFrameDebuggingGameThread )
 	{
@@ -650,4 +665,29 @@ FGraphAppearanceInfo FWidgetBlueprintEditor::GetGraphAppearance() const
 	return AppearanceInfo;
 }
 
+void FWidgetBlueprintEditor::ClearHoveredWidget()
+{
+	HoveredWidget = FWidgetReference();
+	HoverTime = 0;
+}
+
+void FWidgetBlueprintEditor::SetHoveredWidget(FWidgetReference& InHoveredWidget)
+{
+	if ( !( InHoveredWidget == HoveredWidget ) )
+	{
+		HoveredWidget = InHoveredWidget;
+		HoverTime = 0;
+	}
+}
+
+FWidgetReference FWidgetBlueprintEditor::GetHoveredWidget() const
+{
+	return HoveredWidget;
+}
+
+float FWidgetBlueprintEditor::GetHoveredWidgetTime() const
+{
+	return HoverTime;
+}
+
 #undef LOCTEXT_NAMESPACE
diff --git a/Engine/Source/Editor/UMGEditor/Private/WidgetBlueprintEditor.h b/Engine/Source/Editor/UMGEditor/Private/WidgetBlueprintEditor.h
--- a/Engine/Source/Editor/UMGEditor/Private/WidgetBlueprintEditor.h
+++ b/Engine/Source/Editor/UMGEditor/Private/WidgetBlueprintEditor.h
@@ -65,7 +65,7 @@ public:
 	const UWidgetAnimation* RefreshCurrentAnimation();
 
 	/** Sets the currently selected set of widgets */
-	void SelectWidgets(const TSet<FWidgetReference>& Widgets);
+	void SelectWidgets(const TSet<FWidgetReference>& Widgets, bool bAppendOrToggle);
 
 	/** Sets the currently selected set of objects */
 	void SelectObjects(const TSet<UObject*>& Objects);
@@ -94,6 +94,18 @@ public:
 	/** Creates a sequencer widget */
 	TSharedRef<SWidget> CreateSequencerWidget();
 
+	/**
+	 * The widget we're now hovering over in any particular context, allows multiple views to 
+	 * synchronize feedback on where that widget is in their representation.
+	 */
+	void SetHoveredWidget(FWidgetReference& InHoveredWidget);
+
+	void ClearHoveredWidget();
+
+	FWidgetReference GetHoveredWidget() const;
+
+	float GetHoveredWidgetTime() const;
+
 public:
 	/** Fires whenever the selected set of widgets changing */
 	FOnSelectedWidgetsChanged OnSelectedWidgetsChanging;
@@ -156,7 +168,7 @@ private:
 	UWidgetBlueprint* PreviewBlueprint;
 
 	/** The currently selected preview widgets in the preview GUI */
-	TSet<FWidgetReference> SelectedWidgets;
+	TSet< FWidgetReference > SelectedWidgets;
 
 	/** The currently selected objects in the designer */
 	TSet< TWeakObjectPtr<UObject> > SelectedObjects;
@@ -176,6 +188,12 @@ private:
 	/** The widget references out in the ether that may need to be updated after being issued. */
 	TArray< TWeakPtr<FWidgetHandle> > WidgetHandlePool;
 
+	/** The wall clock time the user has been hovering over a single widget */
+	float HoverTime;
+
+	/** The current widget being hovered */
+	FWidgetReference HoveredWidget;
+
 	/** The preview becomes invalid and needs to be rebuilt on the next tick. */
 	bool bPreviewInvalidated;
 };
diff --git a/Engine/Source/Editor/UMGEditor/Public/IUMGDesigner.h b/Engine/Source/Editor/UMGEditor/Public/IUMGDesigner.h
--- a/Engine/Source/Editor/UMGEditor/Public/IUMGDesigner.h
+++ b/Engine/Source/Editor/UMGEditor/Public/IUMGDesigner.h
@@ -28,6 +28,9 @@ public:
 	/** @return the effective preview scale after both the DPI and Zoom scale has been applied. */
 	virtual float GetPreviewScale() const = 0;
 
+	/** @return The currently selected widgets */
+	virtual const TSet<FWidgetReference>& GetSelectedWidgets() const = 0;
+
 	/** @return The currently selected widget. */
 	virtual FWidgetReference GetSelectedWidget() const = 0;
 
diff --git a/Engine/Source/Editor/UMGEditor/Public/WidgetReference.h b/Engine/Source/Editor/UMGEditor/Public/WidgetReference.h
--- a/Engine/Source/Editor/UMGEditor/Public/WidgetReference.h
+++ b/Engine/Source/Editor/UMGEditor/Public/WidgetReference.h
@@ -39,6 +39,17 @@ public:
 	/** @returns the preview widget.  This is the transient representation of the template.  Constantly being destroyed and recreated.  Do not cache this pointer. */
 	UWidget* GetPreview() const;
 
+	/** @returns the preview slate widget.  This is the transient representation of the template.  Constantly being destroyed and recreated.  Do not cache this pointer. */
+	TSharedPtr<SWidget> GetPreviewSlate() const
+	{
+		if ( IsValid() )
+		{
+			return GetPreview()->GetCachedWidget();
+		}
+
+		return TSharedPtr<SWidget>();
+	}
+
 	/** Checks if widget reference is the same as nother widget reference, based on the template pointers. */
 	bool operator==( const FWidgetReference& Other ) const
 	{
diff --git a/Engine/Source/Runtime/UMG/Private/Slate/SObjectWidget.cpp b/Engine/Source/Runtime/UMG/Private/Slate/SObjectWidget.cpp
--- a/Engine/Source/Runtime/UMG/Private/Slate/SObjectWidget.cpp
+++ b/Engine/Source/Runtime/UMG/Private/Slate/SObjectWidget.cpp
@@ -212,7 +212,7 @@ FReply SObjectWidget::OnMouseButtonDoubleClick(const FGeometry& MyGeometry, cons
 
 FReply SObjectWidget::OnDragDetected(const FGeometry& MyGeometry, const FPointerEvent& PointerEvent)
 {
-	UDragDropOperation* Operation = NULL;
+	UDragDropOperation* Operation = nullptr;
 	WidgetObject->OnDragDetected(MyGeometry, PointerEvent, Operation);
 
 	if ( Operation )
diff --git a/Engine/Source/Runtime/UMG/Public/Components/PanelSlot.h b/Engine/Source/Runtime/UMG/Public/Components/PanelSlot.h
--- a/Engine/Source/Runtime/UMG/Public/Components/PanelSlot.h
+++ b/Engine/Source/Runtime/UMG/Public/Components/PanelSlot.h
@@ -5,7 +5,7 @@
 #include "PanelSlot.generated.h"
 
 /** The base class for all Slots in UMG. */
-UCLASS()
+UCLASS(BlueprintType)
 class UMG_API UPanelSlot : public UVisual
 {
 	GENERATED_UCLASS_BODY()
