UMG - Improving hover queue to not apply to already selected widgets, broke when multiselection was added.  Now hiding the resolution/aspect ratio when custom size is active.

diff --git a/Engine/Source/Editor/UMGEditor/Private/Designer/SDesignerView.cpp b/Engine/Source/Editor/UMGEditor/Private/Designer/SDesignerView.cpp
--- a/Engine/Source/Editor/UMGEditor/Private/Designer/SDesignerView.cpp
+++ b/Engine/Source/Editor/UMGEditor/Private/Designer/SDesignerView.cpp
@@ -376,6 +376,7 @@ void SDesignerView::Construct(const FArguments& InArgs, TSharedPtr<FWidgetBluepr
 				.Padding(6, 0, 0, 2)
 				[
 					SNew(STextBlock)
+					.Visibility(this, &SDesignerView::GetResolutionTextVisibility)
 					.TextStyle(FEditorStyle::Get(), "Graph.ZoomText")
 					.Text(this, &SDesignerView::GetCurrentResolutionText)
 					.ColorAndOpacity(this, &SDesignerView::GetResolutionTextColorAndOpacity)
@@ -1097,35 +1098,12 @@ int32 SDesignerView::HandleEffectsPainting(const FOnPaintHandlerParams& PaintArg
 	static const FLinearColor SelectedTint(0, 1, 0);
 	const FLinearColor HoveredTint(0, 0.5, 1, FMath::Clamp(BlueprintEditor.Pin()->GetHoveredWidgetTime() / HoveredAnimationTime, 0.0f, 1.0f)); // Azure = 0x007FFF
 
+	const FSlateBrush* SelectionOutlineBrush = FEditorStyle::Get().GetBrush(SelectionOutlineName);
+	FVector2D SelectionBrushInflationAmount = FVector2D(16, 16) * FVector2D(SelectionOutlineBrush->Margin.Left, SelectionOutlineBrush->Margin.Top) * ( 1.0f / GetPreviewScale() );
+
 	for ( const FWidgetReference& SelectedWidget : SelectedWidgets )
 	{
-		const FSlateBrush* SelectionOutlineBrush = FEditorStyle::Get().GetBrush(SelectionOutlineName);
-		FVector2D SelectionBrushInflationAmount = FVector2D(16, 16) * FVector2D(SelectionOutlineBrush->Margin.Left, SelectionOutlineBrush->Margin.Top) * ( 1.0f / GetPreviewScale() );
-
 		TSharedPtr<SWidget> SelectedSlateWidget = SelectedWidget.GetPreviewSlate();
-		TSharedPtr<SWidget> HoveredSlateWidget = BlueprintEditor.Pin()->GetHoveredWidget().GetPreviewSlate();
-
-		// Don't draw the hovered effect if it's also the selected widget
-		if ( HoveredSlateWidget.IsValid() && HoveredSlateWidget != SelectedSlateWidget )
-		{
-			TSharedRef<SWidget> Widget = HoveredSlateWidget.ToSharedRef();
-
-			FArrangedWidget ArrangedWidget(SNullWidget::NullWidget, FGeometry());
-			FDesignTimeUtils::GetArrangedWidgetRelativeToWindow(Widget, ArrangedWidget);
-
-			// Draw hovered effect
-			FPaintGeometry HoveredGeometry = ArrangedWidget.Geometry.ToInflatedPaintGeometry(SelectionBrushInflationAmount);
-
-			FSlateDrawElement::MakeBox(
-				PaintArgs.OutDrawElements,
-				PaintArgs.Layer,
-				HoveredGeometry,
-				SelectionOutlineBrush,
-				PaintArgs.ClippingRect,
-				ESlateDrawEffect::None,
-				HoveredTint
-				);
-		}
 
 		if ( SelectedSlateWidget.IsValid() )
 		{
@@ -1149,6 +1127,31 @@ int32 SDesignerView::HandleEffectsPainting(const FOnPaintHandlerParams& PaintArg
 		}
 	}
 
+	FWidgetReference HoveredWidget = BlueprintEditor.Pin()->GetHoveredWidget();
+	TSharedPtr<SWidget> HoveredSlateWidget = HoveredWidget.GetPreviewSlate();
+
+	// Don't draw the hovered effect if it's also the selected widget
+	if ( HoveredSlateWidget.IsValid() && !SelectedWidgets.Contains(HoveredWidget) )
+	{
+		TSharedRef<SWidget> Widget = HoveredSlateWidget.ToSharedRef();
+
+		FArrangedWidget ArrangedWidget(SNullWidget::NullWidget, FGeometry());
+		FDesignTimeUtils::GetArrangedWidgetRelativeToWindow(Widget, ArrangedWidget);
+
+		// Draw hovered effect
+		FPaintGeometry HoveredGeometry = ArrangedWidget.Geometry.ToInflatedPaintGeometry(SelectionBrushInflationAmount);
+
+		FSlateDrawElement::MakeBox(
+			PaintArgs.OutDrawElements,
+			PaintArgs.Layer,
+			HoveredGeometry,
+			SelectionOutlineBrush,
+			PaintArgs.ClippingRect,
+			ESlateDrawEffect::None,
+			HoveredTint
+			);
+	}
+
 	return PaintArgs.Layer + 1;
 }
 
@@ -1679,6 +1682,20 @@ FSlateColor SDesignerView::GetResolutionTextColorAndOpacity() const
 	return FLinearColor(1, 1, 1, 1.25f - ResolutionTextFade.GetLerp());
 }
 
+EVisibility SDesignerView::GetResolutionTextVisibility() const
+{
+	// If we're using a custom design time size, don't bother showing the resolution
+	if ( UUserWidget* DefaultWidget = GetDefaultWidget() )
+	{
+		if ( DefaultWidget->bUseDesignTimeSize )
+		{
+			return EVisibility::Collapsed;
+		}
+	}
+
+	return EVisibility::SelfHitTestInvisible;
+}
+
 FReply SDesignerView::HandleDPISettingsClicked()
 {
 	FModuleManager::LoadModuleChecked<ISettingsModule>("Settings").ShowViewer("Project", "Engine", "UI");
diff --git a/Engine/Source/Editor/UMGEditor/Private/Designer/SDesignerView.h b/Engine/Source/Editor/UMGEditor/Private/Designer/SDesignerView.h
--- a/Engine/Source/Editor/UMGEditor/Private/Designer/SDesignerView.h
+++ b/Engine/Source/Editor/UMGEditor/Private/Designer/SDesignerView.h
@@ -120,6 +120,7 @@ private:
 	FText GetCurrentResolutionText() const;
 	FText GetCurrentDPIScaleText() const;
 	FSlateColor GetResolutionTextColorAndOpacity() const;
+	EVisibility GetResolutionTextVisibility() const;
 
 	TOptional<int32> GetCustomResolutionWidth() const;
 	TOptional<int32> GetCustomResolutionHeight() const;
