- Double-click user widget in Palette to open in editor
- Hierarchy context menu option to open user widget for edit
- User widget tooltips now use the widget BP description (from Blueprint Props); also shows up in CB tooltip now too
- Hierarchy has widget description tooltip when hovering over image; widget type tooltip on label if the user has provided a name

diff --git a/Engine/Source/Editor/UMGEditor/Private/AssetTypeActions_WidgetBlueprint.cpp b/Engine/Source/Editor/UMGEditor/Private/AssetTypeActions_WidgetBlueprint.cpp
--- a/Engine/Source/Editor/UMGEditor/Private/AssetTypeActions_WidgetBlueprint.cpp
+++ b/Engine/Source/Editor/UMGEditor/Private/AssetTypeActions_WidgetBlueprint.cpp
@@ -55,4 +55,18 @@ bool FAssetTypeActions_WidgetBlueprint::ShouldUseDataOnlyEditor(const UBlueprint
 		&& !Blueprint->bIsNewlyCreated;
 }
 
+FText FAssetTypeActions_WidgetBlueprint::GetAssetDescription( const FAssetData& AssetData ) const
+{
+	if ( const FString* pDescription = AssetData.TagsAndValues.Find( GET_MEMBER_NAME_CHECKED( UBlueprint, BlueprintDescription ) ) )
+	{
+		if ( !pDescription->IsEmpty() )
+		{
+			const FString DescriptionStr( *pDescription );
+			return FText::FromString( DescriptionStr.Replace( TEXT( "\\n" ), TEXT( "\n" ) ) );
+		}
+	}
+
+	return FText::GetEmpty();
+}
+
 #undef LOCTEXT_NAMESPACE
\ No newline at end of file
diff --git a/Engine/Source/Editor/UMGEditor/Private/AssetTypeActions_WidgetBlueprint.h b/Engine/Source/Editor/UMGEditor/Private/AssetTypeActions_WidgetBlueprint.h
--- a/Engine/Source/Editor/UMGEditor/Private/AssetTypeActions_WidgetBlueprint.h
+++ b/Engine/Source/Editor/UMGEditor/Private/AssetTypeActions_WidgetBlueprint.h
@@ -12,6 +12,7 @@ public:
 	virtual UClass* GetSupportedClass() const override;
 	virtual void OpenAssetEditor(const TArray<UObject*>& InObjects, TSharedPtr<class IToolkitHost> EditWithinLevelEditor = TSharedPtr<IToolkitHost>()) override;
 	virtual uint32 GetCategories() override { return EAssetTypeCategories::UI; }
+	virtual FText GetAssetDescription( const FAssetData& AssetData ) const override;
 	// End IAssetTypeActions Implementation
 
 private:
diff --git a/Engine/Source/Editor/UMGEditor/Private/Hierarchy/SHierarchyViewItem.cpp b/Engine/Source/Editor/UMGEditor/Private/Hierarchy/SHierarchyViewItem.cpp
--- a/Engine/Source/Editor/UMGEditor/Private/Hierarchy/SHierarchyViewItem.cpp
+++ b/Engine/Source/Editor/UMGEditor/Private/Hierarchy/SHierarchyViewItem.cpp
@@ -545,6 +545,39 @@ FText FHierarchyWidget::GetText() const
 	return FText::GetEmpty();
 }
 
+FText FHierarchyWidget::GetImageToolTipText() const
+{
+	UWidget* WidgetTemplate = Item.GetTemplate();
+	if ( WidgetTemplate )
+	{
+		UClass* WidgetClass = WidgetTemplate->GetClass();
+		if ( WidgetClass->IsChildOf( UUserWidget::StaticClass() ) )
+		{
+			auto& Description = Cast<UWidgetBlueprint>( WidgetClass->ClassGeneratedBy )->BlueprintDescription;
+			if ( Description.Len() > 0 )
+			{
+				return FText::FromString( Description );
+			}
+		}
+		
+		return WidgetClass->GetToolTipText();
+	}
+	
+	return FText::GetEmpty();
+}
+
+FText FHierarchyWidget::GetLabelToolTipText() const
+{
+	// If the user has provided a name, give a tooltip with the widget type for easy reference
+	UWidget* WidgetTemplate = Item.GetTemplate();
+	if ( WidgetTemplate && !WidgetTemplate->IsGeneratedName() )
+	{
+		return FText::FromString(TEXT( "[" ) + WidgetTemplate->GetClass()->GetName() + TEXT( "]" ) );
+	}
+
+	return FText::GetEmpty();
+}
+
 const FSlateBrush* FHierarchyWidget::GetImage() const
 {
 	return Item.GetTemplate()->GetEditorIcon();
@@ -726,7 +759,7 @@ void SHierarchyViewItem::Construct(const FArguments& InArgs, const TSharedRef< S
 		.Content()
 		[
 			SNew(SHorizontalBox)
-
+			
 			// Widget icon
 			+ SHorizontalBox::Slot()
 			.AutoWidth()
@@ -734,6 +767,7 @@ void SHierarchyViewItem::Construct(const FArguments& InArgs, const TSharedRef< S
 				SNew(SImage)
 				.ColorAndOpacity(FLinearColor(1,1,1,0.5))
 				.Image(Model->GetImage())
+				.ToolTipText(Model->GetImageToolTipText())
 			]
 
 			// Name of the widget
@@ -745,6 +779,7 @@ void SHierarchyViewItem::Construct(const FArguments& InArgs, const TSharedRef< S
 				SNew(SInlineEditableTextBlock)
 				.Font(this, &SHierarchyViewItem::GetItemFont)
 				.Text(this, &SHierarchyViewItem::GetItemText)
+				.ToolTipText(Model->GetLabelToolTipText())
 				.HighlightText(InArgs._HighlightText)
 				.OnVerifyTextChanged(this, &SHierarchyViewItem::OnVerifyNameTextChanged)
 				.OnTextCommitted(this, &SHierarchyViewItem::OnNameTextCommited)
diff --git a/Engine/Source/Editor/UMGEditor/Private/Hierarchy/SHierarchyViewItem.h b/Engine/Source/Editor/UMGEditor/Private/Hierarchy/SHierarchyViewItem.h
--- a/Engine/Source/Editor/UMGEditor/Private/Hierarchy/SHierarchyViewItem.h
+++ b/Engine/Source/Editor/UMGEditor/Private/Hierarchy/SHierarchyViewItem.h
@@ -22,6 +22,12 @@ public:
 	/* @returns the widget name to use for the tree item */
 	virtual FText GetText() const = 0;
 
+	/** @return The tooltip for the tree item image */
+	virtual FText GetImageToolTipText() const { return FText::GetEmpty(); }
+
+	/** @return The tooltip for the tree item label */
+	virtual FText GetLabelToolTipText() const { return FText::GetEmpty(); }
+
 	virtual const FSlateBrush* GetImage() const = 0;
 
 	virtual FSlateFontInfo GetFont() const = 0;
@@ -43,7 +49,7 @@ public:
 
 	virtual void OnMouseEnter() {}
 	virtual void OnMouseLeave() {}
-	
+
 	void RefreshSelection();
 	bool ContainsSelection();
 	bool IsSelected() const;
@@ -142,6 +148,8 @@ public:
 	
 	/* @returns the widget name to use for the tree item */
 	virtual FText GetText() const override;
+	virtual FText GetImageToolTipText() const override;
+	virtual FText GetLabelToolTipText() const override;
 
 	virtual const FSlateBrush* GetImage() const override;
 
diff --git a/Engine/Source/Editor/UMGEditor/Private/Palette/SPaletteView.cpp b/Engine/Source/Editor/UMGEditor/Private/Palette/SPaletteView.cpp
--- a/Engine/Source/Editor/UMGEditor/Private/Palette/SPaletteView.cpp
+++ b/Engine/Source/Editor/UMGEditor/Private/Palette/SPaletteView.cpp
@@ -77,6 +77,11 @@ public:
 		];
 	}
 
+	virtual FReply OnMouseButtonDoubleClick( const FGeometry& InMyGeometry, const FPointerEvent& InMouseEvent ) override
+	{
+		return Template->OnDoubleClicked();
+	}
+
 private:
 	TSharedPtr<FWidgetTemplate> Template;
 };
diff --git a/Engine/Source/Editor/UMGEditor/Private/Palette/SPaletteView.h b/Engine/Source/Editor/UMGEditor/Private/Palette/SPaletteView.h
--- a/Engine/Source/Editor/UMGEditor/Private/Palette/SPaletteView.h
+++ b/Engine/Source/Editor/UMGEditor/Private/Palette/SPaletteView.h
@@ -6,8 +6,7 @@
 #include "BlueprintEditor.h"
 #include "TreeFilterHandler.h"
 
-#include "WidgetTemplate.h"
-
+class FWidgetTemplate;
 class UWidgetBlueprint;
 
 /** View model for the items in the widget template list */
diff --git a/Engine/Source/Editor/UMGEditor/Private/Templates/WidgetTemplateBlueprintClass.cpp b/Engine/Source/Editor/UMGEditor/Private/Templates/WidgetTemplateBlueprintClass.cpp
--- a/Engine/Source/Editor/UMGEditor/Private/Templates/WidgetTemplateBlueprintClass.cpp
+++ b/Engine/Source/Editor/UMGEditor/Private/Templates/WidgetTemplateBlueprintClass.cpp
@@ -35,7 +35,7 @@ FText FWidgetTemplateBlueprintClass::GetCategory() const
 	}
 	else
 	{
-		//TODO UMG if the blueprint is unloaded we need to extract it from the asset metadata.
+		//If the blueprint is unloaded we need to extract it from the asset metadata.
 		auto DefaultUserWidget = UUserWidget::StaticClass()->GetDefaultObject<UUserWidget>();
 		return DefaultUserWidget->GetPaletteCategory();
 	}
@@ -62,12 +62,26 @@ const FSlateBrush* FWidgetTemplateBlueprintClass::GetIcon() const
 
 TSharedRef<IToolTip> FWidgetTemplateBlueprintClass::GetToolTip() const
 {
-	if (WidgetClass.Get())
+	FText Description;
+	if ( const FString* DescriptionStr = WidgetAssetData.TagsAndValues.Find( GET_MEMBER_NAME_CHECKED( UBlueprint, BlueprintDescription ) ) )
 	{
-		return FWidgetTemplateClass::GetToolTip();
+		if ( !DescriptionStr->IsEmpty() )
+		{
+			Description = FText::FromString( DescriptionStr->Replace( TEXT( "\\n" ), TEXT( "\n" ) ) );
+		}
+	}
+	else
+	{
+		Description = Name;
 	}
 
-	return IDocumentation::Get()->CreateToolTip(Name, nullptr, FString(TEXT("Shared/Types/")) + Name.ToString(), TEXT("Class"));
+	return IDocumentation::Get()->CreateToolTip( Description, nullptr, FString( TEXT( "Shared/Types/" ) ) + Name.ToString(), TEXT( "Class" ) );
+}
+
+FReply FWidgetTemplateBlueprintClass::OnDoubleClicked()
+{
+	FAssetEditorManager::Get().OpenEditorForAsset( WidgetAssetData.GetAsset() );
+	return FReply::Handled();
 }
 
 #undef LOCTEXT_NAMESPACE
diff --git a/Engine/Source/Editor/UMGEditor/Private/Templates/WidgetTemplateBlueprintClass.h b/Engine/Source/Editor/UMGEditor/Private/Templates/WidgetTemplateBlueprintClass.h
--- a/Engine/Source/Editor/UMGEditor/Private/Templates/WidgetTemplateBlueprintClass.h
+++ b/Engine/Source/Editor/UMGEditor/Private/Templates/WidgetTemplateBlueprintClass.h
@@ -35,6 +35,9 @@ public:
 	/** Gets the tooltip widget for this palette item. */
 	virtual TSharedRef<IToolTip> GetToolTip() const override;
 
+	/** Opens the widget blueprint for edit */
+	virtual FReply OnDoubleClicked() override;
+
 protected:
 	/** The asset data for the widget blueprint */
 	FAssetData WidgetAssetData;
diff --git a/Engine/Source/Editor/UMGEditor/Private/WidgetBlueprintEditorUtils.cpp b/Engine/Source/Editor/UMGEditor/Private/WidgetBlueprintEditorUtils.cpp
--- a/Engine/Source/Editor/UMGEditor/Private/WidgetBlueprintEditorUtils.cpp
+++ b/Engine/Source/Editor/UMGEditor/Private/WidgetBlueprintEditorUtils.cpp
@@ -199,6 +199,18 @@ void FWidgetBlueprintEditorUtils::CreateWidgetContextMenu(FMenuBuilder& MenuBuil
 
 	MenuBuilder.BeginSection("Actions");
 	{
+		MenuBuilder.AddMenuEntry(
+			LOCTEXT( "EditBlueprint_Label", "Edit Blueprint..." ),
+			LOCTEXT( "EditBlueprint_Tooltip", "Open the selected widget blueprint(s) for edit." ),
+			FSlateIcon(),
+			FUIAction(
+				FExecuteAction::CreateStatic( &FWidgetBlueprintEditorUtils::ExecuteOpenSelectedWidgetsForEdit, Widgets ),
+				FCanExecuteAction(),
+				FIsActionChecked(),
+				FIsActionButtonVisible::CreateStatic( &FWidgetBlueprintEditorUtils::CanOpenSelectedWidgetsForEdit, Widgets )
+				)
+			);
+
 		MenuBuilder.AddSubMenu(
 			LOCTEXT("WidgetTree_WrapWith", "Wrap With..."),
 			LOCTEXT("WidgetTree_WrapWithToolTip", "Wraps the currently selected widgets inside of another container widget"),
@@ -210,6 +222,30 @@ void FWidgetBlueprintEditorUtils::CreateWidgetContextMenu(FMenuBuilder& MenuBuil
 	MenuBuilder.PopCommandList();
 }
 
+void FWidgetBlueprintEditorUtils::ExecuteOpenSelectedWidgetsForEdit( TSet<FWidgetReference> SelectedWidgets )
+{
+	for ( auto& Widget : SelectedWidgets )
+	{
+		FAssetEditorManager::Get().OpenEditorForAsset( Widget.GetTemplate()->GetClass()->ClassGeneratedBy );
+	}
+}
+
+bool FWidgetBlueprintEditorUtils::CanOpenSelectedWidgetsForEdit( TSet<FWidgetReference> SelectedWidgets )
+{
+	bool bCanOpenAllForEdit = SelectedWidgets.Num() > 0;
+	for ( auto& Widget : SelectedWidgets )
+	{
+		auto Blueprint = Widget.GetTemplate()->GetClass()->ClassGeneratedBy;
+		if ( !Blueprint || !Blueprint->IsA( UWidgetBlueprint::StaticClass() ) )
+		{
+			bCanOpenAllForEdit = false;
+			break;
+		}
+	}
+
+	return bCanOpenAllForEdit;
+}
+
 void FWidgetBlueprintEditorUtils::DeleteWidgets(UWidgetBlueprint* BP, TSet<FWidgetReference> Widgets)
 {
 	if ( Widgets.Num() > 0 )
diff --git a/Engine/Source/Editor/UMGEditor/Private/WidgetBlueprintEditorUtils.h b/Engine/Source/Editor/UMGEditor/Private/WidgetBlueprintEditorUtils.h
--- a/Engine/Source/Editor/UMGEditor/Private/WidgetBlueprintEditorUtils.h
+++ b/Engine/Source/Editor/UMGEditor/Private/WidgetBlueprintEditorUtils.h
@@ -36,6 +36,11 @@ public:
 	static void ImportPropertiesFromText(UObject* Object, const TMap<FName, FString>& ExportedProperties);
 
 private:
+
+	static void ExecuteOpenSelectedWidgetsForEdit( TSet<FWidgetReference> SelectedWidgets );
+
+	static bool CanOpenSelectedWidgetsForEdit( TSet<FWidgetReference> SelectedWidgets );
+
 	static bool FindAndRemoveNamedSlotContent(UWidget* WidgetTemplate, UWidgetTree* WidgetTree);
 
 	static void BuildWrapWithMenu(FMenuBuilder& Menu, UWidgetBlueprint* BP, TSet<FWidgetReference> Widgets);
diff --git a/Engine/Source/Editor/UMGEditor/Public/WidgetTemplate.h b/Engine/Source/Editor/UMGEditor/Public/WidgetTemplate.h
--- a/Engine/Source/Editor/UMGEditor/Public/WidgetTemplate.h
+++ b/Engine/Source/Editor/UMGEditor/Public/WidgetTemplate.h
@@ -30,6 +30,9 @@ public:
 	/** Gets tooltip widget for this palette item. */
 	virtual TSharedRef<IToolTip> GetToolTip() const = 0;
 
+	/** The the action to perform when the template item is double clicked */
+	virtual FReply OnDoubleClicked() { return FReply::Unhandled(); }
+
 public:
 	/** The name of the widget template. */
 	FText Name;
