UMGEditor PCH cleanup

diff --git a/Engine/Source/Editor/UMGEditor/Private/BlueprintModes/WidgetBlueprintApplicationMode.cpp b/Engine/Source/Editor/UMGEditor/Private/BlueprintModes/WidgetBlueprintApplicationMode.cpp
--- a/Engine/Source/Editor/UMGEditor/Private/BlueprintModes/WidgetBlueprintApplicationMode.cpp
+++ b/Engine/Source/Editor/UMGEditor/Private/BlueprintModes/WidgetBlueprintApplicationMode.cpp
@@ -9,6 +9,7 @@
 
 #include "WidgetBlueprintEditor.h"
 #include "WidgetBlueprintApplicationMode.h"
+#include "WidgetBlueprintApplicationModes.h"
 
 /////////////////////////////////////////////////////
 // FWidgetBlueprintApplicationMode
diff --git a/Engine/Source/Editor/UMGEditor/Private/Details/SPropertyBinding.cpp b/Engine/Source/Editor/UMGEditor/Private/Details/SPropertyBinding.cpp
--- a/Engine/Source/Editor/UMGEditor/Private/Details/SPropertyBinding.cpp
+++ b/Engine/Source/Editor/UMGEditor/Private/Details/SPropertyBinding.cpp
@@ -4,6 +4,10 @@
 #include "DetailLayoutBuilder.h"
 #include "PropertyHandle.h"
 #include "SPropertyBinding.h"
+#include "WidgetBlueprintApplicationModes.h"
+#include "WidgetGraphSchema.h"
+#include "BlueprintEditorUtils.h"
+#include "ScopedTransaction.h"
 
 
 #define LOCTEXT_NAMESPACE "UMG"
diff --git a/Engine/Source/Editor/UMGEditor/Private/Details/SPropertyBinding.h b/Engine/Source/Editor/UMGEditor/Private/Details/SPropertyBinding.h
--- a/Engine/Source/Editor/UMGEditor/Private/Details/SPropertyBinding.h
+++ b/Engine/Source/Editor/UMGEditor/Private/Details/SPropertyBinding.h
@@ -1,6 +1,7 @@
 // Copyright 1998-2014 Epic Games, Inc. All Rights Reserved.
 
 #pragma once
+#include "WidgetBlueprintEditor.h"
 
 #include "SCompoundWidget.h"
 
diff --git a/Engine/Source/Editor/UMGEditor/Private/UMGEditorPrivatePCH.h b/Engine/Source/Editor/UMGEditor/Private/UMGEditorPrivatePCH.h
--- a/Engine/Source/Editor/UMGEditor/Private/UMGEditorPrivatePCH.h
+++ b/Engine/Source/Editor/UMGEditor/Private/UMGEditorPrivatePCH.h
@@ -3,52 +3,40 @@
 #pragma once
 #include "SlateBasics.h"
 #include "SlateCore.h"
-#include "UnrealEd.h"
 
-#include "IDocumentation.h"
+//#include "UnrealEd.h" // Most used stuff from UnrealEd
+#if WITH_EDITOR
+#include "EngineDefines.h"
+#include "TickableEditorObject.h"
+#include "Editor.h"
+#include "EditorViewportClient.h"
+#include "LevelEditorViewport.h"
+#include "ParticleHelper.h"
+#include "Viewports.h"
+#endif // WITH_EDITOR
+//#include "UnrealEd.h" // End
 
+#include "IDocumentation.h"
 #include "BlueprintGraphDefinitions.h"
-
-#include "BlueprintEditorUtils.h"
-
-#include "ScopedTransaction.h"
-#include "TextFilter.h"
-
 #include "ObjectEditorUtils.h"
 
-#include "UMG.h"
-#include "UMGStyle.h"
+//#include "UMG.h" // Most used stuff from UMG
+#include "Anchors.h"
+#include "SConstraintCanvas.h"
+#include "SlateWrapperTypes.h"
+#include "PanelSlot.h"
+#include "WidgetNavigation.h"
+#include "Widget.h"
+#include "PanelWidget.h"
+#include "NamedSlotInterface.h"
+#include "CanvasPanelSlot.h"
+#include "Viewport.h"
+#include "WidgetTree.h"
+#include "UserWidget.h"
+#include "WidgetBlueprintGeneratedClass.h"
+//#include "UMG.h" // End
 
 #include "WidgetDesignerSettings.h"
 
-#include "UMGEditorClasses.h"
 #include "WidgetBlueprint.h"
-
 #include "BlueprintEditorModes.h"
-#include "WidgetBlueprintApplicationModes.h"
-#include "WidgetBlueprintApplicationMode.h"
-
-#include "WidgetReference.h"
-#include "WidgetBlueprintEditor.h"
-
-#include "AssetTypeActions_Base.h"
-
-#include "DesignTimeUtils.h"
-
-#include "DesignerExtension.h"
-#include "VerticalSlotExtension.h"
-#include "HorizontalSlotExtension.h"
-#include "CanvasSlotExtension.h"
-#include "UniformGridSlotExtension.h"
-#include "GridSlotExtension.h"
-
-#include "DetailCustomizations.h"
-
-#include "WidgetTemplate.h"
-#include "WidgetTemplateDragDropOp.h"
-
-#include "WidgetBlueprintEditorUtils.h"
-
-#include "WidgetBlueprintLibrary.h"
-
-#include "WidgetBlueprintCompiler.h"
diff --git a/Engine/Source/Editor/UMGEditor/Private/WidgetBlueprint.cpp b/Engine/Source/Editor/UMGEditor/Private/WidgetBlueprint.cpp
--- a/Engine/Source/Editor/UMGEditor/Private/WidgetBlueprint.cpp
+++ b/Engine/Source/Editor/UMGEditor/Private/WidgetBlueprint.cpp
@@ -6,6 +6,7 @@
 #include "PropertyTag.h"
 #include "Blueprint/WidgetBlueprintGeneratedClass.h"
 #include "WidgetBlueprint.h"
+#include "WidgetBlueprintCompiler.h"
 
 #define LOCTEXT_NAMESPACE "UMG"
 
diff --git a/Engine/Source/Editor/UMGEditor/Private/WidgetBlueprintCompiler.cpp b/Engine/Source/Editor/UMGEditor/Private/WidgetBlueprintCompiler.cpp
--- a/Engine/Source/Editor/UMGEditor/Private/WidgetBlueprintCompiler.cpp
+++ b/Engine/Source/Editor/UMGEditor/Private/WidgetBlueprintCompiler.cpp
@@ -11,6 +11,7 @@
 #include "Blueprint/WidgetTree.h"
 #include "WidgetBlueprint.h"
 #include "Kismet2/BlueprintEditorUtils.h"
+#include "NamedSlot.h"
 
 FOnWidgetBlueprintCompiled FWidgetBlueprintCompiler::OnWidgetBlueprintCompiled = FOnWidgetBlueprintCompiled();
 
diff --git a/Engine/Source/Runtime/UMG/Public/Components/Viewport.h b/Engine/Source/Runtime/UMG/Public/Components/Viewport.h
--- a/Engine/Source/Runtime/UMG/Public/Components/Viewport.h
+++ b/Engine/Source/Runtime/UMG/Public/Components/Viewport.h
@@ -3,7 +3,7 @@
 #pragma once
 
 #include "Camera/CameraTypes.h"
-
+#include "ContentWidget.h"
 #include "Viewport.generated.h"
 
 class FPreviewScene;
