Refactor to unify keyboard and controller input and focus.

diff --git a/Engine/Source/Editor/LevelEditor/Private/SLevelEditor.h b/Engine/Source/Editor/LevelEditor/Private/SLevelEditor.h
--- a/Engine/Source/Editor/LevelEditor/Private/SLevelEditor.h
+++ b/Engine/Source/Editor/LevelEditor/Private/SLevelEditor.h
@@ -95,9 +95,9 @@ public:
 	 * Processes level editor keybindings using events made in a viewport
 	 * 
 	 * @param MyGeometry		Information about the size and position of the viewport widget
-	 * @param InKeyboardEvent	The keyboard event which just occurred	
+	 * @param InKeyEvent	The event which just occurred	
 	 */
-	virtual FReply OnKeyDownInViewport( const FGeometry& MyGeometry, const FKeyboardEvent& InKeyboardEvent );
+	virtual FReply OnKeyDownInViewport( const FGeometry& MyGeometry, const FKeyEvent& InKeyEvent );
 
 	bool CanCloseApp();
 
@@ -161,9 +161,9 @@ private:
 	 * Processes keybindings on the level editor
 	 * 
 	 * @param MyGeometry		Information about the size and position of the level editor widget
-	 * @param InKeyboardEvent	The keyboard event which just occurred	
+	 * @param InKeyEvent	The event which just occurred	
 	 */
-	virtual FReply OnKeyDown( const FGeometry& MyGeometry, const FKeyboardEvent& InKeyboardEvent );
+	virtual FReply OnKeyDown( const FGeometry& MyGeometry, const FKeyEvent& InKeyEvent );
 
 	/** Callback for when the property view changes */
 	void OnPropertyObjectArrayChanged(const FString& NewTitle, const TArray< UObject* >& UObjects );
diff --git a/Engine/Source/Editor/LevelEditor/Private/SLevelViewport.cpp b/Engine/Source/Editor/LevelEditor/Private/SLevelViewport.cpp
--- a/Engine/Source/Editor/LevelEditor/Private/SLevelViewport.cpp
+++ b/Engine/Source/Editor/LevelEditor/Private/SLevelViewport.cpp
@@ -381,14 +381,14 @@ float SLevelViewport::GetGameViewportDPIScale() const
 	return 1.0f;
 }
 
-FReply SLevelViewport::OnKeyDown( const FGeometry& MyGeometry, const FKeyboardEvent& InKeyboardEvent )
+FReply SLevelViewport::OnKeyDown( const FGeometry& MyGeometry, const FKeyEvent& InKeyEvent )
 {
 	FReply Reply = FReply::Unhandled();
 
 	if( HasPlayInEditorViewport() || LevelViewportClient->IsSimulateInEditorViewport() )
 	{
 		// Only process commands for pie when a play world is active
-		FPlayWorldCommands::GlobalPlayWorldActions->ProcessCommandBindings( InKeyboardEvent );
+		FPlayWorldCommands::GlobalPlayWorldActions->ProcessCommandBindings( InKeyEvent );
 
 		// Always handle commands in pie so they arent bubbled to editor only widgets
 		Reply = FReply::Handled();
@@ -396,7 +396,7 @@ FReply SLevelViewport::OnKeyDown( const FGeometry& MyGeometry, const FKeyboardEv
 	
 	if( !IsPlayInEditorViewportActive() )
 	{
-		Reply = SEditorViewport::OnKeyDown(MyGeometry,InKeyboardEvent);
+		Reply = SEditorViewport::OnKeyDown(MyGeometry,InKeyEvent);
 
 
 		// If we are in immersive mode and the event was not handled, we will check to see if the the 
@@ -408,7 +408,7 @@ FReply SLevelViewport::OnKeyDown( const FGeometry& MyGeometry, const FKeyboardEv
 			TSharedPtr<SLevelEditor> ParentLevelEditorSharedPtr = ParentLevelEditor.Pin();
 			if( ParentLevelEditorSharedPtr.IsValid() )
 			{
-				Reply = ParentLevelEditorSharedPtr->OnKeyDownInViewport( MyGeometry, InKeyboardEvent );
+				Reply = ParentLevelEditorSharedPtr->OnKeyDownInViewport( MyGeometry, InKeyEvent );
 			}
 		}
 	}
@@ -3224,7 +3224,7 @@ void SLevelViewport::StartPlayInEditorSession( UGameViewportClient* PlayClient,
 
 	// Remove keyboard focus to send a focus lost message to the widget to clean up any saved state from the viewport interface thats about to be swapped out
 	// Focus will be set when the game viewport is registered
-	FSlateApplication::Get().ClearKeyboardFocus(EKeyboardFocusCause::SetDirectly);
+	FSlateApplication::Get().ClearKeyboardFocus(EFocusCause::SetDirectly);
 
 	ActiveViewport = MakeShareable( new FSceneViewport( PlayClient, ViewportWidget ) );
 	ActiveViewport->SetPlayInEditorViewport( true );
diff --git a/Engine/Source/Editor/Levels/Private/SLevelBrowser.h b/Engine/Source/Editor/Levels/Private/SLevelBrowser.h
--- a/Engine/Source/Editor/Levels/Private/SLevelBrowser.h
+++ b/Engine/Source/Editor/Levels/Private/SLevelBrowser.h
@@ -64,7 +64,7 @@ protected:
 	}
 
 	/** Overridden from SWidget: Called when a key is pressed down */
-	FReply OnKeyDown( const FGeometry& MyGeometry, const FKeyboardEvent& InKeyboardEvent )
+	FReply OnKeyDown( const FGeometry& MyGeometry, const FKeyEvent& InKeyEvent )
 	{
 		return FReply::Unhandled();
 	}
diff --git a/Engine/Source/Editor/Levels/Private/SLevelsView.h b/Engine/Source/Editor/Levels/Private/SLevelsView.h
--- a/Engine/Source/Editor/Levels/Private/SLevelsView.h
+++ b/Engine/Source/Editor/Levels/Private/SLevelsView.h
@@ -198,16 +198,16 @@ protected:
 	}
 
 	/**
-	 * Called after a key is pressed when this widget has keyboard focus
+	 * Called after a key is pressed when this widget has focus
 	 *
 	 * @param MyGeometry The Geometry of the widget receiving the event
-	 * @param  InKeyboardEvent  Keyboard event
+	 * @param  InKeyEvent  Key event
 	 *
 	 * @return  Returns whether the event was handled, along with other possible actions
 	 */
-	virtual FReply OnKeyDown( const FGeometry& MyGeometry, const FKeyboardEvent& InKeyboardEvent ) override
+	virtual FReply OnKeyDown( const FGeometry& MyGeometry, const FKeyEvent& InKeyEvent ) override
 	{
-		return ViewModel->GetCommandList()->ProcessCommandBindings( InKeyboardEvent ) ? FReply::Handled() : FReply::Unhandled();
+		return ViewModel->GetCommandList()->ProcessCommandBindings( InKeyEvent ) ? FReply::Handled() : FReply::Unhandled();
 	}
 
 	/**
diff --git a/Engine/Source/Editor/MainFrame/Private/Frame/MainFrameActions.cpp b/Engine/Source/Editor/MainFrame/Private/Frame/MainFrameActions.cpp
--- a/Engine/Source/Editor/MainFrame/Private/Frame/MainFrameActions.cpp
+++ b/Engine/Source/Editor/MainFrame/Private/Frame/MainFrameActions.cpp
@@ -36,7 +36,7 @@ void FMainFrameCommands::RegisterCommands()
 	if ( !IsRunningCommandlet() )
 	{
 		// The global action list was created at static initialization time. Create a handler for otherwise unhandled keyboard input to route key commands through this list.
-		FSlateApplication::Get().SetUnhandledKeyDownEventHandler( FOnKeyboardEvent::CreateStatic( &FMainFrameActionCallbacks::OnUnhandledKeyDownEvent ) );
+		FSlateApplication::Get().SetUnhandledKeyDownEventHandler( FOnKeyEvent::CreateStatic( &FMainFrameActionCallbacks::OnUnhandledKeyDownEvent ) );
 	}
 
 	// Make a default can execute action that disables input when in debug mode
@@ -157,13 +157,13 @@ void FMainFrameCommands::RegisterCommands()
 	FGlobalEditorCommonCommands::MapActions(ActionList);
 }
 
-FReply FMainFrameActionCallbacks::OnUnhandledKeyDownEvent(const FKeyboardEvent& InKeyboardEvent)
+FReply FMainFrameActionCallbacks::OnUnhandledKeyDownEvent(const FKeyEvent& InKeyEvent)
 {
-	if ( FMainFrameCommands::ActionList->ProcessCommandBindings( InKeyboardEvent ) )
+	if ( FMainFrameCommands::ActionList->ProcessCommandBindings( InKeyEvent ) )
 	{
 		return FReply::Handled();
 	}
-	else if( GEditor && GEditor->PlayWorld && InKeyboardEvent.GetKey() == EKeys::Escape )
+	else if( GEditor && GEditor->PlayWorld && InKeyEvent.GetKey() == EKeys::Escape )
 	{
 		FLevelEditorModule& LevelEditor = FModuleManager::GetModuleChecked< FLevelEditorModule >(TEXT("LevelEditor"));
 		if( LevelEditor.GetLevelEditorTab()->IsForeground()  )
diff --git a/Engine/Source/Editor/MainFrame/Private/Frame/MainFrameActions.h b/Engine/Source/Editor/MainFrame/Private/Frame/MainFrameActions.h
--- a/Engine/Source/Editor/MainFrame/Private/Frame/MainFrameActions.h
+++ b/Engine/Source/Editor/MainFrame/Private/Frame/MainFrameActions.h
@@ -77,7 +77,7 @@ class FMainFrameActionCallbacks
 
 public:
 	/** Global handler for unhandled key-down events in the editor. */
-	static FReply OnUnhandledKeyDownEvent(const FKeyboardEvent& InKeyboardEvent);
+	static FReply OnUnhandledKeyDownEvent(const FKeyEvent& InKeyEvent);
 
 	/**
 	 * The default can execute action for all commands unless they override it
diff --git a/Engine/Source/Editor/MaterialEditor/Private/FindInMaterial.cpp b/Engine/Source/Editor/MaterialEditor/Private/FindInMaterial.cpp
--- a/Engine/Source/Editor/MaterialEditor/Private/FindInMaterial.cpp
+++ b/Engine/Source/Editor/MaterialEditor/Private/FindInMaterial.cpp
@@ -177,7 +177,7 @@ void SFindInMaterial::FocusForUse()
 	FSlateApplication::Get().GeneratePathToWidgetUnchecked(SearchTextField.ToSharedRef(), FilterTextBoxWidgetPath);
 
 	// Set keyboard focus directly
-	FSlateApplication::Get().SetKeyboardFocus(FilterTextBoxWidgetPath, EKeyboardFocusCause::SetDirectly);
+	FSlateApplication::Get().SetKeyboardFocus(FilterTextBoxWidgetPath, EFocusCause::SetDirectly);
 }
 
 void SFindInMaterial::OnSearchTextChanged(const FText& Text)
diff --git a/Engine/Source/Editor/Matinee/Private/MatineeViewportClient.cpp b/Engine/Source/Editor/Matinee/Private/MatineeViewportClient.cpp
--- a/Engine/Source/Editor/Matinee/Private/MatineeViewportClient.cpp
+++ b/Engine/Source/Editor/Matinee/Private/MatineeViewportClient.cpp
@@ -832,7 +832,7 @@ bool FMatineeViewportClient::InputKey(FViewport* Viewport, int32 ControllerId, F
 		}
 
 		FModifierKeysState ModKeys(bShiftDown, bShiftDown, bCtrlDown, bCtrlDown, bAltDown, bAltDown, bCmdDown, bCmdDown, bCapsDown);
-		FKeyboardEvent KeyEvent(Key, ModKeys, false, 0, 0);
+		FKeyEvent KeyEvent(Key, ModKeys, 0/*UserIndex*/, false, 0, 0);
 		InterpEd->ProcessCommandBindings(KeyEvent);
 	}
 
diff --git a/Engine/Source/Editor/Matinee/Private/SMatineeRecorder.cpp b/Engine/Source/Editor/Matinee/Private/SMatineeRecorder.cpp
--- a/Engine/Source/Editor/Matinee/Private/SMatineeRecorder.cpp
+++ b/Engine/Source/Editor/Matinee/Private/SMatineeRecorder.cpp
@@ -129,7 +129,7 @@ void SMatineeRecorder::Construct(const FArguments& InArgs)
 
 	LevelViewportClient->SetRealtime( true );
 	LevelViewportClient->SetAllowMatineePreview( true );
-	LevelViewportClient->Viewport->CaptureJoystickInput(true);
+	LevelViewportClient->Viewport->SetUserFocus(true);
 	LevelViewportClient->SetMatineeRecordingWindow( ParentMatineeWindow.Pin().Get() );
 
 	LevelViewportClient->VisibilityDelegate.BindSP( this, &SMatineeRecorder::IsVisible );
diff --git a/Engine/Source/Editor/NewLevelDialog/Private/NewLevelDialogModule.cpp b/Engine/Source/Editor/NewLevelDialog/Private/NewLevelDialogModule.cpp
--- a/Engine/Source/Editor/NewLevelDialog/Private/NewLevelDialogModule.cpp
+++ b/Engine/Source/Editor/NewLevelDialog/Private/NewLevelDialogModule.cpp
@@ -259,14 +259,14 @@ private:
 		}
 	}
 
-	virtual FReply OnKeyDown( const FGeometry& MyGeometry, const FKeyboardEvent& InKeyboardEvent ) override
+	virtual FReply OnKeyDown( const FGeometry& MyGeometry, const FKeyEvent& InKeyEvent ) override
 	{
-		if( InKeyboardEvent.GetKey() == EKeys::Escape )
+		if( InKeyEvent.GetKey() == EKeys::Escape )
 		{
 			return OnCancelClicked();
 		}
 
-		return SCompoundWidget::OnKeyDown( MyGeometry, InKeyboardEvent );
+		return SCompoundWidget::OnKeyDown( MyGeometry, InKeyEvent );
 	}
 
 	/** Pointer to the parent window, so we know to destroy it when done */
diff --git a/Engine/Source/Editor/PListEditor/Private/SPlistEditor.cpp b/Engine/Source/Editor/PListEditor/Private/SPlistEditor.cpp
--- a/Engine/Source/Editor/PListEditor/Private/SPlistEditor.cpp
+++ b/Engine/Source/Editor/PListEditor/Private/SPlistEditor.cpp
@@ -2016,11 +2016,11 @@ bool SPListEditorPanel::DetermineAddArrayContext() const
 }
 
 /** Callback for keyboard shortcut commands */
-FReply SPListEditorPanel::OnKeyDown( const FGeometry& /*MyGeometry*/, const FKeyboardEvent& InKeyboardEvent )
+FReply SPListEditorPanel::OnKeyDown( const FGeometry& /*MyGeometry*/, const FKeyEvent& InKeyEvent )
 {
 	// Perform commands if necessary
 	FReply Reply = FReply::Unhandled();
-	if( UICommandList->ProcessCommandBindings( InKeyboardEvent ) )
+	if( UICommandList->ProcessCommandBindings( InKeyEvent ) )
 	{
 		// handle the event if a command was processed
 		Reply = FReply::Handled();
diff --git a/Engine/Source/Editor/PackagesDialog/Private/SPackagesDialog.cpp b/Engine/Source/Editor/PackagesDialog/Private/SPackagesDialog.cpp
--- a/Engine/Source/Editor/PackagesDialog/Private/SPackagesDialog.cpp
+++ b/Engine/Source/Editor/PackagesDialog/Private/SPackagesDialog.cpp
@@ -511,9 +511,9 @@ void SPackagesDialog::OnToggleSelectedCheckBox(ESlateCheckBoxState::Type InNewSt
 	ItemListView->RequestListRefresh();
 }
 
-FReply SPackagesDialog::OnKeyDown( const FGeometry& MyGeometry, const FKeyboardEvent& InKeyboardEvent )
+FReply SPackagesDialog::OnKeyDown( const FGeometry& MyGeometry, const FKeyEvent& InKeyEvent )
 {
-	if( InKeyboardEvent.GetKey() == EKeys::Escape )
+	if( InKeyEvent.GetKey() == EKeys::Escape )
 	{
 		for( int32 ButtonIndex = 0; ButtonIndex < Buttons.Num(); ++ButtonIndex )
 		{
@@ -525,7 +525,7 @@ FReply SPackagesDialog::OnKeyDown( const FGeometry& MyGeometry, const FKeyboardE
 		}
 	}
 
-	return SCompoundWidget::OnKeyDown( MyGeometry, InKeyboardEvent );
+	return SCompoundWidget::OnKeyDown( MyGeometry, InKeyEvent );
 }
 
 EVisibility SPackagesDialog::GetConnectToSourceControlVisibility() const
diff --git a/Engine/Source/Editor/Persona/Private/SAnimNotifyPanel.cpp b/Engine/Source/Editor/Persona/Private/SAnimNotifyPanel.cpp
--- a/Engine/Source/Editor/Persona/Private/SAnimNotifyPanel.cpp
+++ b/Engine/Source/Editor/Persona/Private/SAnimNotifyPanel.cpp
@@ -1940,7 +1940,7 @@ FReply SAnimNotifyTrack::OnMouseButtonUp( const FGeometry& MyGeometry, const FPo
 		WidgetToFocus = SummonContextMenu(MyGeometry, MouseEvent);
 
 		return (WidgetToFocus.IsValid())
-			? FReply::Handled().ReleaseMouseCapture().SetKeyboardFocus( WidgetToFocus.ToSharedRef(), EKeyboardFocusCause::SetDirectly )
+			? FReply::Handled().ReleaseMouseCapture().SetUserFocus(WidgetToFocus.ToSharedRef(), EFocusCause::SetDirectly)
 			: FReply::Handled().ReleaseMouseCapture();
 	}
 	else if ( bLeftMouseButton )
@@ -3514,9 +3514,9 @@ void SAnimNotifyPanel::BindCommands()
 		FExecuteAction::CreateSP(this, &SAnimNotifyPanel::OnDeletePressed));
 }
 
-FReply SAnimNotifyPanel::OnKeyDown(const FGeometry& MyGeometry, const FKeyboardEvent& InKeyboardEvent)
+FReply SAnimNotifyPanel::OnKeyDown(const FGeometry& MyGeometry, const FKeyEvent& InKeyEvent)
 {
-	if(UICommandList->ProcessCommandBindings(InKeyboardEvent))
+	if(UICommandList->ProcessCommandBindings(InKeyEvent))
 	{
 		return FReply::Handled();
 	}
@@ -3631,7 +3631,7 @@ FReply SAnimNotifyPanel::OnDragDetected(const FGeometry& MyGeometry, const FPoin
 	return FReply::Handled().CaptureMouse(SharedThis(this));
 }
 
-void SAnimNotifyPanel::OnKeyboardFocusLost(const FKeyboardFocusEvent& InKeyboardFocusEvent)
+void SAnimNotifyPanel::OnFocusLost(const FFocusEvent& InFocusEvent)
 {
 	if(Marquee.bActive)
 	{
diff --git a/Engine/Source/Editor/Persona/Private/SAnimNotifyPanel.h b/Engine/Source/Editor/Persona/Private/SAnimNotifyPanel.h
--- a/Engine/Source/Editor/Persona/Private/SAnimNotifyPanel.h
+++ b/Engine/Source/Editor/Persona/Private/SAnimNotifyPanel.h
@@ -195,12 +195,12 @@ public:
 	void OnPropertyChanged(UObject* ChangedObject, FPropertyChangedEvent& PropertyEvent);
 
 	/** SWidget Interface */
-	virtual FReply OnKeyDown(const FGeometry& MyGeometry, const FKeyboardEvent& InKeyboardEvent);	
+	virtual FReply OnKeyDown(const FGeometry& MyGeometry, const FKeyEvent& InKeyEvent);	
 	virtual FReply OnMouseButtonDown(const FGeometry& MyGeometry, const FPointerEvent& MouseEvent) override;
 	virtual FReply OnMouseButtonUp(const FGeometry& MyGeometry, const FPointerEvent& MouseEvent) override;
 	virtual FReply OnMouseMove(const FGeometry& MyGeometry, const FPointerEvent& MouseEvent) override;
 	virtual FReply OnDragDetected(const FGeometry& MyGeometry, const FPointerEvent& MouseEvent) override;
-	virtual void OnKeyboardFocusLost(const FKeyboardFocusEvent& InKeyboardFocusEvent);
+	virtual void OnFocusLost(const FFocusEvent& InFocusEvent);
 	virtual int32 OnPaint(const FPaintArgs& Args, const FGeometry& AllottedGeometry, const FSlateRect& MyClippingRect, FSlateWindowElementList& OutDrawElements, int32 LayerId, const FWidgetStyle& InWidgetStyle, bool bParentEnabled) const override;
 	/** End SWidget Interface */
 
diff --git a/Engine/Source/Editor/Persona/Private/SSkeletonTree.cpp b/Engine/Source/Editor/Persona/Private/SSkeletonTree.cpp
--- a/Engine/Source/Editor/Persona/Private/SSkeletonTree.cpp
+++ b/Engine/Source/Editor/Persona/Private/SSkeletonTree.cpp
@@ -2418,9 +2418,9 @@ bool SSkeletonTree::IsAddingSocketsAllowed() const
 	return false;
 }
 
-FReply SSkeletonTree::OnKeyDown( const FGeometry& MyGeometry, const FKeyboardEvent& InKeyboardEvent )
+FReply SSkeletonTree::OnKeyDown( const FGeometry& MyGeometry, const FKeyEvent& InKeyEvent )
 {
-	if ( UICommandList->ProcessCommandBindings( InKeyboardEvent ) )
+	if ( UICommandList->ProcessCommandBindings( InKeyEvent ) )
 	{
 		return FReply::Handled();
 	}
diff --git a/Engine/Source/Editor/Persona/Private/SSkeletonTree.h b/Engine/Source/Editor/Persona/Private/SSkeletonTree.h
--- a/Engine/Source/Editor/Persona/Private/SSkeletonTree.h
+++ b/Engine/Source/Editor/Persona/Private/SSkeletonTree.h
@@ -543,7 +543,7 @@ private:
 	void SetPreviewComponentSocketFilter() const;
 
 	/** Override OnKeyDown */
-	virtual FReply OnKeyDown( const FGeometry& MyGeometry, const FKeyboardEvent& InKeyboardEvent );
+	virtual FReply OnKeyDown( const FGeometry& MyGeometry, const FKeyEvent& InKeyEvent );
 
 	/** Function to delete all the selected sockets/assets */
 	void OnDeleteSelectedRows();
diff --git a/Engine/Source/Editor/Persona/Private/STrack.cpp b/Engine/Source/Editor/Persona/Private/STrack.cpp
--- a/Engine/Source/Editor/Persona/Private/STrack.cpp
+++ b/Engine/Source/Editor/Persona/Private/STrack.cpp
@@ -673,7 +673,7 @@ FReply STrack::OnMouseButtonUp( const FGeometry& MyGeometry, const FPointerEvent
 		WidgetToFocus = SummonContextMenu(MyGeometry, MouseEvent);
 
 		return (WidgetToFocus.IsValid())
-			? FReply::Handled().ReleaseMouseCapture().SetKeyboardFocus( WidgetToFocus.ToSharedRef(), EKeyboardFocusCause::SetDirectly )
+			? FReply::Handled().ReleaseMouseCapture().SetUserFocus( WidgetToFocus.ToSharedRef(), EFocusCause::SetDirectly )
 			: FReply::Handled().ReleaseMouseCapture();
 	}
 	else if ( bLeftMouseButton )
diff --git a/Engine/Source/Editor/PhAT/Private/SPhATNewAssetDlg.cpp b/Engine/Source/Editor/PhAT/Private/SPhATNewAssetDlg.cpp
--- a/Engine/Source/Editor/PhAT/Private/SPhATNewAssetDlg.cpp
+++ b/Engine/Source/Editor/PhAT/Private/SPhATNewAssetDlg.cpp
@@ -223,10 +223,10 @@ void SPhATNewAssetDlg::Construct(const FArguments& InArgs)
 	];
 }
 
-FReply SPhATNewAssetDlg::OnKeyDown(const FGeometry& MyGeometry, const FKeyboardEvent& InKeyboardEvent)
+FReply SPhATNewAssetDlg::OnKeyDown(const FGeometry& MyGeometry, const FKeyEvent& InKeyEvent)
 {
 	//see if we pressed the Enter or Spacebar keys
-	if (InKeyboardEvent.GetKey() == EKeys::Escape)
+	if (InKeyEvent.GetKey() == EKeys::Escape)
 	{
 		return OnClicked(EAppReturnType::Cancel);
 	}
diff --git a/Engine/Source/Editor/PhAT/Private/SPhATNewAssetDlg.h b/Engine/Source/Editor/PhAT/Private/SPhATNewAssetDlg.h
--- a/Engine/Source/Editor/PhAT/Private/SPhATNewAssetDlg.h
+++ b/Engine/Source/Editor/PhAT/Private/SPhATNewAssetDlg.h
@@ -23,7 +23,7 @@ public:
 	void Construct(const FArguments& InArgs);
 
 	/** SWidget interface */
-	virtual FReply OnKeyDown(const FGeometry& MyGeometry, const FKeyboardEvent& InKeyboardEvent) override;
+	virtual FReply OnKeyDown(const FGeometry& MyGeometry, const FKeyEvent& InKeyEvent) override;
 	virtual bool SupportsKeyboardFocus() const override;
 
 private:
diff --git a/Engine/Source/Editor/PlacementMode/Private/SPlacementModeTools.cpp b/Engine/Source/Editor/PlacementMode/Private/SPlacementModeTools.cpp
--- a/Engine/Source/Editor/PlacementMode/Private/SPlacementModeTools.cpp
+++ b/Engine/Source/Editor/PlacementMode/Private/SPlacementModeTools.cpp
@@ -953,11 +953,11 @@ void SPlacementModeTools::Tick( const FGeometry& AllottedGeometry, const double
 	}
 }
 
-FReply SPlacementModeTools::OnKeyDown( const FGeometry& MyGeometry, const FKeyboardEvent& InKeyboardEvent )
+FReply SPlacementModeTools::OnKeyDown( const FGeometry& MyGeometry, const FKeyEvent& InKeyEvent )
 {
 	FReply Reply = FReply::Unhandled();
 
-	if ( InKeyboardEvent.GetKey() == EKeys::Escape )
+	if ( InKeyEvent.GetKey() == EKeys::Escape )
 	{
 		FPlacementMode* PlacementEditMode = (FPlacementMode*)GLevelEditorModeTools().GetActiveMode( FBuiltinEditorModes::EM_Placement );
 		PlacementEditMode->StopPlacing();
diff --git a/Engine/Source/Editor/PropertyEditor/Private/SDetailsViewBase.cpp b/Engine/Source/Editor/PropertyEditor/Private/SDetailsViewBase.cpp
--- a/Engine/Source/Editor/PropertyEditor/Private/SDetailsViewBase.cpp
+++ b/Engine/Source/Editor/PropertyEditor/Private/SDetailsViewBase.cpp
@@ -621,13 +621,13 @@ bool SDetailsViewBase::SupportsKeyboardFocus() const
 	return DetailsViewArgs.bSearchInitialKeyFocus && SearchBox->SupportsKeyboardFocus() && GetFilterBoxVisibility() == EVisibility::Visible;
 }
 
-FReply SDetailsViewBase::OnKeyboardFocusReceived(const FGeometry& MyGeometry, const FKeyboardFocusEvent& InKeyboardFocusEvent)
+FReply SDetailsViewBase::OnFocusReceived(const FGeometry& MyGeometry, const FFocusEvent& InFocusEvent)
 {
 	FReply Reply = FReply::Handled();
 
-	if (InKeyboardFocusEvent.GetCause() != EKeyboardFocusCause::Cleared)
+	if (InFocusEvent.GetCause() != EFocusCause::Cleared)
 	{
-		Reply.SetKeyboardFocus(SearchBox.ToSharedRef(), InKeyboardFocusEvent.GetCause());
+		Reply.SetUserFocus(SearchBox.ToSharedRef(), InFocusEvent.GetCause());
 	}
 
 	return Reply;
diff --git a/Engine/Source/Editor/PropertyEditor/Private/SDetailsViewBase.h b/Engine/Source/Editor/PropertyEditor/Private/SDetailsViewBase.h
--- a/Engine/Source/Editor/PropertyEditor/Private/SDetailsViewBase.h
+++ b/Engine/Source/Editor/PropertyEditor/Private/SDetailsViewBase.h
@@ -260,7 +260,7 @@ public:
 
 	// SWidget interface
 	virtual bool SupportsKeyboardFocus() const override;
-	virtual FReply OnKeyboardFocusReceived(const FGeometry& MyGeometry, const FKeyboardFocusEvent& InKeyboardFocusEvent) override;
+	virtual FReply OnFocusReceived(const FGeometry& MyGeometry, const FFocusEvent& InFocusEvent) override;
 	virtual void Tick(const FGeometry& AllottedGeometry, const double InCurrentTime, const float InDeltaTime) override;
 	// End of SWidget interface
 
diff --git a/Engine/Source/Editor/PropertyEditor/Private/UserInterface/PropertyEditor/SPropertyComboBox.cpp b/Engine/Source/Editor/PropertyEditor/Private/UserInterface/PropertyEditor/SPropertyComboBox.cpp
--- a/Engine/Source/Editor/PropertyEditor/Private/UserInterface/PropertyEditor/SPropertyComboBox.cpp
+++ b/Engine/Source/Editor/PropertyEditor/Private/UserInterface/PropertyEditor/SPropertyComboBox.cpp
@@ -125,9 +125,9 @@ TSharedRef<SWidget> SPropertyComboBox::OnGenerateComboWidget( TSharedPtr<FString
 		.IsEnabled(bEnabled);
 }
 
-FReply SPropertyComboBox::OnKeyDown( const FGeometry& MyGeometry, const FKeyboardEvent& InKeyboardEvent )
+FReply SPropertyComboBox::OnKeyDown( const FGeometry& MyGeometry, const FKeyEvent& InKeyEvent )
 {
-	const FKey Key = InKeyboardEvent.GetKey();
+	const FKey Key = InKeyEvent.GetKey();
 
 	if(Key == EKeys::Up)
 	{
@@ -180,7 +180,7 @@ FReply SPropertyComboBox::OnKeyDown( const FGeometry& MyGeometry, const FKeyboar
 		return FReply::Handled();
 	}
 
-	return SComboBox< TSharedPtr<FString> >::OnKeyDown( MyGeometry, InKeyboardEvent );
+	return SComboBox< TSharedPtr<FString> >::OnKeyDown( MyGeometry, InKeyEvent );
 }
 
 #undef LOCTEXT_NAMESPACE
diff --git a/Engine/Source/Editor/PropertyEditor/Private/UserInterface/PropertyEditor/SPropertyComboBox.h b/Engine/Source/Editor/PropertyEditor/Private/UserInterface/PropertyEditor/SPropertyComboBox.h
--- a/Engine/Source/Editor/PropertyEditor/Private/UserInterface/PropertyEditor/SPropertyComboBox.h
+++ b/Engine/Source/Editor/PropertyEditor/Private/UserInterface/PropertyEditor/SPropertyComboBox.h
@@ -46,7 +46,7 @@ private:
 	TSharedRef<SWidget> OnGenerateComboWidget( TSharedPtr<FString> InComboString );
 
 	/** SWidget interface */
-	FReply OnKeyDown( const FGeometry& MyGeometry, const FKeyboardEvent& InKeyboardEvent ) override;
+	FReply OnKeyDown( const FGeometry& MyGeometry, const FKeyEvent& InKeyEvent ) override;
 
 private:
 
diff --git a/Engine/Source/Editor/PropertyEditor/Private/UserInterface/PropertyEditor/SPropertyEditorBool.cpp b/Engine/Source/Editor/PropertyEditor/Private/UserInterface/PropertyEditor/SPropertyEditorBool.cpp
--- a/Engine/Source/Editor/PropertyEditor/Private/UserInterface/PropertyEditor/SPropertyEditorBool.cpp
+++ b/Engine/Source/Editor/PropertyEditor/Private/UserInterface/PropertyEditor/SPropertyEditorBool.cpp
@@ -45,10 +45,10 @@ bool SPropertyEditorBool::HasKeyboardFocus() const
 	return CheckBox->HasKeyboardFocus();
 }
 
-FReply SPropertyEditorBool::OnKeyboardFocusReceived( const FGeometry& MyGeometry, const FKeyboardFocusEvent& InKeyboardFocusEvent )
+FReply SPropertyEditorBool::OnFocusReceived( const FGeometry& MyGeometry, const FFocusEvent& InFocusEvent )
 {
 	// Forward keyboard focus to our editable text widget
-	return FReply::Handled().SetKeyboardFocus( CheckBox.ToSharedRef(), InKeyboardFocusEvent.GetCause() );
+	return FReply::Handled().SetUserFocus(CheckBox.ToSharedRef(), InFocusEvent.GetCause());
 }
 
 ESlateCheckBoxState::Type SPropertyEditorBool::OnGetCheckState() const
@@ -93,7 +93,7 @@ FReply SPropertyEditorBool::OnMouseButtonDoubleClick( const FGeometry& InMyGeome
 		CheckBox->ToggleCheckedState();
 
 		// Set focus to this object, but don't capture the mouse
-		return FReply::Handled().SetKeyboardFocus( AsShared(), EKeyboardFocusCause::Mouse );
+		return FReply::Handled().SetUserFocus(AsShared(), EFocusCause::Mouse);
 	}
 
 	return FReply::Unhandled();
diff --git a/Engine/Source/Editor/PropertyEditor/Private/UserInterface/PropertyEditor/SPropertyEditorBool.h b/Engine/Source/Editor/PropertyEditor/Private/UserInterface/PropertyEditor/SPropertyEditorBool.h
--- a/Engine/Source/Editor/PropertyEditor/Private/UserInterface/PropertyEditor/SPropertyEditorBool.h
+++ b/Engine/Source/Editor/PropertyEditor/Private/UserInterface/PropertyEditor/SPropertyEditorBool.h
@@ -17,7 +17,7 @@ public:
 	// SWidget overrides
 	virtual bool SupportsKeyboardFocus() const override;
 	virtual bool HasKeyboardFocus() const override;
-	virtual FReply OnKeyboardFocusReceived( const FGeometry& MyGeometry, const FKeyboardFocusEvent& InKeyboardFocusEvent ) override;
+	virtual FReply OnFocusReceived( const FGeometry& MyGeometry, const FFocusEvent& InFocusEvent ) override;
 
 private:
 
diff --git a/Engine/Source/Editor/PropertyEditor/Private/UserInterface/PropertyEditor/SPropertyEditorInteractiveActorPicker.cpp b/Engine/Source/Editor/PropertyEditor/Private/UserInterface/PropertyEditor/SPropertyEditorInteractiveActorPicker.cpp
--- a/Engine/Source/Editor/PropertyEditor/Private/UserInterface/PropertyEditor/SPropertyEditorInteractiveActorPicker.cpp
+++ b/Engine/Source/Editor/PropertyEditor/Private/UserInterface/PropertyEditor/SPropertyEditorInteractiveActorPicker.cpp
@@ -38,9 +38,9 @@ void SPropertyEditorInteractiveActorPicker::Construct( const FArguments& InArgs
 	);
 }
 
-FReply SPropertyEditorInteractiveActorPicker::OnKeyDown( const FGeometry& MyGeometry, const FKeyboardEvent& InKeyboardEvent )
+FReply SPropertyEditorInteractiveActorPicker::OnKeyDown( const FGeometry& MyGeometry, const FKeyEvent& InKeyEvent )
 {
-	if(InKeyboardEvent.GetKey() == EKeys::Escape)
+	if(InKeyEvent.GetKey() == EKeys::Escape)
 	{
 		FActorPickerModeModule& ActorPickerMode = FModuleManager::Get().GetModuleChecked<FActorPickerModeModule>("ActorPickerMode");
 
diff --git a/Engine/Source/Editor/PropertyEditor/Private/UserInterface/PropertyEditor/SPropertyEditorNumeric.h b/Engine/Source/Editor/PropertyEditor/Private/UserInterface/PropertyEditor/SPropertyEditorNumeric.h
--- a/Engine/Source/Editor/PropertyEditor/Private/UserInterface/PropertyEditor/SPropertyEditorNumeric.h
+++ b/Engine/Source/Editor/PropertyEditor/Private/UserInterface/PropertyEditor/SPropertyEditorNumeric.h
@@ -117,10 +117,10 @@ public:
 		return PrimaryWidget.IsValid() && PrimaryWidget->SupportsKeyboardFocus();
 	}
 
-	virtual FReply OnKeyboardFocusReceived( const FGeometry& MyGeometry, const FKeyboardFocusEvent& InKeyboardFocusEvent ) override
+	virtual FReply OnFocusReceived( const FGeometry& MyGeometry, const FFocusEvent& InFocusEvent ) override
 	{
 		// Forward keyboard focus to our editable text widget
-		return FReply::Handled().SetKeyboardFocus( PrimaryWidget.ToSharedRef(), InKeyboardFocusEvent.GetCause() );
+		return FReply::Handled().SetUserFocus(PrimaryWidget.ToSharedRef(), InFocusEvent.GetCause());
 	}
 	
 	void GetDesiredWidth( float& OutMinDesiredWidth, float& OutMaxDesiredWidth )
diff --git a/Engine/Source/Editor/PropertyEditor/Private/UserInterface/PropertyEditor/SPropertyEditorTableRow.cpp b/Engine/Source/Editor/PropertyEditor/Private/UserInterface/PropertyEditor/SPropertyEditorTableRow.cpp
--- a/Engine/Source/Editor/PropertyEditor/Private/UserInterface/PropertyEditor/SPropertyEditorTableRow.cpp
+++ b/Engine/Source/Editor/PropertyEditor/Private/UserInterface/PropertyEditor/SPropertyEditorTableRow.cpp
@@ -215,7 +215,7 @@ FReply SPropertyEditorTableRow::OnNameDoubleClicked()
 		FSlateApplication::Get().GeneratePathToWidgetUnchecked( ValueEditorWidget.ToSharedRef(), EditableWidgetPath );
 
 		// Set keyboard focus directly
-		FSlateApplication::Get().SetKeyboardFocus( EditableWidgetPath, EKeyboardFocusCause::SetDirectly );
+		FSlateApplication::Get().SetKeyboardFocus( EditableWidgetPath, EFocusCause::SetDirectly );
 		Reply = FReply::Handled();
 	}
 	else if ( DoesItemHaveChildren() )
diff --git a/Engine/Source/Editor/PropertyEditor/Private/UserInterface/PropertyEditor/SPropertyEditorText.cpp b/Engine/Source/Editor/PropertyEditor/Private/UserInterface/PropertyEditor/SPropertyEditorText.cpp
--- a/Engine/Source/Editor/PropertyEditor/Private/UserInterface/PropertyEditor/SPropertyEditorText.cpp
+++ b/Engine/Source/Editor/PropertyEditor/Private/UserInterface/PropertyEditor/SPropertyEditorText.cpp
@@ -100,10 +100,10 @@ bool SPropertyEditorText::SupportsKeyboardFocus() const
 	return PrimaryWidget.IsValid() && PrimaryWidget->SupportsKeyboardFocus();
 }
 
-FReply SPropertyEditorText::OnKeyboardFocusReceived( const FGeometry& MyGeometry, const FKeyboardFocusEvent& InKeyboardFocusEvent )
+FReply SPropertyEditorText::OnFocusReceived( const FGeometry& MyGeometry, const FFocusEvent& InFocusEvent )
 {
 	// Forward keyboard focus to our editable text widget
-	return FReply::Handled().SetKeyboardFocus( PrimaryWidget.ToSharedRef(), InKeyboardFocusEvent.GetCause() );
+	return FReply::Handled().SetUserFocus(PrimaryWidget.ToSharedRef(), InFocusEvent.GetCause());
 }
 
 void SPropertyEditorText::Tick( const FGeometry& AllottedGeometry, const double InCurrentTime, const float InDeltaTime )
diff --git a/Engine/Source/Editor/PropertyEditor/Private/UserInterface/PropertyEditor/SPropertyEditorText.h b/Engine/Source/Editor/PropertyEditor/Private/UserInterface/PropertyEditor/SPropertyEditorText.h
--- a/Engine/Source/Editor/PropertyEditor/Private/UserInterface/PropertyEditor/SPropertyEditorText.h
+++ b/Engine/Source/Editor/PropertyEditor/Private/UserInterface/PropertyEditor/SPropertyEditorText.h
@@ -21,7 +21,7 @@ public:
 
 	bool SupportsKeyboardFocus() const override;
 
-	FReply OnKeyboardFocusReceived( const FGeometry& MyGeometry, const FKeyboardFocusEvent& InKeyboardFocusEvent ) override;
+	FReply OnFocusReceived( const FGeometry& MyGeometry, const FFocusEvent& InFocusEvent ) override;
 
 	virtual void Tick( const FGeometry& AllottedGeometry, const double InCurrentTime, const float InDeltaTime ) override;
 
diff --git a/Engine/Source/Editor/PropertyEditor/Private/UserInterface/PropertyTable/SPropertyTable.h b/Engine/Source/Editor/PropertyEditor/Private/UserInterface/PropertyTable/SPropertyTable.h
--- a/Engine/Source/Editor/PropertyEditor/Private/UserInterface/PropertyTable/SPropertyTable.h
+++ b/Engine/Source/Editor/PropertyEditor/Private/UserInterface/PropertyTable/SPropertyTable.h
@@ -123,7 +123,7 @@ public:
 		return SharedThis(this);
 	}
 
-	virtual FReply OnKeyDown( const FGeometry& MyGeometry, const FKeyboardEvent& InKeyboardEvent ) override
+	virtual FReply OnKeyDown( const FGeometry& MyGeometry, const FKeyEvent& InKeyEvent ) override
 	{
 		FReply Reply = FReply::Unhandled();
 		const ESelectionMode::Type SelectionMode = Table->GetSelectionMode();
@@ -133,7 +133,7 @@ public:
 			return Reply;
 		}
 
-		const FKey Key = InKeyboardEvent.GetKey();
+		const FKey Key = InKeyEvent.GetKey();
 		const EPropertyTableSelectionUnit::Type SelectionUnit = Table->GetSelectionUnit();
 			
 		const TSharedPtr< IPropertyTableCell > CurrentCell = Table->GetCurrentCell();
@@ -151,7 +151,7 @@ public:
 				Reply = FReply::Handled();
 			}
 		}
-		else if( Key == EKeys::C && InKeyboardEvent.IsControlDown() ) 
+		else if( Key == EKeys::C && InKeyEvent.IsControlDown() ) 
 		{
 			if ( CurrentCell.IsValid() && FirstCellInSelection.IsValid() && LastCellInSelection.IsValid() )
 			{
@@ -218,7 +218,7 @@ public:
 				Reply = FReply::Handled();
 			}
 		}
-		else if( Key == EKeys::V && InKeyboardEvent.IsControlDown() ) 
+		else if( Key == EKeys::V && InKeyEvent.IsControlDown() ) 
 		{
 			if ( CurrentCell.IsValid() )
 			{
@@ -231,21 +231,21 @@ public:
 				Reply = FReply::Handled();
 			}
 		}
-		else if ( Key == EKeys::A && InKeyboardEvent.IsControlDown() )
+		else if ( Key == EKeys::A && InKeyEvent.IsControlDown() )
 		{
 			return SelectRange( Table->GetFirstCellInTable(), Table->GetLastCellInTable(), Table->GetCurrentCell() );
 		}
 		else if ( Key == EKeys::Home )
 		{
-			if ( CurrentCell.IsValid() && InKeyboardEvent.IsShiftDown() && InKeyboardEvent.IsControlDown() )
+			if ( CurrentCell.IsValid() && InKeyEvent.IsShiftDown() && InKeyEvent.IsControlDown() )
 			{
 				return SelectRange( Table->GetFirstCellInTable(), CurrentCell, CurrentCell );
 			}
-			else if ( InKeyboardEvent.IsControlDown() )
+			else if ( InKeyEvent.IsControlDown() )
 			{
 				return MoveToCell( Table->GetFirstCellInTable() );
 			}
-			else if ( CurrentRow.IsValid() && CurrentCell.IsValid() && InKeyboardEvent.IsShiftDown() )
+			else if ( CurrentRow.IsValid() && CurrentCell.IsValid() && InKeyEvent.IsShiftDown() )
 			{
 				return SelectRange( Table->GetFirstCellInRow( CurrentRow.ToSharedRef() ), CurrentCell, CurrentCell );
 			}
@@ -256,15 +256,15 @@ public:
 		}
 		else if ( Key == EKeys::End )
 		{
-			if ( CurrentCell.IsValid() && InKeyboardEvent.IsShiftDown() && InKeyboardEvent.IsControlDown() )
+			if ( CurrentCell.IsValid() && InKeyEvent.IsShiftDown() && InKeyEvent.IsControlDown() )
 			{
 				return SelectRange( CurrentCell, Table->GetLastCellInTable(), CurrentCell );
 			}
-			else if ( InKeyboardEvent.IsControlDown() )
+			else if ( InKeyEvent.IsControlDown() )
 			{
 				return MoveToCell( Table->GetLastCellInTable() );
 			}
-			else if ( CurrentRow.IsValid() && CurrentCell.IsValid() && InKeyboardEvent.IsShiftDown() )
+			else if ( CurrentRow.IsValid() && CurrentCell.IsValid() && InKeyEvent.IsShiftDown() )
 			{
 				return SelectRange( CurrentCell, Table->GetLastCellInRow( CurrentRow.ToSharedRef() ), CurrentCell );
 			}
@@ -275,15 +275,15 @@ public:
 		}
 		else if ( Key == EKeys::Left )
 		{
-			if ( CurrentRow.IsValid() && CurrentCell.IsValid() && InKeyboardEvent.IsShiftDown() && InKeyboardEvent.IsControlDown() )
+			if ( CurrentRow.IsValid() && CurrentCell.IsValid() && InKeyEvent.IsShiftDown() && InKeyEvent.IsControlDown() )
 			{
 				return SelectRange( Table->GetFirstCellInRow( CurrentRow.ToSharedRef() ), CurrentCell, CurrentCell );
 			}
-			else if ( InKeyboardEvent.IsControlDown() && CurrentRow.IsValid() )
+			else if ( InKeyEvent.IsControlDown() && CurrentRow.IsValid() )
 			{
 				return MoveToCell( Table->GetFirstCellInRow( CurrentRow.ToSharedRef() ) );
 			}
-			else if ( CurrentCell.IsValid() && InKeyboardEvent.IsShiftDown() )
+			else if ( CurrentCell.IsValid() && InKeyEvent.IsShiftDown() )
 			{
 				if ( FirstCellInSelection.IsValid() && LastCellInSelection.IsValid() )
 				{
@@ -312,15 +312,15 @@ public:
 		}
 		else if ( Key == EKeys::Right )
 		{
-			if ( CurrentRow.IsValid() && CurrentCell.IsValid() && InKeyboardEvent.IsShiftDown() && InKeyboardEvent.IsControlDown() )
+			if ( CurrentRow.IsValid() && CurrentCell.IsValid() && InKeyEvent.IsShiftDown() && InKeyEvent.IsControlDown() )
 			{
 				return SelectRange( Table->GetLastCellInRow( CurrentRow.ToSharedRef() ), CurrentCell, CurrentCell );
 			}
-			else if ( InKeyboardEvent.IsControlDown() && CurrentRow.IsValid() )
+			else if ( InKeyEvent.IsControlDown() && CurrentRow.IsValid() )
 			{
 				return MoveToCell( Table->GetLastCellInRow( CurrentRow.ToSharedRef() ) );
 			}
-			else if ( CurrentCell.IsValid() && InKeyboardEvent.IsShiftDown() )
+			else if ( CurrentCell.IsValid() && InKeyEvent.IsShiftDown() )
 			{
 				if ( FirstCellInSelection.IsValid() && LastCellInSelection.IsValid() )
 				{
@@ -349,15 +349,15 @@ public:
 		}
 		else if ( Key == EKeys::Up )
 		{
-			if ( InKeyboardEvent.IsShiftDown() && InKeyboardEvent.IsControlDown() && CurrentColumn.IsValid() && CurrentCell.IsValid() )
+			if ( InKeyEvent.IsShiftDown() && InKeyEvent.IsControlDown() && CurrentColumn.IsValid() && CurrentCell.IsValid() )
 			{
 				return SelectRange( Table->GetFirstCellInColumn( CurrentColumn.ToSharedRef() ), CurrentCell, CurrentCell );
 			}
-			else if ( InKeyboardEvent.IsControlDown() && CurrentColumn.IsValid() )
+			else if ( InKeyEvent.IsControlDown() && CurrentColumn.IsValid() )
 			{
 				return MoveToCell( Table->GetFirstCellInColumn( CurrentColumn.ToSharedRef() ) );
 			}
-			else if ( CurrentCell.IsValid() && InKeyboardEvent.IsShiftDown() )
+			else if ( CurrentCell.IsValid() && InKeyEvent.IsShiftDown() )
 			{
 				if ( FirstCellInSelection.IsValid() && LastCellInSelection.IsValid() )
 				{
@@ -386,15 +386,15 @@ public:
 		}
 		else if ( Key == EKeys::Down )
 		{
-			if ( InKeyboardEvent.IsShiftDown() && InKeyboardEvent.IsControlDown() && CurrentColumn.IsValid() && CurrentCell.IsValid() )
+			if ( InKeyEvent.IsShiftDown() && InKeyEvent.IsControlDown() && CurrentColumn.IsValid() && CurrentCell.IsValid() )
 			{
 				return SelectRange( CurrentCell, Table->GetLastCellInColumn( CurrentColumn.ToSharedRef() ), CurrentCell );
 			}
-			else if ( InKeyboardEvent.IsControlDown() && CurrentColumn.IsValid() )
+			else if ( InKeyEvent.IsControlDown() && CurrentColumn.IsValid() )
 			{
 				return MoveToCell( Table->GetLastCellInColumn( CurrentColumn.ToSharedRef() ) );
 			}
-			else if ( CurrentCell.IsValid() && InKeyboardEvent.IsShiftDown() )
+			else if ( CurrentCell.IsValid() && InKeyEvent.IsShiftDown() )
 			{
 				if ( FirstCellInSelection.IsValid() && LastCellInSelection.IsValid() )
 				{
@@ -425,7 +425,7 @@ public:
 		{
 			if ( CurrentCell.IsValid() )
 			{
-				if ( InKeyboardEvent.IsShiftDown() )
+				if ( InKeyEvent.IsShiftDown() )
 				{
 					MoveToCell( Table->GetPreviousCellInRow( CurrentCell.ToSharedRef() ) );
 				}
@@ -451,12 +451,12 @@ public:
 			}
 		}
 
-		return STreeView< TSharedRef< IPropertyTableRow > >::OnKeyDown( MyGeometry, InKeyboardEvent );
+		return STreeView< TSharedRef< IPropertyTableRow > >::OnKeyDown( MyGeometry, InKeyEvent );
 	}
 
-	virtual FReply OnPreviewKeyDown( const FGeometry& MyGeometry, const FKeyboardEvent& InKeyboardEvent ) override
+	virtual FReply OnPreviewKeyDown( const FGeometry& MyGeometry, const FKeyEvent& InKeyEvent ) override
 	{
-		const FKey Key = InKeyboardEvent.GetKey();
+		const FKey Key = InKeyEvent.GetKey();
 		if ( Key == EKeys::Enter )
 		{
 			const TSharedPtr< IPropertyTableCell > CurrentCell = Table->GetCurrentCell();
@@ -466,7 +466,7 @@ public:
 			}
 		}
 
-		return STreeView< TSharedRef< IPropertyTableRow > >::OnPreviewKeyDown( MyGeometry, InKeyboardEvent );
+		return STreeView< TSharedRef< IPropertyTableRow > >::OnPreviewKeyDown( MyGeometry, InKeyEvent );
 	}
 
 	FReply SelectRange( const TSharedPtr< IPropertyTableCell >& StartingCell, const TSharedPtr< IPropertyTableCell >& EndingCell, const TSharedPtr< IPropertyTableCell >& CellToFocus )
@@ -542,7 +542,7 @@ public:
 		return FReply::Handled();
 	}
 
-	virtual FReply OnKeyUp( const FGeometry& MyGeometry, const FKeyboardEvent& InKeyboardEvent ) override
+	virtual FReply OnKeyUp( const FGeometry& MyGeometry, const FKeyEvent& InKeyEvent ) override
 	{
 		FReply Reply = FReply::Unhandled();
 
@@ -550,7 +550,7 @@ public:
 
 		if ( CurrentCell.IsValid() )
 		{
-			FKey Key = InKeyboardEvent.GetKey();
+			FKey Key = InKeyEvent.GetKey();
 
 			if ( Key == EKeys::F2 )
 			{
diff --git a/Engine/Source/Editor/PropertyEditor/Private/UserInterface/PropertyTable/SPropertyTableCell.cpp b/Engine/Source/Editor/PropertyEditor/Private/UserInterface/PropertyTable/SPropertyTableCell.cpp
--- a/Engine/Source/Editor/PropertyEditor/Private/UserInterface/PropertyTable/SPropertyTableCell.cpp
+++ b/Engine/Source/Editor/PropertyEditor/Private/UserInterface/PropertyTable/SPropertyTableCell.cpp
@@ -125,11 +125,11 @@ void SPropertyTableCell::Tick( const FGeometry& AllottedGeometry, const double I
 					DropDownAnchor->SetIsOpen( true, false );
 				}
 
-				FSlateApplication::Get().SetKeyboardFocus( Presenter->WidgetToFocusOnEdit(), EKeyboardFocusCause::SetDirectly );
+				FSlateApplication::Get().SetKeyboardFocus( Presenter->WidgetToFocusOnEdit(), EFocusCause::SetDirectly );
 			}
 			else
 			{
-				FSlateApplication::Get().SetKeyboardFocus( ChildSlot.GetChildAt( 0 ), EKeyboardFocusCause::SetDirectly );
+				FSlateApplication::Get().SetKeyboardFocus( ChildSlot.GetChildAt( 0 ), EFocusCause::SetDirectly );
 			}
 		}
 
diff --git a/Engine/Source/Editor/SceneOutliner/Private/SSceneOutliner.cpp b/Engine/Source/Editor/SceneOutliner/Private/SSceneOutliner.cpp
--- a/Engine/Source/Editor/SceneOutliner/Private/SSceneOutliner.cpp
+++ b/Engine/Source/Editor/SceneOutliner/Private/SSceneOutliner.cpp
@@ -1543,7 +1543,7 @@ namespace SceneOutliner
 				// Set the keyboard focus back to the SceneOutliner, so the user can perform keyboard commands
 				FWidgetPath SceneOutlinerWidgetPath;
 				FSlateApplication::Get().GeneratePathToWidgetUnchecked( SharedThis( this ), SceneOutlinerWidgetPath );
-				FSlateApplication::Get().SetKeyboardFocus( SceneOutlinerWidgetPath, EKeyboardFocusCause::SetDirectly );
+				FSlateApplication::Get().SetKeyboardFocus( SceneOutlinerWidgetPath, EFocusCause::SetDirectly );
 			}
 		}
 		else
@@ -2846,7 +2846,7 @@ namespace SceneOutliner
 						FSlateApplication::Get().GeneratePathToWidgetUnchecked( SharedThis( this ), SceneOutlinerWidgetPath );
 
 						// Set keyboard focus directly
-						FSlateApplication::Get().SetKeyboardFocus( SceneOutlinerWidgetPath, EKeyboardFocusCause::SetDirectly );
+						FSlateApplication::Get().SetKeyboardFocus( SceneOutlinerWidgetPath, EFocusCause::SetDirectly );
 					}
 				}
 
@@ -2978,7 +2978,7 @@ namespace SceneOutliner
 		return false;
 	}
 
-	FReply SSceneOutliner::OnKeyDown( const FGeometry& MyGeometry, const FKeyboardEvent& InKeyboardEvent )
+	FReply SSceneOutliner::OnKeyDown( const FGeometry& MyGeometry, const FKeyEvent& InKeyEvent )
 	{
 		// @todo outliner: Use command system for these for discoverability? (allow bindings?)
 
@@ -2986,7 +2986,7 @@ namespace SceneOutliner
 		if( InitOptions.Mode == ESceneOutlinerMode::ActorBrowsing )
 		{
 			// Delete key: Delete selected actors (not rebindable, because it doesn't make much sense to bind.)
-			if( InKeyboardEvent.GetKey() == EKeys::F2 )
+			if( InKeyEvent.GetKey() == EKeys::F2 )
 			{
 				auto SelectedItems = OutlinerTreeView->GetSelectedItems();
 				if (SelectedItems.Num() == 1 && SelectedItems[0]->Type == TOutlinerTreeItem::Folder)
@@ -2995,7 +2995,7 @@ namespace SceneOutliner
 					return FReply::Handled();
 				}
 			}
-			else if ( InKeyboardEvent.GetKey() == EKeys::Platform_Delete )
+			else if ( InKeyEvent.GetKey() == EKeys::Platform_Delete )
 			{
 				if( InitOptions.CustomDelete.IsBound() )
 				{
@@ -3079,7 +3079,7 @@ namespace SceneOutliner
 		{
 			FWidgetPath WidgetToFocusPath;
 			FSlateApplication::Get().GeneratePathToWidgetUnchecked( FilterTextBoxWidget.ToSharedRef(), WidgetToFocusPath );
-			FSlateApplication::Get().SetKeyboardFocus( WidgetToFocusPath, EKeyboardFocusCause::SetDirectly );
+			FSlateApplication::Get().SetKeyboardFocus( WidgetToFocusPath, EFocusCause::SetDirectly );
 			bPendingFocusNextFrame = false;
 		}
 
diff --git a/Engine/Source/Editor/Sequencer/Private/SAnimationOutlinerView.cpp b/Engine/Source/Editor/Sequencer/Private/SAnimationOutlinerView.cpp
--- a/Engine/Source/Editor/Sequencer/Private/SAnimationOutlinerView.cpp
+++ b/Engine/Source/Editor/Sequencer/Private/SAnimationOutlinerView.cpp
@@ -117,7 +117,7 @@ FReply SAnimationOutlinerTreeNode::OnMouseButtonUp( const FGeometry& MyGeometry,
 				FPopupTransitionEffect( FPopupTransitionEffect::ContextMenu )
 				);
 			
-			return FReply::Handled().ReleaseMouseCapture().SetKeyboardFocus( MenuContent.ToSharedRef(), EKeyboardFocusCause::SetDirectly );
+			return FReply::Handled().ReleaseMouseCapture().SetUserFocus(MenuContent.ToSharedRef(), EFocusCause::SetDirectly);
 		}
 
 		return FReply::Handled().ReleaseMouseCapture();
diff --git a/Engine/Source/Editor/Sequencer/Private/SSection.cpp b/Engine/Source/Editor/Sequencer/Private/SSection.cpp
--- a/Engine/Source/Editor/Sequencer/Private/SSection.cpp
+++ b/Engine/Source/Editor/Sequencer/Private/SSection.cpp
@@ -533,7 +533,7 @@ FReply SSection::OnMouseButtonUp( const FGeometry& MyGeometry, const FPointerEve
 					FPopupTransitionEffect( FPopupTransitionEffect::ContextMenu )
 					);
 
-				return FReply::Handled().ReleaseMouseCapture().SetKeyboardFocus( MenuContent.ToSharedRef(), EKeyboardFocusCause::SetDirectly );
+				return FReply::Handled().ReleaseMouseCapture().SetUserFocus(MenuContent.ToSharedRef(), EFocusCause::SetDirectly);
 			}
 		}
 	}
diff --git a/Engine/Source/Editor/Sequencer/Private/SSequencer.cpp b/Engine/Source/Editor/Sequencer/Private/SSequencer.cpp
--- a/Engine/Source/Editor/Sequencer/Private/SSequencer.cpp
+++ b/Engine/Source/Editor/Sequencer/Private/SSequencer.cpp
@@ -517,10 +517,10 @@ FReply SSequencer::OnDrop( const FGeometry& MyGeometry, const FDragDropEvent& Dr
 	return bWasDropHandled ? FReply::Handled() : FReply::Unhandled();
 }
 
-FReply SSequencer::OnKeyDown( const FGeometry& MyGeometry, const FKeyboardEvent& InKeyboardEvent ) 
+FReply SSequencer::OnKeyDown( const FGeometry& MyGeometry, const FKeyEvent& InKeyEvent ) 
 {
 	// A toolkit tab is active, so direct all command processing to it
-	if( Sequencer.Pin()->GetCommandBindings().ProcessCommandBindings( InKeyboardEvent ) )
+	if( Sequencer.Pin()->GetCommandBindings().ProcessCommandBindings( InKeyEvent ) )
 	{
 		return FReply::Handled();
 	}
diff --git a/Engine/Source/Editor/TranslationEditor/Private/TranslationPickerFloatingWindow.cpp b/Engine/Source/Editor/TranslationEditor/Private/TranslationPickerFloatingWindow.cpp
--- a/Engine/Source/Editor/TranslationEditor/Private/TranslationPickerFloatingWindow.cpp
+++ b/Engine/Source/Editor/TranslationEditor/Private/TranslationPickerFloatingWindow.cpp
@@ -128,9 +128,9 @@ FText STranslationPickerFloatingWindow::FormatFTextInfo(FText TextToFormat)
 	return OutText; 
 }
 
-FReply STranslationPickerFloatingWindow::OnKeyDown(const FGeometry& MyGeometry, const FKeyboardEvent& InKeyboardEvent)
+FReply STranslationPickerFloatingWindow::OnKeyDown(const FGeometry& MyGeometry, const FKeyEvent& InKeyEvent)
 {
-	if (InKeyboardEvent.GetKey() == EKeys::Escape)
+	if (InKeyEvent.GetKey() == EKeys::Escape)
 	{
 		//TODO: pop up window for in-place translation
 
diff --git a/Engine/Source/Editor/TranslationEditor/Private/TranslationPickerFloatingWindow.h b/Engine/Source/Editor/TranslationEditor/Private/TranslationPickerFloatingWindow.h
--- a/Engine/Source/Editor/TranslationEditor/Private/TranslationPickerFloatingWindow.h
+++ b/Engine/Source/Editor/TranslationEditor/Private/TranslationPickerFloatingWindow.h
@@ -33,7 +33,7 @@ private:
 		return FText::Format(LOCTEXT("TootipHint", "{0}\n\n(Esc to pick)"), TranslationInfoPreviewText);
 	}
 
-	virtual FReply OnKeyDown(const FGeometry& MyGeometry, const FKeyboardEvent& InKeyboardEvent) override;
+	virtual FReply OnKeyDown(const FGeometry& MyGeometry, const FKeyEvent& InKeyEvent) override;
 
 	/** We need to support keyboard focus to process the 'Esc' key */
 	virtual bool SupportsKeyboardFocus() const override
diff --git a/Engine/Source/Editor/UMGEditor/Private/Customizations/WidgetNavigationCustomization.cpp b/Engine/Source/Editor/UMGEditor/Private/Customizations/WidgetNavigationCustomization.cpp
--- a/Engine/Source/Editor/UMGEditor/Private/Customizations/WidgetNavigationCustomization.cpp
+++ b/Engine/Source/Editor/UMGEditor/Private/Customizations/WidgetNavigationCustomization.cpp
@@ -17,17 +17,135 @@
 
 void FWidgetNavigationCustomization::CustomizeHeader(TSharedRef<IPropertyHandle> PropertyHandle, FDetailWidgetRow& HeaderRow, IPropertyTypeCustomizationUtils& CustomizationUtils)
 {
-	//HeaderRow
-	//.WholeRowContent()
-	//[
-	//	SNew(SButton)
-	//	.HAlign(HAlign_Center)
-	//	.OnClicked(this, &FWidgetNavigationCustomization::OnCustomizeNavigation, TWeakPtr<IPropertyHandle>(PropertyHandle))
-	//	[
-	//		SNew(STextBlock)
-	//		.Text(LOCTEXT("CustomizeNavigation", "Customize Navigation"))
-	//	]
-	//];
+	// Disabled for now.
+	if (true) { return; }
+
+	TWeakPtr<IPropertyHandle> PropertyHandlePtr(PropertyHandle);
+
+	HeaderRow
+	.WholeRowContent()
+	[
+		SNew(SVerticalBox)
+
+		+ SVerticalBox::Slot()
+		.AutoHeight()
+		.Padding(2, 4, 2, 4)
+		[
+			SNew(SButton)
+			.HAlign(HAlign_Center)
+			.OnClicked(this, &FWidgetNavigationCustomization::OnCustomizeNavigation, PropertyHandlePtr)
+			[
+				SNew(STextBlock)
+				.Text(LOCTEXT("CustomizeNavigation", "Customize Navigation"))
+			]
+		]
+
+		+ SVerticalBox::Slot()
+		.AutoHeight()
+		.Padding(2, 4, 2, 4)
+		[
+			MakeNavRow(PropertyHandlePtr, EUINavigation::Left, LOCTEXT("LeftNavigation", "Left"))
+		]
+
+		+ SVerticalBox::Slot()
+		.AutoHeight()
+		.Padding(2, 4, 2, 4)
+		[
+			SNew(SHorizontalBox)
+
+			+ SHorizontalBox::Slot()
+			.AutoWidth()
+			.Padding(0)
+			[
+				SNew(STextBlock)
+				.Text(LOCTEXT("RightNavigation", "Right"))
+			]
+
+			+ SHorizontalBox::Slot()
+			.AutoWidth()
+			.Padding(0)
+			[
+				SNew(SComboButton)
+				.HAlign(HAlign_Center)
+				.ButtonContent()
+				[
+					SNew(STextBlock)
+					.Text(this, &FWidgetNavigationCustomization::HandleNavigationText, PropertyHandlePtr, EUINavigation::Right)
+				]
+				.ContentPadding(FMargin(4.0f, 2.0f))
+				.MenuContent()
+				[
+					MakeNavMenu(PropertyHandlePtr, EUINavigation::Right)
+				]
+			]
+		]
+
+		+ SVerticalBox::Slot()
+		.AutoHeight()
+		.Padding(2, 4, 2, 4)
+		[
+			SNew(SHorizontalBox)
+
+			+ SHorizontalBox::Slot()
+			.AutoWidth()
+			.Padding(0)
+			[
+				SNew(STextBlock)
+				.Text(LOCTEXT("UpNavigation", "Up"))
+			]
+
+			+ SHorizontalBox::Slot()
+			.AutoWidth()
+			.Padding(0)
+			[
+				SNew(SComboButton)
+				.HAlign(HAlign_Center)
+				.ButtonContent()
+				[
+					SNew(STextBlock)
+					.Text(this, &FWidgetNavigationCustomization::HandleNavigationText, PropertyHandlePtr, EUINavigation::Up)
+				]
+				.ContentPadding(FMargin(4.0f, 2.0f))
+				.MenuContent()
+				[
+					MakeNavMenu(PropertyHandlePtr, EUINavigation::Up)
+				]
+			]
+		]
+
+		+ SVerticalBox::Slot()
+		.AutoHeight()
+		.Padding(2, 4, 2, 4)
+		[
+			SNew(SHorizontalBox)
+
+			+ SHorizontalBox::Slot()
+			.AutoWidth()
+			.Padding(0)
+			[
+				SNew(STextBlock)
+				.Text(LOCTEXT("DownNavigation", "Down"))
+			]
+
+			+ SHorizontalBox::Slot()
+			.AutoWidth()
+			.Padding(0)
+			[
+				SNew(SComboButton)
+				.HAlign(HAlign_Center)
+				.ButtonContent()
+				[
+					SNew(STextBlock)
+					.Text(this, &FWidgetNavigationCustomization::HandleNavigationText, PropertyHandlePtr, EUINavigation::Down)
+				]
+				.ContentPadding(FMargin(4.0f, 2.0f))
+				.MenuContent()
+				[
+					MakeNavMenu(PropertyHandlePtr, EUINavigation::Down)
+				]
+			]
+		]
+	];
 }
 
 void FWidgetNavigationCustomization::CustomizeChildren(TSharedRef<IPropertyHandle> PropertyHandle, IDetailChildrenBuilder& ChildBuilder, IPropertyTypeCustomizationUtils& CustomizationUtils)
@@ -82,4 +200,152 @@ FReply FWidgetNavigationCustomization::OnCustomizeNavigation(TWeakPtr<IPropertyH
 	return FReply::Handled();
 }
 
+FText FWidgetNavigationCustomization::HandleNavigationText(TWeakPtr<IPropertyHandle> PropertyHandle, EUINavigation Nav) const
+{
+	TArray<UObject*> OuterObjects;
+	TSharedPtr<IPropertyHandle> PropertyHandlePtr = PropertyHandle.Pin();
+	PropertyHandlePtr->GetOuterObjects(OuterObjects);
+
+	EUINavigationRule Rule = EUINavigationRule::Invalid;
+	for (UObject* OuterObject : OuterObjects)
+	{
+		if (UWidget* Widget = Cast<UWidget>(OuterObject))
+		{
+			EUINavigationRule CurRule = EUINavigationRule::Escape;
+			UWidgetNavigation* WidgetNavigation = Widget->Navigation;
+			if (Widget->Navigation != nullptr)
+			{
+				CurRule = WidgetNavigation->GetNavigationRule(Nav);
+			}
+
+			if (Rule != EUINavigationRule::Invalid && CurRule != Rule)
+			{
+				return LOCTEXT("NavigationMultipleValues", "Multiple Values");
+			}
+			Rule = CurRule;
+		}
+	}
+
+	switch (Rule)
+	{
+	case EUINavigationRule::Escape:
+		return LOCTEXT("NavigationEscape", "Escape");
+	case EUINavigationRule::Stop:
+		return LOCTEXT("NavigationStop", "Stop");
+	case EUINavigationRule::Wrap:
+		return LOCTEXT("NavigationWrap", "Wrap");
+	}
+
+	return FText::GetEmpty();
+}
+
+TSharedRef<class SWidget> FWidgetNavigationCustomization::MakeNavRow(TWeakPtr<IPropertyHandle> PropertyHandle, EUINavigation Nav, FText NavName)
+{
+	TSharedRef<SWidget> Row =
+		SNew(SHorizontalBox)
+
+		+ SHorizontalBox::Slot()
+		.AutoWidth()
+		.Padding(0)
+		[
+			SNew(STextBlock)
+			.Text(NavName)
+		]
+
+	+ SHorizontalBox::Slot()
+		.AutoWidth()
+		.Padding(0)
+		[
+			SNew(SComboButton)
+			.HAlign(HAlign_Center)
+			.ButtonContent()
+			[
+				SNew(STextBlock)
+				.Text(this, &FWidgetNavigationCustomization::HandleNavigationText, PropertyHandle, Nav)
+			]
+			.ContentPadding(FMargin(4.0f, 2.0f))
+				.MenuContent()
+				[
+					MakeNavMenu(PropertyHandle, Nav)
+				]
+		];
+
+	return Row;
+}
+
+TSharedRef<class SWidget> FWidgetNavigationCustomization::MakeNavMenu(TWeakPtr<IPropertyHandle> PropertyHandle, EUINavigation Nav)
+{
+	// create build configurations menu
+	FMenuBuilder MenuBuilder(true, NULL);
+	{
+		FUIAction EscapeAction(FExecuteAction::CreateSP(this, &FWidgetNavigationCustomization::HandleNavMenuEntryClicked, PropertyHandle, Nav, EUINavigationRule::Escape));
+		MenuBuilder.AddMenuEntry(LOCTEXT("NavigationRuleEscape", "Escape"), LOCTEXT("NavigationRuleEscapeHint", "Navigation is allowed to escape the bounds of this widget."), FSlateIcon(), EscapeAction);
+
+		FUIAction StopAction(FExecuteAction::CreateSP(this, &FWidgetNavigationCustomization::HandleNavMenuEntryClicked, PropertyHandle, Nav, EUINavigationRule::Stop));
+		MenuBuilder.AddMenuEntry(LOCTEXT("NavigationRuleStop", "Stop"), LOCTEXT("NavigationRuleStopHint", "Navigation stops at the bounds of this widget."), FSlateIcon(), StopAction);
+
+		FUIAction WrapAction(FExecuteAction::CreateSP(this, &FWidgetNavigationCustomization::HandleNavMenuEntryClicked, PropertyHandle, Nav, EUINavigationRule::Wrap));
+		MenuBuilder.AddMenuEntry(LOCTEXT("NavigationRuleWrap", "Wrap"), LOCTEXT("NavigationRuleWrapHint", "Navigation will wrap to the opposite bound of this object."), FSlateIcon(), WrapAction);
+
+		//FUIAction ExplicitAction(FExecuteAction::CreateSP(this, &FWidgetNavigationCustomization::HandleNavMenuEntryClicked, PropertyHandle, Nav, EUINavigationRule::Explicit));
+		//MenuBuilder.AddMenuEntry(LOCTEXT("NavigationRuleExplicit", "Explicit"), LOCTEXT("NavigationRuleExplicitHint", "Navigation will go to a specified widget."), FSlateIcon(), ExplicitAction);
+
+		//FUIAction CustomAction(FExecuteAction::CreateSP(this, &FWidgetNavigationCustomization::HandleNavMenuEntryClicked, PropertyHandle, Nav, EUINavigationRule::Custom));
+		//MenuBuilder.AddMenuEntry(LOCTEXT("NavigationRuleCustom", "Custom"), LOCTEXT("NavigationRuleCustomHint", "Custom function can determine what widget is navigated to."), FSlateIcon(), CustomAction);
+	}
+
+	return MenuBuilder.MakeWidget();
+}
+
+// Callback for clicking a menu entry for a navigations rule.
+void FWidgetNavigationCustomization::HandleNavMenuEntryClicked(TWeakPtr<IPropertyHandle> PropertyHandle, EUINavigation Nav, EUINavigationRule Rule)
+{
+	TArray<UObject*> OuterObjects;
+	TSharedPtr<IPropertyHandle> PropertyHandlePtr = PropertyHandle.Pin();
+	PropertyHandlePtr->GetOuterObjects(OuterObjects);
+
+	const FScopedTransaction Transaction(LOCTEXT("InitializeNavigation", "Initialize Navigation"));
+
+	for (UObject* OuterObject : OuterObjects)
+	{
+		if (UWidget* Widget = Cast<UWidget>(OuterObject))
+		{
+			FWidgetReference WidgetReference = Editor.Pin()->GetReferenceFromPreview(Widget);
+
+			SetNav(WidgetReference.GetPreview(), Nav, Rule);
+			SetNav(WidgetReference.GetTemplate(), Nav, Rule);
+		}
+	}
+}
+
+void FWidgetNavigationCustomization::SetNav(UWidget* Widget, EUINavigation Nav, EUINavigationRule Rule)
+{
+	Widget->Modify();
+
+	UWidgetNavigation* WidgetNavigation = Widget->Navigation;
+	if (!Widget->Navigation)
+	{
+		WidgetNavigation = ConstructObject<UWidgetNavigation>(UWidgetNavigation::StaticClass(), Widget);
+	}
+
+	//@Todo: Clear out the navigation if all defaults.
+	switch (Nav)
+	{
+	case EUINavigation::Left:
+		WidgetNavigation->Left.Rule = Rule;
+		break;
+	case EUINavigation::Right:
+		WidgetNavigation->Right.Rule = Rule;
+		break;
+	case EUINavigation::Up:
+		WidgetNavigation->Up.Rule = Rule;
+		break;
+	case EUINavigation::Down:
+		WidgetNavigation->Down.Rule = Rule;
+		break;
+	}
+
+	Widget->Navigation = WidgetNavigation;
+}
+
 #undef LOCTEXT_NAMESPACE
diff --git a/Engine/Source/Editor/UMGEditor/Private/Customizations/WidgetNavigationCustomization.h b/Engine/Source/Editor/UMGEditor/Private/Customizations/WidgetNavigationCustomization.h
--- a/Engine/Source/Editor/UMGEditor/Private/Customizations/WidgetNavigationCustomization.h
+++ b/Engine/Source/Editor/UMGEditor/Private/Customizations/WidgetNavigationCustomization.h
@@ -6,13 +6,13 @@ class FWidgetNavigationCustomization : public IPropertyTypeCustomization
 {
 public:
 	/** Makes a new instance of this detail layout class for a specific detail view requesting it */
-	static TSharedRef<class IPropertyTypeCustomization> MakeInstance(UBlueprint* Blueprint)
+	static TSharedRef<class IPropertyTypeCustomization> MakeInstance(TSharedRef<FWidgetBlueprintEditor> InEditor)
 	{
-		return MakeShareable(new FWidgetNavigationCustomization(Blueprint));
+		return MakeShareable(new FWidgetNavigationCustomization(InEditor));
 	}
 
-	FWidgetNavigationCustomization(UBlueprint* InBlueprint)
-		: Blueprint(Cast<UWidgetBlueprint>(InBlueprint))
+	FWidgetNavigationCustomization(TSharedRef<FWidgetBlueprintEditor> InEditor)
+		: Editor(InEditor)
 	{
 	}
 	
@@ -21,10 +21,21 @@ public:
 	virtual void CustomizeChildren(TSharedRef<IPropertyHandle> PropertyHandle, class IDetailChildrenBuilder& ChildBuilder, IPropertyTypeCustomizationUtils& CustomizationUtils) override;
 	
 private:
+
 	void FillOutChildren(TSharedRef<IPropertyHandle> PropertyHandle, class IDetailChildrenBuilder& ChildBuilder, IPropertyTypeCustomizationUtils& CustomizationUtils);
 
 	FReply OnCustomizeNavigation(TWeakPtr<IPropertyHandle> PropertyHandle);
+	
+	FText HandleNavigationText(TWeakPtr<IPropertyHandle> PropertyHandle, EUINavigation Nav) const;
+
+	TSharedRef<class SWidget> MakeNavRow(TWeakPtr<IPropertyHandle> PropertyHandle, EUINavigation Nav, FText NavName);
+
+	TSharedRef<class SWidget> MakeNavMenu(TWeakPtr<IPropertyHandle> PropertyHandle, EUINavigation Nav);
+
+	void HandleNavMenuEntryClicked(TWeakPtr<IPropertyHandle> PropertyHandle, EUINavigation Nav, EUINavigationRule Rule);
+
+	void SetNav(UWidget* Widget, EUINavigation Nav, EUINavigationRule Rule);
 
 private:
-	UWidgetBlueprint* Blueprint;
+	TWeakPtr<FWidgetBlueprintEditor> Editor;
 };
diff --git a/Engine/Source/Editor/UMGEditor/Private/Designer/SDesignerView.cpp b/Engine/Source/Editor/UMGEditor/Private/Designer/SDesignerView.cpp
--- a/Engine/Source/Editor/UMGEditor/Private/Designer/SDesignerView.cpp
+++ b/Engine/Source/Editor/UMGEditor/Private/Designer/SDesignerView.cpp
@@ -929,7 +929,7 @@ FReply SDesignerView::OnMouseButtonDown(const FGeometry& MyGeometry, const FPoin
 	}
 
 	// Capture mouse for the drag handle and general mouse actions
-	return FReply::Handled().PreventThrottling().SetKeyboardFocus(AsShared(), EKeyboardFocusCause::Mouse).CaptureMouse(AsShared());
+	return FReply::Handled().PreventThrottling().SetUserFocus(AsShared(), EFocusCause::Mouse).CaptureMouse(AsShared());
 }
 
 FReply SDesignerView::OnMouseButtonUp(const FGeometry& MyGeometry, const FPointerEvent& MouseEvent)
@@ -1036,16 +1036,16 @@ void SDesignerView::OnMouseLeave(const FPointerEvent& MouseEvent)
 	BlueprintEditor.Pin()->ClearHoveredWidget();
 }
 
-FReply SDesignerView::OnKeyDown(const FGeometry& MyGeometry, const FKeyboardEvent& InKeyboardEvent)
+FReply SDesignerView::OnKeyDown(const FGeometry& MyGeometry, const FKeyEvent& InKeyEvent)
 {
 	BlueprintEditor.Pin()->PasteDropLocation = FVector2D(0, 0);
 
-	if ( BlueprintEditor.Pin()->DesignerCommandList->ProcessCommandBindings(InKeyboardEvent) )
+	if ( BlueprintEditor.Pin()->DesignerCommandList->ProcessCommandBindings(InKeyEvent) )
 	{
 		return FReply::Handled();
 	}
 
-	if ( CommandList->ProcessCommandBindings(InKeyboardEvent) )
+	if ( CommandList->ProcessCommandBindings(InKeyEvent) )
 	{
 		return FReply::Handled();
 	}
diff --git a/Engine/Source/Editor/UMGEditor/Private/Details/SWidgetDetailsView.cpp b/Engine/Source/Editor/UMGEditor/Private/Details/SWidgetDetailsView.cpp
--- a/Engine/Source/Editor/UMGEditor/Private/Details/SWidgetDetailsView.cpp
+++ b/Engine/Source/Editor/UMGEditor/Private/Details/SWidgetDetailsView.cpp
@@ -180,7 +180,7 @@ void SWidgetDetailsView::RegisterCustomizations()
 	static FName PropertyEditor("PropertyEditor");
 	FPropertyEditorModule& PropertyModule = FModuleManager::GetModuleChecked<FPropertyEditorModule>(PropertyEditor);
 
-	PropertyModule.RegisterCustomPropertyTypeLayout(TEXT("WidgetNavigation"), FOnGetPropertyTypeCustomizationInstance::CreateStatic(&FWidgetNavigationCustomization::MakeInstance, BlueprintEditor.Pin()->GetBlueprintObj()), nullptr, PropertyView);
+	PropertyModule.RegisterCustomPropertyTypeLayout(TEXT("WidgetNavigation"), FOnGetPropertyTypeCustomizationInstance::CreateStatic(&FWidgetNavigationCustomization::MakeInstance, BlueprintEditor.Pin().ToSharedRef()), nullptr, PropertyView);
 	PropertyModule.RegisterCustomPropertyTypeLayout(TEXT("PanelSlot"), FOnGetPropertyTypeCustomizationInstance::CreateStatic(&FCanvasSlotCustomization::MakeInstance, BlueprintEditor.Pin()->GetBlueprintObj()), nullptr, PropertyView);
 	PropertyModule.RegisterCustomPropertyTypeLayout(TEXT("EHorizontalAlignment"), FOnGetPropertyTypeCustomizationInstance::CreateStatic(&FHorizontalAlignmentCustomization::MakeInstance), nullptr, PropertyView);
 	PropertyModule.RegisterCustomPropertyTypeLayout(TEXT("EVerticalAlignment"), FOnGetPropertyTypeCustomizationInstance::CreateStatic(&FVerticalAlignmentCustomization::MakeInstance), nullptr, PropertyView);
@@ -241,7 +241,7 @@ void SWidgetDetailsView::ClearFocusIfOwned()
 		// if this occurs we need need to immediately clear keyboard focus
 		if ( FSlateApplication::Get().HasFocusedDescendants(AsShared()) )
 		{
-			FSlateApplication::Get().ClearKeyboardFocus(EKeyboardFocusCause::Mouse);
+			FSlateApplication::Get().ClearKeyboardFocus(EFocusCause::Mouse);
 		}
 		bIsReentrant = false;
 	}
diff --git a/Engine/Source/Editor/UMGEditor/Private/Hierarchy/SHierarchyView.cpp b/Engine/Source/Editor/UMGEditor/Private/Hierarchy/SHierarchyView.cpp
--- a/Engine/Source/Editor/UMGEditor/Private/Hierarchy/SHierarchyView.cpp
+++ b/Engine/Source/Editor/UMGEditor/Private/Hierarchy/SHierarchyView.cpp
@@ -129,11 +129,11 @@ void SHierarchyView::OnMouseLeave(const FPointerEvent& MouseEvent)
 	BlueprintEditor.Pin()->ClearHoveredWidget();
 }
 
-FReply SHierarchyView::OnKeyDown(const FGeometry& MyGeometry, const FKeyboardEvent& InKeyboardEvent)
+FReply SHierarchyView::OnKeyDown(const FGeometry& MyGeometry, const FKeyEvent& InKeyEvent)
 {
 	BlueprintEditor.Pin()->PasteDropLocation = FVector2D(0, 0);
 
-	if ( BlueprintEditor.Pin()->DesignerCommandList->ProcessCommandBindings(InKeyboardEvent) )
+	if ( BlueprintEditor.Pin()->DesignerCommandList->ProcessCommandBindings(InKeyEvent) )
 	{
 		return FReply::Handled();
 	}
diff --git a/Engine/Source/Editor/UMGEditor/Private/TabFactory/AnimationTabSummoner.cpp b/Engine/Source/Editor/UMGEditor/Private/TabFactory/AnimationTabSummoner.cpp
--- a/Engine/Source/Editor/UMGEditor/Private/TabFactory/AnimationTabSummoner.cpp
+++ b/Engine/Source/Editor/UMGEditor/Private/TabFactory/AnimationTabSummoner.cpp
@@ -240,10 +240,10 @@ public:
 		CreateCommandList();
 	}
 
-	virtual FReply OnKeyDown( const FGeometry& MyGeometry, const FKeyboardEvent& InKeyboardEvent )
+	virtual FReply OnKeyDown( const FGeometry& MyGeometry, const FKeyEvent& InKeyEvent )
 	{
 		FReply Reply = FReply::Unhandled();
-		if( CommandList->ProcessCommandBindings( InKeyboardEvent ) )
+		if( CommandList->ProcessCommandBindings( InKeyEvent ) )
 		{
 			Reply = FReply::Handled();
 		}
diff --git a/Engine/Source/Editor/UnrealEd/Private/ConsolidateWindow.cpp b/Engine/Source/Editor/UnrealEd/Private/ConsolidateWindow.cpp
--- a/Engine/Source/Editor/UnrealEd/Private/ConsolidateWindow.cpp
+++ b/Engine/Source/Editor/UnrealEd/Private/ConsolidateWindow.cpp
@@ -208,7 +208,7 @@ private:
 	// Input Responses
 
 	/** Called in response to the user releasing a keyboard key while the consolidation panel has keyboard focus */
-	FReply OnKeyUp( const FGeometry& MyGeometry, const FKeyboardEvent& InKeyboardEvent ) override;
+	FReply OnKeyUp( const FGeometry& MyGeometry, const FKeyEvent& InKeyEvent ) override;
 
 	/** Track if the panel has already warned the user about consolidating assets with different types, so as not to repeatedly (and annoyingly) warn */
 	bool bAlreadyWarnedAboutTypes;
@@ -890,9 +890,9 @@ FReply SConsolidateToolWidget::OnDragOver( const FGeometry& MyGeometry, const FD
 }
 
 /** Called in response to the user releasing a keyboard key while the consolidation panel has keyboard focus */
-FReply SConsolidateToolWidget::OnKeyUp( const FGeometry& MyGeometry, const FKeyboardEvent& InKeyboardEvent )
+FReply SConsolidateToolWidget::OnKeyUp( const FGeometry& MyGeometry, const FKeyEvent& InKeyEvent )
 {
-	const FKey Key = InKeyboardEvent.GetKey();
+	const FKey Key = InKeyEvent.GetKey();
 	if ( Key == EKeys::Platform_Delete )
 	{
 		RemoveSelectedObject();
diff --git a/Engine/Source/Editor/UnrealEd/Private/Dialogs/Dialogs.cpp b/Engine/Source/Editor/UnrealEd/Private/Dialogs/Dialogs.cpp
--- a/Engine/Source/Editor/UnrealEd/Private/Dialogs/Dialogs.cpp
+++ b/Engine/Source/Editor/UnrealEd/Private/Dialogs/Dialogs.cpp
@@ -166,15 +166,15 @@ public:
 		return Response;
 	}
 
-	virtual	FReply OnKeyDown( const FGeometry& MyGeometry, const FKeyboardEvent& InKeyboardEvent )
+	virtual	FReply OnKeyDown( const FGeometry& MyGeometry, const FKeyEvent& InKeyEvent )
 	{
 		//see if we pressed the Enter or Spacebar keys
-		if( InKeyboardEvent.GetKey() == EKeys::Escape )
+		if( InKeyEvent.GetKey() == EKeys::Escape )
 		{
 			return HandleButtonClicked(EAppReturnType::Cancel);
 		}
 
-		if (InKeyboardEvent.GetKey() == EKeys::C && InKeyboardEvent.IsControlDown())
+		if (InKeyEvent.GetKey() == EKeys::C && InKeyEvent.IsControlDown())
 		{
 			CopyMessageToClipboard();
 
@@ -370,9 +370,9 @@ public:
 	bool GetResponse() const { return bUserResponse; }
 	virtual bool SupportsKeyboardFocus() const override { return true; }
 
-	virtual FReply OnKeyDown( const FGeometry& MyGeometry, const FKeyboardEvent& InKeyboardEvent ) override
+	virtual FReply OnKeyDown( const FGeometry& MyGeometry, const FKeyEvent& InKeyEvent ) override
 	{
-		if (InKeyboardEvent.GetKey() == EKeys::C && InKeyboardEvent.IsControlDown())
+		if (InKeyEvent.GetKey() == EKeys::C && InKeyEvent.IsControlDown())
 		{
 			FPlatformMisc::ClipboardCopy( *MyMessage.Get() );
 			return FReply::Handled();
@@ -586,15 +586,15 @@ public:
 	}
 
 	/** Used to intercept Escape key presses, then interprets them as cancel */
-	virtual FReply OnKeyDown( const FGeometry& MyGeometry, const FKeyboardEvent& InKeyboardEvent )
+	virtual FReply OnKeyDown( const FGeometry& MyGeometry, const FKeyEvent& InKeyEvent )
 	{
 		// Pressing escape returns as if the user canceled
-		if ( InKeyboardEvent.GetKey() == EKeys::Escape )
+		if ( InKeyEvent.GetKey() == EKeys::Escape )
 		{
 			return OnCancelClicked();
 		}
 
-		if (InKeyboardEvent.GetKey() == EKeys::C && InKeyboardEvent.IsControlDown())
+		if (InKeyEvent.GetKey() == EKeys::C && InKeyEvent.IsControlDown())
 		{
 			FPlatformMisc::ClipboardCopy( *MyMessage.Get().ToString() );
 			return FReply::Handled();
diff --git a/Engine/Source/Editor/UnrealEd/Private/Dialogs/DlgMoveAssets.cpp b/Engine/Source/Editor/UnrealEd/Private/Dialogs/DlgMoveAssets.cpp
--- a/Engine/Source/Editor/UnrealEd/Private/Dialogs/DlgMoveAssets.cpp
+++ b/Engine/Source/Editor/UnrealEd/Private/Dialogs/DlgMoveAssets.cpp
@@ -213,10 +213,10 @@ public:
 	}
 
 	/** Used to intercept Escape key presses, then interprets them as cancel */
-	virtual FReply OnKeyDown( const FGeometry& MyGeometry, const FKeyboardEvent& InKeyboardEvent )
+	virtual FReply OnKeyDown( const FGeometry& MyGeometry, const FKeyEvent& InKeyEvent )
 	{
 		// Pressing escape returns as if the user canceled
-		if ( InKeyboardEvent.GetKey() == EKeys::Escape )
+		if ( InKeyEvent.GetKey() == EKeys::Escape )
 		{
 			return OnButtonClick(FDlgMoveAsset::Cancel);
 		}
diff --git a/Engine/Source/Editor/UnrealEd/Private/Dialogs/SDeleteAssetsDialog.cpp b/Engine/Source/Editor/UnrealEd/Private/Dialogs/SDeleteAssetsDialog.cpp
--- a/Engine/Source/Editor/UnrealEd/Private/Dialogs/SDeleteAssetsDialog.cpp
+++ b/Engine/Source/Editor/UnrealEd/Private/Dialogs/SDeleteAssetsDialog.cpp
@@ -492,15 +492,15 @@ TSharedRef<SWidget> SDeleteAssetsDialog::BuildForceDeleteWidget()
 	];
 }
 
-FReply SDeleteAssetsDialog::OnKeyDown( const FGeometry& MyGeometry, const FKeyboardEvent& InKeyboardEvent )
+FReply SDeleteAssetsDialog::OnKeyDown( const FGeometry& MyGeometry, const FKeyEvent& InKeyEvent )
 {
-	if( InKeyboardEvent.GetKey() == EKeys::Escape )
+	if( InKeyEvent.GetKey() == EKeys::Escape )
 	{
 		ParentWindow.Get()->RequestDestroyWindow();
 		return FReply::Handled();
 	}
 
-	if ( ReferencerCommands->ProcessCommandBindings(InKeyboardEvent) )
+	if ( ReferencerCommands->ProcessCommandBindings(InKeyEvent) )
 	{
 		return FReply::Handled();
 	}
diff --git a/Engine/Source/Editor/UnrealEd/Private/EditorViewportClient.cpp b/Engine/Source/Editor/UnrealEd/Private/EditorViewportClient.cpp
--- a/Engine/Source/Editor/UnrealEd/Private/EditorViewportClient.cpp
+++ b/Engine/Source/Editor/UnrealEd/Private/EditorViewportClient.cpp
@@ -722,7 +722,7 @@ void FEditorViewportClient::ReceivedFocus(FViewport* InViewport)
 	if( IsMatineeRecordingWindow() )
 	{
 		// Allow the joystick to be used for matinee capture
-		InViewport->CaptureJoystickInput( true );
+		InViewport->SetUserFocus( true );
 	}
 }
 
diff --git a/Engine/Source/Editor/UnrealEd/Private/Factories/AnimBlueprintFactory.cpp b/Engine/Source/Editor/UnrealEd/Private/Factories/AnimBlueprintFactory.cpp
--- a/Engine/Source/Editor/UnrealEd/Private/Factories/AnimBlueprintFactory.cpp
+++ b/Engine/Source/Editor/UnrealEd/Private/Factories/AnimBlueprintFactory.cpp
@@ -298,14 +298,14 @@ private:
 		return FReply::Handled();
 	}
 
-	FReply OnKeyDown( const FGeometry& MyGeometry, const FKeyboardEvent& InKeyboardEvent )
+	FReply OnKeyDown( const FGeometry& MyGeometry, const FKeyEvent& InKeyEvent )
 	{
-		if (InKeyboardEvent.GetKey() == EKeys::Escape)
+		if (InKeyEvent.GetKey() == EKeys::Escape)
 		{
 			CloseDialog();
 			return FReply::Handled();
 		}
-		return SWidget::OnKeyDown(MyGeometry, InKeyboardEvent);
+		return SWidget::OnKeyDown(MyGeometry, InKeyEvent);
 	}
 
 private:
diff --git a/Engine/Source/Editor/UnrealEd/Private/FbxOptionWindow.h b/Engine/Source/Editor/UnrealEd/Private/FbxOptionWindow.h
--- a/Engine/Source/Editor/UnrealEd/Private/FbxOptionWindow.h
+++ b/Engine/Source/Editor/UnrealEd/Private/FbxOptionWindow.h
@@ -54,9 +54,9 @@ public:
 		return FReply::Handled();
 	}
 
-	virtual FReply OnKeyDown( const FGeometry& MyGeometry, const FKeyboardEvent& InKeyboardEvent ) override
+	virtual FReply OnKeyDown( const FGeometry& MyGeometry, const FKeyEvent& InKeyEvent ) override
 	{
-		if( InKeyboardEvent.GetKey() == EKeys::Escape )
+		if( InKeyEvent.GetKey() == EKeys::Escape )
 		{
 			return OnCancel();
 		}
diff --git a/Engine/Source/Editor/UnrealEd/Private/PlayLevel.cpp b/Engine/Source/Editor/UnrealEd/Private/PlayLevel.cpp
--- a/Engine/Source/Editor/UnrealEd/Private/PlayLevel.cpp
+++ b/Engine/Source/Editor/UnrealEd/Private/PlayLevel.cpp
@@ -3227,7 +3227,7 @@ void UEditorEngine::FocusNextPIEWorld(UWorld *CurrentPieWorld, bool previous)
 			SlateInfoPtr->SlatePlayInEditorWindow.Pin()->BringToFront();
 
 			// Set viewport widget to have keyboard focus
-			FSlateApplication::Get().SetKeyboardFocus(SlateInfoPtr->SlatePlayInEditorWindowViewport->GetViewportWidget().Pin(), EKeyboardFocusCause::Keyboard);
+			FSlateApplication::Get().SetKeyboardFocus(SlateInfoPtr->SlatePlayInEditorWindowViewport->GetViewportWidget().Pin(), EFocusCause::Navigation);
 
 			// Execute notifcation delegate incase game code has to do anything else
 			PIEInstanceWindowSwitchDelegate.ExecuteIfBound();
diff --git a/Engine/Source/Editor/UnrealEd/Private/SClassPickerDialog.cpp b/Engine/Source/Editor/UnrealEd/Private/SClassPickerDialog.cpp
--- a/Engine/Source/Editor/UnrealEd/Private/SClassPickerDialog.cpp
+++ b/Engine/Source/Editor/UnrealEd/Private/SClassPickerDialog.cpp
@@ -289,17 +289,17 @@ EVisibility SClassPickerDialog::GetSelectButtonVisibility() const
 }
 
 /** Overridden from SWidget: Called when a key is pressed down - capturing copy */
-FReply SClassPickerDialog::OnKeyDown( const FGeometry& MyGeometry, const FKeyboardEvent& InKeyboardEvent )
+FReply SClassPickerDialog::OnKeyDown( const FGeometry& MyGeometry, const FKeyEvent& InKeyEvent )
 {
 	WeakParentWindow.Pin().Get()->SetWidgetToFocusOnActivate(ClassViewer);
 
-	if (InKeyboardEvent.GetKey() == EKeys::Escape)
+	if (InKeyEvent.GetKey() == EKeys::Escape)
 	{
 		return OnClassPickerCanceled();
 	}
 	else
 	{
-		return ClassViewer->OnKeyDown(MyGeometry, InKeyboardEvent);
+		return ClassViewer->OnKeyDown(MyGeometry, InKeyEvent);
 	}
 	return FReply::Handled();
 }
diff --git a/Engine/Source/Editor/UnrealEd/Private/SColorGradientEditor.cpp b/Engine/Source/Editor/UnrealEd/Private/SColorGradientEditor.cpp
--- a/Engine/Source/Editor/UnrealEd/Private/SColorGradientEditor.cpp
+++ b/Engine/Source/Editor/UnrealEd/Private/SColorGradientEditor.cpp
@@ -426,9 +426,9 @@ FReply SColorGradientEditor::OnMouseButtonUp( const FGeometry& MyGeometry, const
 	return FReply::Unhandled();
 }
 
-FReply SColorGradientEditor::OnKeyDown( const FGeometry& MyGeometry, const FKeyboardEvent& InKeyboardEvent )
+FReply SColorGradientEditor::OnKeyDown( const FGeometry& MyGeometry, const FKeyEvent& InKeyEvent )
 {
-	if ( IsEditingEnabled.Get() == true && InKeyboardEvent.GetKey() == EKeys::Platform_Delete )
+	if ( IsEditingEnabled.Get() == true && InKeyEvent.GetKey() == EKeys::Platform_Delete )
 	{
 		DeleteStop( SelectedStop );
 		return FReply::Handled();
diff --git a/Engine/Source/Editor/UnrealEd/Private/SCurveEditor.cpp b/Engine/Source/Editor/UnrealEd/Private/SCurveEditor.cpp
--- a/Engine/Source/Editor/UnrealEd/Private/SCurveEditor.cpp
+++ b/Engine/Source/Editor/UnrealEd/Private/SCurveEditor.cpp
@@ -1071,7 +1071,7 @@ FReply SCurveEditor::OnMouseButtonDown( const FGeometry& InMyGeometry, const FPo
 		// Set keyboard focus to this so that selected text box doesn't try to apply to newly selected keys
 		if(!HasKeyboardFocus())
 		{
-			FSlateApplication::Get().SetKeyboardFocus(SharedThis(this), EKeyboardFocusCause::SetDirectly);
+			FSlateApplication::Get().SetKeyboardFocus(SharedThis(this), EFocusCause::SetDirectly);
 		}
 
 		// Always capture mouse if we left or right click on the widget
@@ -1212,16 +1212,16 @@ FReply SCurveEditor::OnMouseWheel( const FGeometry& MyGeometry, const FPointerEv
 	return FReply::Handled();
 }
 
-FReply SCurveEditor::OnKeyDown( const FGeometry& MyGeometry, const FKeyboardEvent& InKeyboardEvent )
+FReply SCurveEditor::OnKeyDown( const FGeometry& MyGeometry, const FKeyEvent& InKeyEvent )
 {
-	if (InKeyboardEvent.GetKey() == EKeys::Platform_Delete && SelectedKeys.Num() != 0)
+	if (InKeyEvent.GetKey() == EKeys::Platform_Delete && SelectedKeys.Num() != 0)
 	{
 		DeleteSelectedKeys();
 		return FReply::Handled();
 	}
 	else
 	{
-		if( Commands->ProcessCommandBindings( InKeyboardEvent ) )
+		if( Commands->ProcessCommandBindings( InKeyEvent ) )
 		{
 			return FReply::Handled();
 		}
diff --git a/Engine/Source/Editor/UnrealEd/Private/SEditorViewport.cpp b/Engine/Source/Editor/UnrealEd/Private/SEditorViewport.cpp
--- a/Engine/Source/Editor/UnrealEd/Private/SEditorViewport.cpp
+++ b/Engine/Source/Editor/UnrealEd/Private/SEditorViewport.cpp
@@ -78,10 +78,10 @@ void SEditorViewport::Construct( const FArguments& InArgs )
 	}
 }
 
-FReply SEditorViewport::OnKeyDown( const FGeometry& MyGeometry, const FKeyboardEvent& InKeyboardEvent )
+FReply SEditorViewport::OnKeyDown( const FGeometry& MyGeometry, const FKeyEvent& InKeyEvent )
 {
 	FReply Reply = FReply::Unhandled();
-	if( CommandList->ProcessCommandBindings( InKeyboardEvent ) )
+	if( CommandList->ProcessCommandBindings( InKeyEvent ) )
 	{
 		Reply = FReply::Handled();
 		Client->Invalidate();
@@ -96,10 +96,10 @@ bool SEditorViewport::SupportsKeyboardFocus() const
 	return true;
 }
 
-FReply SEditorViewport::OnKeyboardFocusReceived( const FGeometry& MyGeometry, const FKeyboardFocusEvent& InKeyboardFocusEvent )
+FReply SEditorViewport::OnFocusReceived( const FGeometry& MyGeometry, const FFocusEvent& InFocusEvent )
 {
 	// forward focus to the viewport
-	return FReply::Handled().SetKeyboardFocus( ViewportWidget.ToSharedRef(), InKeyboardFocusEvent.GetCause() );
+	return FReply::Handled().SetUserFocus(ViewportWidget.ToSharedRef(), InFocusEvent.GetCause());
 }
 
 void SEditorViewport::Tick( const FGeometry& AllottedGeometry, const double InCurrentTime, const float InDeltaTime )
diff --git a/Engine/Source/Editor/UnrealEd/Private/Tests/EditorBuildPromotionTests.cpp b/Engine/Source/Editor/UnrealEd/Private/Tests/EditorBuildPromotionTests.cpp
--- a/Engine/Source/Editor/UnrealEd/Private/Tests/EditorBuildPromotionTests.cpp
+++ b/Engine/Source/Editor/UnrealEd/Private/Tests/EditorBuildPromotionTests.cpp
@@ -355,11 +355,11 @@ namespace EditorBuildPromotionTestUtils
 		TSharedPtr<SWidget> FocusWidget = FindFirstWidgetByClass(EditorWindow, WidgetTypeToFocus);
 		if (FocusWidget.IsValid())
 		{
-			FSlateApplication::Get().SetKeyboardFocus(FocusWidget.ToSharedRef(), EKeyboardFocusCause::SetDirectly);
+			FSlateApplication::Get().SetKeyboardFocus(FocusWidget.ToSharedRef(), EFocusCause::SetDirectly);
 
 			//Send the command
 			FModifierKeysState ModifierKeys(InGesture.bShift, false, InGesture.bCtrl, false, InGesture.bAlt, false, InGesture.bCmd, false, false);
-			FKeyboardEvent KeyEvent(InGesture.Key, ModifierKeys, false, 0, 0);
+			FKeyEvent KeyEvent(InGesture.Key, ModifierKeys, 0/*UserIndex*/, false, 0, 0);
 			FSlateApplication::Get().ProcessKeyDownEvent(KeyEvent);
 			FSlateApplication::Get().ProcessKeyUpEvent(KeyEvent);
 		}
@@ -4368,7 +4368,7 @@ bool FBuildPromotionSettingsTest::RunTest(const FString& Parameters)
 	FSlateApplication::Get().ProcessWindowActivatedEvent(FWindowActivateEvent(FWindowActivateEvent::EA_Activate, AllWindows[0]));
 
 	//Send the PIE event
-	FKeyboardEvent PIEKeyEvent(EKeys::L, FModifierKeysState(false, false, true, false, false, false, false, false, false), false, 0x4C, 0x4C);
+	FKeyEvent PIEKeyEvent(EKeys::L, FModifierKeysState(false, false, true, false, false, false, false, false, false), false, 0/*UserIndex*/, 0x4C, 0x4C);
 	FSlateApplication::Get().ProcessKeyDownEvent(PIEKeyEvent);
 	FSlateApplication::Get().ProcessKeyUpEvent(PIEKeyEvent);
 
diff --git a/Engine/Source/Editor/UnrealEd/Private/Toolkits/BaseToolkit.cpp b/Engine/Source/Editor/UnrealEd/Private/Toolkits/BaseToolkit.cpp
--- a/Engine/Source/Editor/UnrealEd/Private/Toolkits/BaseToolkit.cpp
+++ b/Engine/Source/Editor/UnrealEd/Private/Toolkits/BaseToolkit.cpp
@@ -72,9 +72,9 @@ FName FBaseToolkit::GetToolkitContextFName() const
 }
 
 
-bool FBaseToolkit::ProcessCommandBindings( const FKeyboardEvent& InKeyboardEvent ) const
+bool FBaseToolkit::ProcessCommandBindings( const FKeyEvent& InKeyEvent ) const
 {
-	if( ToolkitCommands->ProcessCommandBindings( InKeyboardEvent ) )
+	if( ToolkitCommands->ProcessCommandBindings( InKeyEvent ) )
 	{
 		return true;
 	}
diff --git a/Engine/Source/Editor/UnrealEd/Private/Toolkits/GlobalEditorCommonCommands.cpp b/Engine/Source/Editor/UnrealEd/Private/Toolkits/GlobalEditorCommonCommands.cpp
--- a/Engine/Source/Editor/UnrealEd/Private/Toolkits/GlobalEditorCommonCommands.cpp
+++ b/Engine/Source/Editor/UnrealEd/Private/Toolkits/GlobalEditorCommonCommands.cpp
@@ -49,7 +49,7 @@ public:
 	}
 
 	// SWidget interface
-	virtual FReply OnPreviewKeyDown(const FGeometry& MyGeometry, const FKeyboardEvent& InKeyboardEvent) override;
+	virtual FReply OnPreviewKeyDown(const FGeometry& MyGeometry, const FKeyEvent& InKeyEvent) override;
 	// End of SWidget interface
 
 protected:
@@ -73,9 +73,9 @@ protected:
 	}
 };
 
-FReply SGlobalOpenAssetDialog::OnPreviewKeyDown(const FGeometry& MyGeometry, const FKeyboardEvent& InKeyboardEvent)
+FReply SGlobalOpenAssetDialog::OnPreviewKeyDown(const FGeometry& MyGeometry, const FKeyEvent& InKeyEvent)
 {
-	if (InKeyboardEvent.GetKey() == EKeys::Escape)
+	if (InKeyEvent.GetKey() == EKeys::Escape)
 	{
 		FSlateApplication::Get().DismissAllMenus();
 		return FReply::Handled();
diff --git a/Engine/Source/Editor/UnrealEd/Private/Toolkits/SStandaloneAssetEditorToolkitHost.cpp b/Engine/Source/Editor/UnrealEd/Private/Toolkits/SStandaloneAssetEditorToolkitHost.cpp
--- a/Engine/Source/Editor/UnrealEd/Private/Toolkits/SStandaloneAssetEditorToolkitHost.cpp
+++ b/Engine/Source/Editor/UnrealEd/Private/Toolkits/SStandaloneAssetEditorToolkitHost.cpp
@@ -265,16 +265,16 @@ UWorld* SStandaloneAssetEditorToolkitHost::GetWorld() const
 }
 
 
-FReply SStandaloneAssetEditorToolkitHost::OnKeyDown( const FGeometry& MyGeometry, const FKeyboardEvent& InKeyboardEvent )
+FReply SStandaloneAssetEditorToolkitHost::OnKeyDown( const FGeometry& MyGeometry, const FKeyEvent& InKeyEvent )
 {
-	// Check to see if any of the actions for the level editor can be processed by the current keyboard event
+	// Check to see if any of the actions for the level editor can be processed by the current event
 	// If we are in debug mode do not process commands
-	if( HostedToolkit.IsValid() && HostedToolkit->ProcessCommandBindings( InKeyboardEvent ) )
+	if( HostedToolkit.IsValid() && HostedToolkit->ProcessCommandBindings( InKeyEvent ) )
 	{
 		return FReply::Handled();
 	}
 
-	return SCompoundWidget::OnKeyDown(MyGeometry, InKeyboardEvent);
+	return SCompoundWidget::OnKeyDown(MyGeometry, InKeyEvent);
 }
 
 
diff --git a/Engine/Source/Editor/UnrealEd/Public/EditorViewportClient.h b/Engine/Source/Editor/UnrealEd/Public/EditorViewportClient.h
--- a/Engine/Source/Editor/UnrealEd/Public/EditorViewportClient.h
+++ b/Engine/Source/Editor/UnrealEd/Public/EditorViewportClient.h
@@ -95,7 +95,7 @@ private:
 	FViewport* Viewport;
 	/** Pressed Key */
 	FKey Key;
-	/** Keyboard event */
+	/** Key event */
 	EInputEvent InputEvent;
 };
 
diff --git a/Engine/Source/Editor/UnrealEd/Public/Kismet2/SClassPickerDialog.h b/Engine/Source/Editor/UnrealEd/Public/Kismet2/SClassPickerDialog.h
--- a/Engine/Source/Editor/UnrealEd/Public/Kismet2/SClassPickerDialog.h
+++ b/Engine/Source/Editor/UnrealEd/Public/Kismet2/SClassPickerDialog.h
@@ -49,7 +49,7 @@ private:
 	EVisibility GetSelectButtonVisibility() const;
 
 	/** Overridden from SWidget: Called when a key is pressed down - capturing copy */
-	FReply OnKeyDown( const FGeometry& MyGeometry, const FKeyboardEvent& InKeyboardEvent );
+	FReply OnKeyDown( const FGeometry& MyGeometry, const FKeyEvent& InKeyEvent );
 
 private:
 	/** A pointer to the window that is asking the user to select a parent class */
diff --git a/Engine/Source/Editor/UnrealEd/Public/SCurveEditor.h b/Engine/Source/Editor/UnrealEd/Public/SCurveEditor.h
--- a/Engine/Source/Editor/UnrealEd/Public/SCurveEditor.h
+++ b/Engine/Source/Editor/UnrealEd/Public/SCurveEditor.h
@@ -375,7 +375,7 @@ private:
 	UNREALED_API virtual FReply OnMouseButtonUp( const FGeometry& MyGeometry, const FPointerEvent& MouseEvent ) override;
 	UNREALED_API virtual FReply OnMouseMove( const FGeometry& MyGeometry, const FPointerEvent& MouseEvent ) override;
 	UNREALED_API virtual FReply OnMouseWheel( const FGeometry& MyGeometry, const FPointerEvent& MouseEvent ) override;
-	UNREALED_API virtual FReply OnKeyDown( const FGeometry& MyGeometry, const FKeyboardEvent& InKeyboardEvent ) override;
+	UNREALED_API virtual FReply OnKeyDown( const FGeometry& MyGeometry, const FKeyEvent& InKeyEvent ) override;
 	UNREALED_API virtual void   OnMouseCaptureLost() override;
 
 	/** Attempts to start a drag operation when the mouse moves. */
diff --git a/Engine/Source/Editor/UnrealEd/Public/SEditorViewport.h b/Engine/Source/Editor/UnrealEd/Public/SEditorViewport.h
--- a/Engine/Source/Editor/UnrealEd/Public/SEditorViewport.h
+++ b/Engine/Source/Editor/UnrealEd/Public/SEditorViewport.h
@@ -14,9 +14,9 @@ public:
 
 	void Construct( const FArguments& InArgs );
 
-	virtual FReply OnKeyDown( const FGeometry& MyGeometry, const FKeyboardEvent& InKeyboardEvent ) override;
+	virtual FReply OnKeyDown( const FGeometry& MyGeometry, const FKeyEvent& InKeyEvent ) override;
 	virtual bool SupportsKeyboardFocus() const override;
-	virtual FReply OnKeyboardFocusReceived( const FGeometry& MyGeometry, const FKeyboardFocusEvent& InKeyboardFocusEvent ) override;
+	virtual FReply OnFocusReceived( const FGeometry& MyGeometry, const FFocusEvent& InFocusEvent ) override;
 	virtual void Tick(const FGeometry& AllottedGeometry, const double InCurrentTime, const float InDeltaTime) override;
 	
 	/**
diff --git a/Engine/Source/Editor/UnrealEd/Public/SListViewSelectorDropdownMenu.h b/Engine/Source/Editor/UnrealEd/Public/SListViewSelectorDropdownMenu.h
--- a/Engine/Source/Editor/UnrealEd/Public/SListViewSelectorDropdownMenu.h
+++ b/Engine/Source/Editor/UnrealEd/Public/SListViewSelectorDropdownMenu.h
@@ -32,16 +32,16 @@ public:
 	}
 
 	// SWidget interface
-	virtual FReply OnKeyDown(const FGeometry& MyGeometry, const FKeyboardEvent& KeyboardEvent) override
+	virtual FReply OnKeyDown(const FGeometry& MyGeometry, const FKeyEvent& KeyEvent) override
 	{
 		TSharedPtr< SListView< ListType > > TargetListViewPtr = TargetListView.Pin();
 
-		if (KeyboardEvent.GetKey() == EKeys::Up || KeyboardEvent.GetKey() == EKeys::Down)
+		if (KeyEvent.GetKey() == EKeys::Up || KeyEvent.GetKey() == EKeys::Down)
 		{
 			// Deliver focus to the tree view, so the user can use the arrow keys to move through the items
-			return TargetListViewPtr->OnKeyDown(FindChildGeometry(MyGeometry, TargetListViewPtr.ToSharedRef()), KeyboardEvent);
+			return TargetListViewPtr->OnKeyDown(FindChildGeometry(MyGeometry, TargetListViewPtr.ToSharedRef()), KeyEvent);
 		}
-		else if(KeyboardEvent.GetKey() == EKeys::Enter)
+		else if(KeyEvent.GetKey() == EKeys::Enter)
 		{
 			// If there is anything selected, re-select it "direct" so that the menu will act upon the selection
 			if(TargetListViewPtr->GetNumItemsSelected() > 0)
@@ -53,7 +53,7 @@ public:
 		else
 		{
 			TSharedPtr< SWidget > DefaultFocusWidgetPtr = DefaultFocusWidget.Pin();
-			return DefaultFocusWidgetPtr->OnKeyDown(FindChildGeometry(MyGeometry, TargetListViewPtr.ToSharedRef()), KeyboardEvent).SetKeyboardFocus(DefaultFocusWidgetPtr.ToSharedRef(), EKeyboardFocusCause::OtherWidgetLostFocus);
+			return DefaultFocusWidgetPtr->OnKeyDown(FindChildGeometry(MyGeometry, TargetListViewPtr.ToSharedRef()), KeyEvent).SetUserFocus(DefaultFocusWidgetPtr.ToSharedRef(), EFocusCause::OtherWidgetLostFocus);
 		}
 		return FReply::Unhandled();
 	}
diff --git a/Engine/Source/Editor/UnrealEd/Public/Toolkits/BaseToolkit.h b/Engine/Source/Editor/UnrealEd/Public/Toolkits/BaseToolkit.h
--- a/Engine/Source/Editor/UnrealEd/Public/Toolkits/BaseToolkit.h
+++ b/Engine/Source/Editor/UnrealEd/Public/Toolkits/BaseToolkit.h
@@ -21,7 +21,7 @@ public:
 
 	/** IToolkit interface */
 	virtual FName GetToolkitContextFName() const override;
-	virtual bool ProcessCommandBindings( const FKeyboardEvent& InKeyboardEvent ) const override;
+	virtual bool ProcessCommandBindings( const FKeyEvent& InKeyEvent ) const override;
 	virtual bool IsHosted() const override;
 	virtual const TSharedRef< class IToolkitHost > GetToolkitHost() const override;
 	virtual const TMap< EToolkitTabSpot::Type, TArray< TWeakPtr< SDockableTab > > >& GetToolkitTabsInSpots() const override;
diff --git a/Engine/Source/Editor/UnrealEd/Public/Toolkits/IToolkit.h b/Engine/Source/Editor/UnrealEd/Public/Toolkits/IToolkit.h
--- a/Engine/Source/Editor/UnrealEd/Public/Toolkits/IToolkit.h
+++ b/Engine/Source/Editor/UnrealEd/Public/Toolkits/IToolkit.h
@@ -97,13 +97,13 @@ public:
 	virtual const TMap< EToolkitTabSpot::Type, TArray< TWeakPtr< class SDockableTab > > >& GetToolkitTabsInSpots() const = 0;
 
 	/**
-	 * Processes any UI commands which are activated by the specified keyboard event
+	 * Processes any UI commands which are activated by the specified event
 	 *
-	 * @param InKeyboardEvent	The keyboard event to check
+	 * @param InKeyEvent	The event to check
 	 *
 	 * @return true if an action was processed
 	 */
-	virtual bool ProcessCommandBindings( const struct FKeyboardEvent& InKeyboardEvent ) const = 0;
+	virtual bool ProcessCommandBindings( const struct FKeyEvent& InKeyEvent ) const = 0;
 
 	/** Call this function to bring all of this toolkit's tabs to the foreground in their respective stacks.  Also causes the toolkit's host window to be foregrounded, too! */
 	virtual void BringToolkitToFront() = 0;
diff --git a/Engine/Source/Editor/UserFeedback/Private/UserFeedback.cpp b/Engine/Source/Editor/UserFeedback/Private/UserFeedback.cpp
--- a/Engine/Source/Editor/UserFeedback/Private/UserFeedback.cpp
+++ b/Engine/Source/Editor/UserFeedback/Private/UserFeedback.cpp
@@ -183,13 +183,13 @@ public:
 	}
 
 	/** Called when a key is pressed on the widget */
-	virtual FReply OnPreviewKeyDown(const FGeometry& MyGeometry, const FKeyboardEvent& InKeyboardEvent) override
+	virtual FReply OnPreviewKeyDown(const FGeometry& MyGeometry, const FKeyEvent& InKeyEvent) override
 	{
-		if (InKeyboardEvent.GetKey() == EKeys::Escape)
+		if (InKeyEvent.GetKey() == EKeys::Escape)
 		{
 			return OnCloseButtonClicked();
 		}
-		else if (InKeyboardEvent.GetKey() == EKeys::Enter)
+		else if (InKeyEvent.GetKey() == EKeys::Enter)
 		{
 			return OnFeedbackSubmitted();
 		}
diff --git a/Engine/Source/Editor/WorldBrowser/Private/SWorldHierarchy.cpp b/Engine/Source/Editor/WorldBrowser/Private/SWorldHierarchy.cpp
--- a/Engine/Source/Editor/WorldBrowser/Private/SWorldHierarchy.cpp
+++ b/Engine/Source/Editor/WorldBrowser/Private/SWorldHierarchy.cpp
@@ -403,21 +403,21 @@ private:
 	}
 
 	/**
-	 * Called after a key is pressed when this widget has keyboard focus
+	 * Called after a key is pressed when this widget has focus
 	 *
 	 * @param MyGeometry The Geometry of the widget receiving the event
-	 * @param  InKeyboardEvent  Keyboard event
+	 * @param  InKeyEvent  Key event
 	 *
 	 * @return  Returns whether the event was handled, along with other possible actions
 	 */
-	virtual FReply OnKeyDown(const FGeometry& MyGeometry, const FKeyboardEvent& InKeyboardEvent) override
+	virtual FReply OnKeyDown(const FGeometry& MyGeometry, const FKeyEvent& InKeyEvent) override
 	{
-		if (WorldModel->GetCommandList()->ProcessCommandBindings(InKeyboardEvent))
+		if (WorldModel->GetCommandList()->ProcessCommandBindings(InKeyEvent))
 		{
 			return FReply::Handled();
 		}
 	
-		return SCompoundWidget::OnKeyDown(MyGeometry, InKeyboardEvent);
+		return SCompoundWidget::OnKeyDown(MyGeometry, InKeyEvent);
 	}
 
 private:
diff --git a/Engine/Source/Editor/WorldBrowser/Private/Tiles/SWorldComposition.cpp b/Engine/Source/Editor/WorldBrowser/Private/Tiles/SWorldComposition.cpp
--- a/Engine/Source/Editor/WorldBrowser/Private/Tiles/SWorldComposition.cpp
+++ b/Engine/Source/Editor/WorldBrowser/Private/Tiles/SWorldComposition.cpp
@@ -573,19 +573,19 @@ protected:
 	}
 		
 	/**  SWidget interface */
-	virtual FReply OnKeyDown(const FGeometry& MyGeometry, const FKeyboardEvent& InKeyboardEvent) override
+	virtual FReply OnKeyDown(const FGeometry& MyGeometry, const FKeyEvent& InKeyEvent) override
 	{
-		if (CommandList->ProcessCommandBindings(InKeyboardEvent))
+		if (CommandList->ProcessCommandBindings(InKeyEvent))
 		{
 			return FReply::Handled();
 		}
 	
-		if (WorldModel->GetCommandList()->ProcessCommandBindings(InKeyboardEvent))
+		if (WorldModel->GetCommandList()->ProcessCommandBindings(InKeyEvent))
 		{
 			return FReply::Handled();
 		}
 
-		return SNodePanel::OnKeyDown(MyGeometry, InKeyboardEvent);
+		return SNodePanel::OnKeyDown(MyGeometry, InKeyEvent);
 	}
 
 	/**  SWidget interface */
diff --git a/Engine/Source/Programs/CrashReporter/CrashReportClient/Private/SCrashReportClient.cpp b/Engine/Source/Programs/CrashReporter/CrashReportClient/Private/SCrashReportClient.cpp
--- a/Engine/Source/Programs/CrashReporter/CrashReportClient/Private/SCrashReportClient.cpp
+++ b/Engine/Source/Programs/CrashReporter/CrashReportClient/Private/SCrashReportClient.cpp
@@ -167,7 +167,7 @@ void SCrashReportClient::Construct(const FArguments& InArgs, TSharedRef<FCrashRe
 		];
 	}
 
-	FSlateApplication::Get().SetUnhandledKeyDownEventHandler(FOnKeyboardEvent::CreateSP(this, &SCrashReportClient::OnUnhandledKeyDown));
+	FSlateApplication::Get().SetUnhandledKeyDownEventHandler(FOnKeyEvent::CreateSP(this, &SCrashReportClient::OnUnhandledKeyDown));
 }
 
 void SCrashReportClient::SetDefaultFocus()
@@ -175,9 +175,9 @@ void SCrashReportClient::SetDefaultFocus()
 	FSlateApplication::Get().SetKeyboardFocus(UserCommentBox.ToSharedRef());
 }
 
-FReply SCrashReportClient::OnUnhandledKeyDown(const FKeyboardEvent& InKeyboardEvent)
+FReply SCrashReportClient::OnUnhandledKeyDown(const FKeyEvent& InKeyEvent)
 {
-	const FKey Key = InKeyboardEvent.GetKey();
+	const FKey Key = InKeyEvent.GetKey();
 	if (Key == EKeys::Escape)
 	{
 		CrashReportClient->Cancel();
@@ -185,7 +185,7 @@ FReply SCrashReportClient::OnUnhandledKeyDown(const FKeyboardEvent& InKeyboardEv
 	}
 	else if (Key == EKeys::Enter)
 	{
-		if (InKeyboardEvent.IsShiftDown())
+		if (InKeyEvent.IsShiftDown())
 		{
 			CrashReportClient->Submit();
 			return FReply::Handled();
diff --git a/Engine/Source/Programs/CrashReporter/CrashReportClient/Private/SCrashReportClient.h b/Engine/Source/Programs/CrashReporter/CrashReportClient/Private/SCrashReportClient.h
--- a/Engine/Source/Programs/CrashReporter/CrashReportClient/Private/SCrashReportClient.h
+++ b/Engine/Source/Programs/CrashReporter/CrashReportClient/Private/SCrashReportClient.h
@@ -36,10 +36,10 @@ public:
 private:
 	/**
 	 * Keyboard short-cut handler
-	 * @param InKeyboardEvent Which key was released, and which auxiliary keys were pressed
+	 * @param InKeyEvent Which key was released, and which auxiliary keys were pressed
 	 * @return Whether the event was handled
 	 */
-	FReply OnUnhandledKeyDown(const FKeyboardEvent& InKeyboardEvent);
+	FReply OnUnhandledKeyDown(const FKeyEvent& InKeyEvent);
 
 	/** Crash report client implementation object */
 	TSharedPtr<FCrashReportClient> CrashReportClient;
diff --git a/Engine/Source/Runtime/AppFramework/Private/Widgets/Colors/SEyeDropperButton.h b/Engine/Source/Runtime/AppFramework/Private/Widgets/Colors/SEyeDropperButton.h
--- a/Engine/Source/Runtime/AppFramework/Private/Widgets/Colors/SEyeDropperButton.h
+++ b/Engine/Source/Runtime/AppFramework/Private/Widgets/Colors/SEyeDropperButton.h
@@ -160,10 +160,10 @@ public:
 		return SButton::OnMouseMove( MyGeometry, MouseEvent );
 	}
 
-	FReply OnKeyDown( const FGeometry& MyGeometry, const FKeyboardEvent& InKeyboardEvent )
+	FReply OnKeyDown( const FGeometry& MyGeometry, const FKeyEvent& InKeyEvent )
 	{
 		// Escape key when in dropper mode
-		if (InKeyboardEvent.GetKey() == EKeys::Escape &&
+		if (InKeyEvent.GetKey() == EKeys::Escape &&
 			HasMouseCapture() &&
 			bWasClickActivated)
 		{
diff --git a/Engine/Source/Runtime/AppFramework/Private/Widgets/Testing/STestSuite.cpp b/Engine/Source/Runtime/AppFramework/Private/Widgets/Testing/STestSuite.cpp
--- a/Engine/Source/Runtime/AppFramework/Private/Widgets/Testing/STestSuite.cpp
+++ b/Engine/Source/Runtime/AppFramework/Private/Widgets/Testing/STestSuite.cpp
@@ -2036,7 +2036,7 @@ public:
 		// Set focus to the editable text, so the user doesn't have to click initially to type
 		FWidgetPath WidgetToFocusPath;
 		FSlateApplication::Get().GeneratePathToWidgetUnchecked( EditableText.ToSharedRef(), WidgetToFocusPath );
-		FSlateApplication::Get().SetKeyboardFocus( WidgetToFocusPath, EKeyboardFocusCause::SetDirectly );
+		FSlateApplication::Get().SetKeyboardFocus( WidgetToFocusPath, EFocusCause::SetDirectly );
 	}
 
 	void InlineEditableTextCommited(const FText& NewText, ETextCommit::Type CommitType )
@@ -2867,7 +2867,7 @@ protected:
 		const FRunInfo RunInfo = TextStyles.CreateRunInfo(ActiveFontFamily, FontSize, FontStyle, FontColor);
 		const FTextBlockStyle TextBlockStyle = TextStyles.CreateTextBlockStyle(ActiveFontFamily, FontSize, FontStyle, FontColor);
 		RichEditableTextBox->ApplyToSelection(RunInfo, TextBlockStyle);
-		FSlateApplication::Get().SetKeyboardFocus(RichEditableTextBox, EKeyboardFocusCause::SetDirectly);
+		FSlateApplication::Get().SetKeyboardFocus(RichEditableTextBox, EFocusCause::SetDirectly);
 	}
 
 	void HandleHyperlinkComboOpened()
diff --git a/Engine/Source/Runtime/Core/Private/Linux/LinuxApplication.cpp b/Engine/Source/Runtime/Core/Private/Linux/LinuxApplication.cpp
--- a/Engine/Source/Runtime/Core/Private/Linux/LinuxApplication.cpp
+++ b/Engine/Source/Runtime/Core/Private/Linux/LinuxApplication.cpp
@@ -168,7 +168,7 @@ void FLinuxApplication::AddPendingEvent( SDL_Event SDLEvent )
 	}
 }
 
-bool FLinuxApplication::GeneratesKeyCharMessage(const SDL_KeyboardEvent & KeyDownEvent)
+bool FLinuxApplication::GeneratesKeyCharMessage(const SDL_KeyEvent & KeyDownEvent)
 {
 	bool bCmdKeyPressed = (KeyDownEvent.keysym.mod & (KMOD_LCTRL | KMOD_RCTRL)) != 0;
 	const SDL_Keycode Sym = KeyDownEvent.keysym.sym;
@@ -237,7 +237,7 @@ void FLinuxApplication::ProcessDeferredMessage( SDL_Event Event )
 	{
 	case SDL_KEYDOWN:
 		{
-			SDL_KeyboardEvent KeyEvent = Event.key;
+			SDL_KeyEvent KeyEvent = Event.key;
 			const SDL_Keycode KeyCode = KeyEvent.keysym.scancode;
 			const bool bIsRepeated = KeyEvent.repeat != 0;
 				
@@ -253,7 +253,7 @@ void FLinuxApplication::ProcessDeferredMessage( SDL_Event Event )
 		break;
 	case SDL_KEYUP:
 		{
-			SDL_KeyboardEvent keyEvent = Event.key;
+			SDL_KeyEvent keyEvent = Event.key;
 			const SDL_Keycode KeyCode = keyEvent.keysym.scancode;
 			const TCHAR Character = ConvertChar( keyEvent.keysym );
 			const bool IsRepeat = keyEvent.repeat != 0;
