Fixed non-unity build

diff --git a/Engine/Source/Editor/UMGEditor/Private/UMGEditorPrivatePCH.h b/Engine/Source/Editor/UMGEditor/Private/UMGEditorPrivatePCH.h
--- a/Engine/Source/Editor/UMGEditor/Private/UMGEditorPrivatePCH.h
+++ b/Engine/Source/Editor/UMGEditor/Private/UMGEditorPrivatePCH.h
@@ -13,6 +13,7 @@
 #include "LevelEditorViewport.h"
 #include "ParticleHelper.h"
 #include "Viewports.h"
+#include "EditorStyle.h"
 #endif // WITH_EDITOR
 //#include "UnrealEd.h" // End
 
