[UE-4926] Fixing a crash that would occur in UMG when adding an event handler to a widget after editing a value on the widget that caused a GC.

diff --git a/Engine/Source/Editor/UMGEditor/Private/Customizations/DetailCustomizations.cpp b/Engine/Source/Editor/UMGEditor/Private/Customizations/DetailCustomizations.cpp
--- a/Engine/Source/Editor/UMGEditor/Private/Customizations/DetailCustomizations.cpp
+++ b/Engine/Source/Editor/UMGEditor/Private/Customizations/DetailCustomizations.cpp
@@ -20,6 +20,68 @@
 
 #define LOCTEXT_NAMESPACE "UMG"
 
+
+class SGraphSchemaActionButton : public SCompoundWidget, public FGCObject
+{
+public:
+
+	SLATE_BEGIN_ARGS(SGraphSchemaActionButton) {}
+	SLATE_END_ARGS()
+
+		void Construct(const FArguments& InArgs, TSharedPtr<FWidgetBlueprintEditor> InEditor, TSharedPtr<FEdGraphSchemaAction> InClickAction)
+	{
+		Editor = InEditor;
+		Action = InClickAction;
+
+		ChildSlot
+		[
+			SNew(SButton)
+			.ToolTipText(Action->TooltipDescription)
+			.OnClicked(this, &SGraphSchemaActionButton::AddOrViewEventBinding)
+			.HAlign(HAlign_Center)
+			[
+				SNew(STextBlock)
+				.Text(Action->MenuDescription)
+			]
+		];
+	}
+
+private:
+	FReply AddOrViewEventBinding()
+	{
+		UBlueprint* Blueprint = Editor.Pin()->GetBlueprintObj();
+
+		UEdGraph* TargetGraph = nullptr;
+		if ( Blueprint->UbergraphPages.Num() > 0 )
+		{
+			TargetGraph = Blueprint->UbergraphPages[0]; // Just use the first graph
+		}
+
+		if ( TargetGraph != nullptr )
+		{
+			Editor.Pin()->SetCurrentMode(FWidgetBlueprintApplicationModes::GraphMode);
+
+			// Figure out a decent place to stick the node
+			const FVector2D NewNodePos = TargetGraph->GetGoodPlaceForNewNode();
+
+			Action->PerformAction(TargetGraph, nullptr, NewNodePos);
+		}
+
+		return FReply::Handled();
+	}
+
+	virtual void AddReferencedObjects(FReferenceCollector& Collector) override
+	{
+		Action->AddReferencedObjects(Collector);
+	}
+
+private:
+	TWeakPtr<FWidgetBlueprintEditor> Editor;
+
+	TSharedPtr<FEdGraphSchemaAction> Action;
+};
+
+
 void FBlueprintWidgetCustomization::CreateDelegateCustomization( IDetailLayoutBuilder& DetailLayout, UDelegateProperty* Property, UWidget* Widget )
 {
 	FString ConstPropertyName = Property->GetName();
@@ -268,7 +330,7 @@ void FBlueprintWidgetCustomization::CreateMulticastEventCustomization(IDetailLay
 
 	IDetailCategoryBuilder& PropertyCategory = DetailLayout.EditCategory("Events", LOCTEXT("Events", "Events").ToString(), ECategoryPriority::Uncommon);
 
-	PropertyCategory.AddCustomRow(LOCTEXT("FindInLevelScript_Filter", "Find in Level Script").ToString())
+	PropertyCategory.AddCustomRow(DelegateProperty->GetName())
 	.NameContent()
 	[
 		SNew(SHorizontalBox)
@@ -293,38 +355,10 @@ void FBlueprintWidgetCustomization::CreateMulticastEventCustomization(IDetailLay
 	.MinDesiredWidth(200)
 	.MaxDesiredWidth(200)
 	[
-		SNew(SButton)
-		.ToolTipText(ClickAction->TooltipDescription)
-		.OnClicked(this, &FBlueprintWidgetCustomization::AddOrViewEventBinding, ClickAction)
-		.HAlign(HAlign_Center)
-		[
-			SNew(STextBlock)
-			.Text(ClickAction->MenuDescription)
-		]
+		SNew(SGraphSchemaActionButton, Editor.Pin(), ClickAction)
 	];
 }
 
-FReply FBlueprintWidgetCustomization::AddOrViewEventBinding(TSharedPtr<FEdGraphSchemaAction> Action)
-{
-	UEdGraph* TargetGraph = NULL;
-	if ( Blueprint->UbergraphPages.Num() > 0 )
-	{
-		TargetGraph = Blueprint->UbergraphPages[0]; // Just use the first graph
-	}
-
-	if ( TargetGraph != NULL )
-	{
-		Editor.Pin()->SetCurrentMode(FWidgetBlueprintApplicationModes::GraphMode);
-
-		// Figure out a decent place to stick the node
-		const FVector2D NewNodePos = TargetGraph->GetGoodPlaceForNewNode();
-
-		Action->PerformAction(TargetGraph, NULL, NewNodePos);
-	}
-
-	return FReply::Handled();
-}
-
 void FBlueprintWidgetCustomization::CustomizeDetails( IDetailLayoutBuilder& DetailLayout )
 {
 	static const FName LayoutCategoryKey(TEXT("Layout"));
