UMG - Exposing a reinstance delegate that occurs for blueprints after map replacement, but before GC for any final destruction that needs to be carried out.  Improving the hit testing code in the designer for future usage, to enable dropping onto a named slot.  Tweaking the text used in the designer for the resolution drop down.

diff --git a/Engine/Source/Editor/UMGEditor/Private/Designer/SDesignerView.cpp b/Engine/Source/Editor/UMGEditor/Private/Designer/SDesignerView.cpp
--- a/Engine/Source/Editor/UMGEditor/Private/Designer/SDesignerView.cpp
+++ b/Engine/Source/Editor/UMGEditor/Private/Designer/SDesignerView.cpp
@@ -45,6 +45,25 @@
 
 const float HoveredAnimationTime = 0.150f;
 
+
+struct FWidgetHitResult
+{
+public:
+	FWidgetReference Widget;
+	FArrangedWidget WidgetArranged;
+
+	UNamedSlot* NamedSlot;
+	FArrangedWidget NamedSlotArranged;
+
+public:
+	FWidgetHitResult()
+		: WidgetArranged(SNullWidget::NullWidget, FGeometry())
+		, NamedSlotArranged(SNullWidget::NullWidget, FGeometry())
+	{
+	}
+};
+
+
 class FSelectedWidgetDragDropOp : public FDecoratedDragDropOp
 {
 public:
@@ -167,7 +186,7 @@ void SDesignerView::Construct(const FArguments& InArgs, TSharedPtr<FWidgetBluepr
 	Register(MakeShareable(new FUniformGridSlotExtension()));
 	Register(MakeShareable(new FGridSlotExtension()));
 
-	FWidgetBlueprintCompiler::OnWidgetBlueprintCompiled.AddSP( this, &SDesignerView::OnBlueprintCompiled );
+	GEditor->OnBlueprintReinstanced().AddRaw(this, &SDesignerView::OnBlueprintReinstanced);
 
 	BindCommands();
 
@@ -292,7 +311,7 @@ void SDesignerView::Construct(const FArguments& InArgs, TSharedPtr<FWidgetBluepr
 					.ButtonContent()
 					[
 						SNew(STextBlock)
-						.Text(LOCTEXT("Resolution", "Resolution"))
+						.Text(LOCTEXT("PreviewSize", "Preview Size"))
 						.TextStyle(FEditorStyle::Get(), "ViewportMenu.Label")
 					]
 				]
@@ -436,6 +455,11 @@ SDesignerView::~SDesignerView()
 	{
 		BlueprintEditor.Pin()->OnSelectedWidgetsChanged.RemoveAll(this);
 	}
+
+	if ( GEditor )
+	{
+		GEditor->OnBlueprintReinstanced().RemoveAll(this);
+	}
 }
 
 void SDesignerView::BindCommands()
@@ -827,7 +851,7 @@ void SDesignerView::Register(TSharedRef<FDesignerExtension> Extension)
 	DesignerExtensions.Add(Extension);
 }
 
-void SDesignerView::OnBlueprintCompiled(UBlueprint* InBlueprint)
+void SDesignerView::OnBlueprintReinstanced()
 {
 	// Because widget blueprints can contain other widget blueprints, the safe thing to do is to have all
 	// designers jettison their previews on the compilation of any widget blueprint.  We do this to prevent
@@ -838,7 +862,14 @@ void SDesignerView::OnBlueprintCompiled(UBlueprint* InBlueprint)
 	PreviewSurface->SetContent(SNullWidget::NullWidget);
 }
 
-FWidgetReference SDesignerView::GetWidgetAtCursor(const FGeometry& MyGeometry, const FPointerEvent& MouseEvent, FArrangedWidget& ArrangedWidget)
+SDesignerView::FWidgetHitResult::FWidgetHitResult()
+	: Widget()
+	, WidgetArranged(SNullWidget::NullWidget, FGeometry())
+	, NamedSlot(NAME_None)
+{
+}
+
+bool SDesignerView::FindWidgetUnderCursor(const FGeometry& MyGeometry, const FPointerEvent& MouseEvent, FWidgetHitResult& HitResult)
 {
 	//@TODO UMG Make it so you can request dropable widgets only, to find the first parentable.
 
@@ -850,37 +881,54 @@ FWidgetReference SDesignerView::GetWidgetAtCursor(const FGeometry& MyGeometry, c
 
 	PreviewHitTestRoot->SetVisibility(EVisibility::HitTestInvisible);
 
-	UUserWidget* WidgetActor = BlueprintEditor.Pin()->GetPreview();
-	if ( WidgetActor )
+	HitResult.Widget = FWidgetReference();
+	HitResult.NamedSlot = NAME_None;
+
+	UUserWidget* PreviewUserWidget = BlueprintEditor.Pin()->GetPreview();
+	if ( PreviewUserWidget )
 	{
-		UWidget* Preview = nullptr;
+		UWidget* WidgetUnderCursor = nullptr;
 
+		// We loop through each hit slate widget until we arrive at one that we can access from the root widget.
 		for ( int32 ChildIndex = Children.Num() - 1; ChildIndex >= 0; ChildIndex-- )
 		{
 			FArrangedWidget& Child = Children.GetInternalArray()[ChildIndex];
-			Preview = WidgetActor->GetWidgetHandle(Child.Widget);
+			WidgetUnderCursor = PreviewUserWidget->GetWidgetHandle(Child.Widget);
 			
 			// Ignore the drop preview widget when doing widget picking
-			if (Preview == DropPreviewWidget)
+			if ( WidgetUnderCursor == DropPreviewWidget )
 			{
-				Preview = nullptr;
+				WidgetUnderCursor = nullptr;
 				continue;
 			}
 			
-			if ( Preview )
+			// We successfully found a widget that's accessible from the root.
+			if ( WidgetUnderCursor )
 			{
-				ArrangedWidget = Child;
-				break;
-			}
-		}
+				HitResult.Widget = BlueprintEditor.Pin()->GetReferenceFromPreview(WidgetUnderCursor);
+				HitResult.WidgetArranged = Child;
 
-		if ( Preview )
-		{
-			return BlueprintEditor.Pin()->GetReferenceFromPreview(Preview);
+				if ( UUserWidget* UserWidgetUnderCursor = Cast<UUserWidget>(WidgetUnderCursor) )
+				{
+					// Find the named slot we're over, if any
+					for ( int32 SubChildIndex = Children.Num() - 1; SubChildIndex > ChildIndex; SubChildIndex-- )
+					{
+						FArrangedWidget& SubChild = Children.GetInternalArray()[SubChildIndex];
+						UNamedSlot* NamedSlot = Cast<UNamedSlot>(UserWidgetUnderCursor->GetWidgetHandle(SubChild.Widget));
+						if ( NamedSlot )
+						{
+							HitResult.NamedSlot = NamedSlot->GetFName();
+							break;
+						}
+					}
+				}
+
+				return true;
+			}
 		}
 	}
 
-	return FWidgetReference();
+	return false;
 }
 
 void SDesignerView::ResolvePendingSelectedWidgets()
@@ -900,19 +948,18 @@ FReply SDesignerView::OnMouseButtonDown(const FGeometry& MyGeometry, const FPoin
 	SDesignSurface::OnMouseButtonDown(MyGeometry, MouseEvent);
 
 	//TODO UMG Undoable Selection
-	FArrangedWidget ArrangedWidget(SNullWidget::NullWidget, FGeometry());
-	FWidgetReference NewSelectedWidget = GetWidgetAtCursor(MyGeometry, MouseEvent, ArrangedWidget);
-	SelectedWidgetContextMenuLocation = ArrangedWidget.Geometry.AbsoluteToLocal(MouseEvent.GetScreenSpacePosition());
-
-	if ( NewSelectedWidget.IsValid() )
+	FWidgetHitResult HitResult;
+	if ( FindWidgetUnderCursor(MyGeometry, MouseEvent, HitResult) )
 	{
-		PendingSelectedWidget = NewSelectedWidget;
+		SelectedWidgetContextMenuLocation = HitResult.WidgetArranged.Geometry.AbsoluteToLocal(MouseEvent.GetScreenSpacePosition());
+
+		PendingSelectedWidget = HitResult.Widget;
 
 		if ( MouseEvent.GetEffectingButton() == EKeys::LeftMouseButton )
 		{
 			const bool bResolvePendingSelectionImmediately =
 				!GetSelectedWidget().IsValid() ||
-				!NewSelectedWidget.GetTemplate()->IsChildOf(GetSelectedWidget().GetTemplate()) ||
+				!PendingSelectedWidget.GetTemplate()->IsChildOf(GetSelectedWidget().GetTemplate()) ||
 				GetSelectedWidget().GetTemplate()->GetParent() == nullptr;
 
 			// If the newly clicked item is a child of the active selection, add it to the pending set of selected 
@@ -979,7 +1026,7 @@ FReply SDesignerView::OnMouseMove(const FGeometry& MyGeometry, const FPointerEve
 		{
 			if ( TransformMode == ETransformMode::Layout )
 			{
-				const bool bIsRootWidget = SelectedWidget.GetTemplate()->GetParent() == NULL;
+				const bool bIsRootWidget = SelectedWidget.GetTemplate()->GetParent() == nullptr;
 				if ( !bIsRootWidget )
 				{
 					bMovingExistingWidget = true;
@@ -1015,9 +1062,11 @@ FReply SDesignerView::OnMouseMove(const FGeometry& MyGeometry, const FPointerEve
 	}
 	
 	// Update the hovered widget under the mouse
-	FArrangedWidget ArrangedWidget(SNullWidget::NullWidget, FGeometry());
-	FWidgetReference NewHoveredWidget = GetWidgetAtCursor(MyGeometry, MouseEvent, ArrangedWidget);
-	BlueprintEditor.Pin()->SetHoveredWidget(NewHoveredWidget);
+	FWidgetHitResult HitResult;
+	if ( FindWidgetUnderCursor(MyGeometry, MouseEvent, HitResult) )
+	{
+		BlueprintEditor.Pin()->SetHoveredWidget(HitResult.Widget);
+	}
 
 	return FReply::Unhandled();
 }
@@ -1325,21 +1374,21 @@ UWidget* SDesignerView::ProcessDropAndAddWidget(const FGeometry& MyGeometry, con
 		DropPreviewWidget = nullptr;
 	}
 
-	FArrangedWidget ArrangedWidget(SNullWidget::NullWidget, FGeometry());
-	FWidgetReference WidgetUnderCursor = GetWidgetAtCursor(MyGeometry, DragDropEvent, ArrangedWidget);
-
-	FGeometry WidgetUnderCursorGeometry = ArrangedWidget.Geometry;
 
 	UWidget* Target = nullptr;
-	if ( WidgetUnderCursor.IsValid() )
+
+	FWidgetHitResult HitResult;
+	if ( FindWidgetUnderCursor(MyGeometry, DragDropEvent, HitResult) )
 	{
-		Target = bIsPreview ? WidgetUnderCursor.GetPreview() : WidgetUnderCursor.GetTemplate();
+		Target = bIsPreview ? HitResult.Widget.GetPreview() : HitResult.Widget.GetTemplate();
 	}
 
+	FGeometry WidgetUnderCursorGeometry = HitResult.WidgetArranged.Geometry;
+
 	TSharedPtr<FWidgetTemplateDragDropOp> TemplateDragDropOp = DragDropEvent.GetOperationAs<FWidgetTemplateDragDropOp>();
 	if ( TemplateDragDropOp.IsValid() )
 	{
-		BlueprintEditor.Pin()->SetHoveredWidget(WidgetUnderCursor);
+		BlueprintEditor.Pin()->SetHoveredWidget(HitResult.Widget);
 
 		TemplateDragDropOp->SetCursorOverride(TOptional<EMouseCursor::Type>());
 
@@ -1835,7 +1884,7 @@ UUserWidget* SDesignerView::GetDefaultWidget() const
 TSharedRef<SWidget> SDesignerView::GetAspectMenu()
 {
 	const ULevelEditorPlaySettings* PlaySettings = GetDefault<ULevelEditorPlaySettings>();
-	FMenuBuilder MenuBuilder(true, NULL);
+	FMenuBuilder MenuBuilder(true, nullptr);
 
 	// Add custom option
 	FExecuteAction OnResolutionSelected = FExecuteAction::CreateRaw(this, &SDesignerView::HandleOnCustomResolutionSelected);
diff --git a/Engine/Source/Editor/UMGEditor/Private/Designer/SDesignerView.h b/Engine/Source/Editor/UMGEditor/Private/Designer/SDesignerView.h
--- a/Engine/Source/Editor/UMGEditor/Private/Designer/SDesignerView.h
+++ b/Engine/Source/Editor/UMGEditor/Private/Designer/SDesignerView.h
@@ -103,14 +103,11 @@ private:
 
 	void OnEditorSelectionChanged();
 
-	/** @returns Gets the widget under the cursor based on a mouse pointer event. */
-	FWidgetReference GetWidgetAtCursor(const FGeometry& MyGeometry, const FPointerEvent& MouseEvent, FArrangedWidget& ArrangedWidget);
-
 	/** Gets the blueprint being edited by the designer */
 	UWidgetBlueprint* GetBlueprint() const;
 
 	/** Called whenever the blueprint is recompiled */
-	void OnBlueprintCompiled(UBlueprint* InBlueprint);
+	void OnBlueprintReinstanced();
 
 	void PopulateWidgetGeometryCache(FArrangedWidget& Root);
 
@@ -146,6 +143,22 @@ private:
 	bool InTransaction() const;
 	void EndTransaction(bool bCancel);
 
+private:
+	struct FWidgetHitResult
+	{
+	public:
+		FWidgetReference Widget;
+		FArrangedWidget WidgetArranged;
+
+		FName NamedSlot;
+
+	public:
+		FWidgetHitResult();
+	};
+
+	/** @returns Gets the widget under the cursor based on a mouse pointer event. */
+	bool FindWidgetUnderCursor(const FGeometry& MyGeometry, const FPointerEvent& MouseEvent, FWidgetHitResult& HitResult);
+
 private:
 	FReply HandleZoomToFitClicked();
 
diff --git a/Engine/Source/Editor/UMGEditor/Private/Palette/SPaletteView.cpp b/Engine/Source/Editor/UMGEditor/Private/Palette/SPaletteView.cpp
--- a/Engine/Source/Editor/UMGEditor/Private/Palette/SPaletteView.cpp
+++ b/Engine/Source/Editor/UMGEditor/Private/Palette/SPaletteView.cpp
@@ -172,7 +172,7 @@ public:
 void SPaletteView::Construct(const FArguments& InArgs, TSharedPtr<FBlueprintEditor> InBlueprintEditor)
 {
 	// Register for events that can trigger a palette rebuild
-	FWidgetBlueprintCompiler::OnWidgetBlueprintCompiled.AddSP(this, &SPaletteView::HandleOnBlueprintCompiled);
+	GEditor->OnBlueprintReinstanced().AddRaw(this, &SPaletteView::OnBlueprintReinstanced);
 	FEditorDelegates::OnAssetsDeleted.AddSP(this, &SPaletteView::HandleOnAssetsDeleted);
 	IHotReloadModule::Get().OnHotReload().AddSP(this, &SPaletteView::HandleOnHotReload);
 	
@@ -242,7 +242,7 @@ SPaletteView::~SPaletteView()
 		FilterHandler->RefreshAndFilterTree();
 	}
 
-	FWidgetBlueprintCompiler::OnWidgetBlueprintCompiled.RemoveAll(this);
+	GEditor->OnBlueprintReinstanced().RemoveAll(this);
 	FEditorDelegates::OnAssetsDeleted.RemoveAll(this);
 	IHotReloadModule::Get().OnHotReload().RemoveAll(this);
 	GEditor->OnObjectsReplaced().RemoveAll( this );
@@ -356,14 +356,19 @@ void SPaletteView::BuildClassWidgetList()
 				continue;
 			}
 
-			if (WidgetClass->IsChildOf(UUserWidget::StaticClass()))
+			// Don't include skeleton classes or the same class as the widget being edited
+			const bool bIsSkeletonClass = WidgetClass->HasAnyFlags(RF_Transient) && WidgetClass->HasAnyClassFlags(CLASS_CompiledFromBlueprint);
+			const bool bIsSameClass = WidgetClass->GetFName() == ActiveWidgetBlueprintClassName;
+
+			// Check that the asset that generated this class is valid (necessary b/c of a larger issue wherein force delete does not wipe the generated class object)
+			if ( bIsSkeletonClass || bIsSameClass )
 			{
-				// Don't include skeleton classes or the same class as the widget being edited
-				const bool bIsSkeletonClass = WidgetClass->HasAnyFlags(RF_Transient) && WidgetClass->HasAnyClassFlags(CLASS_CompiledFromBlueprint);
-				const bool bIsSameClass = WidgetClass->GetFName() == ActiveWidgetBlueprintClassName;
+				continue;
+			}
 
-				// Check that the asset that generated this class is valid (necessary b/c of a larger issue wherein force delete does not wipe the generated class object)
-				if ( !(bIsSkeletonClass || bIsSameClass) && WidgetClass->ClassGeneratedBy )
+			if (WidgetClass->IsChildOf(UUserWidget::StaticClass()))
+			{
+				if ( WidgetClass->ClassGeneratedBy )
 				{
 					// Track the widget blueprint classes that are already loaded
 					LoadedWidgetBlueprintClassesByName.Add(WidgetClass->ClassGeneratedBy->GetFName()) = WidgetClass;
@@ -475,13 +480,10 @@ void SPaletteView::OnObjectsReplaced(const TMap<UObject*, UObject*>& Replacement
 	//bRebuildRequested = true;
 }
 
-void SPaletteView::HandleOnBlueprintCompiled(UBlueprint* Blueprint)
+void SPaletteView::OnBlueprintReinstanced()
 {
-	if (Blueprint->IsA<UWidgetBlueprint>())
-	{
-		bRebuildRequested = true;
-		bRefreshRequested = true;
-	}
+	bRebuildRequested = true;
+	bRefreshRequested = true;
 }
 
 void SPaletteView::HandleOnHotReload(bool bWasTriggeredAutomatically)
diff --git a/Engine/Source/Editor/UMGEditor/Private/Palette/SPaletteView.h b/Engine/Source/Editor/UMGEditor/Private/Palette/SPaletteView.h
--- a/Engine/Source/Editor/UMGEditor/Private/Palette/SPaletteView.h
+++ b/Engine/Source/Editor/UMGEditor/Private/Palette/SPaletteView.h
@@ -69,7 +69,7 @@ private:
 	void TransformWidgetViewModelToString(TSharedPtr<FWidgetViewModel> WidgetViewModel, OUT TArray< FString >& Array);
 
 	/** Requests a rebuild of the widget list if a widget blueprint was compiled */
-	void HandleOnBlueprintCompiled(UBlueprint* Blueprint);
+	void OnBlueprintReinstanced();
 
 	/** Requests a rebuild of the widget list */
 	void HandleOnHotReload(bool bWasTriggeredAutomatically);
diff --git a/Engine/Source/Editor/UMGEditor/Private/WidgetBlueprintCompiler.cpp b/Engine/Source/Editor/UMGEditor/Private/WidgetBlueprintCompiler.cpp
--- a/Engine/Source/Editor/UMGEditor/Private/WidgetBlueprintCompiler.cpp
+++ b/Engine/Source/Editor/UMGEditor/Private/WidgetBlueprintCompiler.cpp
@@ -13,8 +13,6 @@
 #include "Kismet2/BlueprintEditorUtils.h"
 #include "NamedSlot.h"
 
-FOnWidgetBlueprintCompiled FWidgetBlueprintCompiler::OnWidgetBlueprintCompiled = FOnWidgetBlueprintCompiled();
-
 ///-------------------------------------------------------------
 
 FWidgetBlueprintCompiler::FWidgetBlueprintCompiler(UWidgetBlueprint* SourceSketch, FCompilerResultsLog& InMessageLog, const FKismetCompilerOptions& InCompilerOptions, TArray<UObject*>* InObjLoaded)
@@ -283,8 +281,6 @@ void FWidgetBlueprintCompiler::Compile()
 {
 	Super::Compile();
 
-	OnWidgetBlueprintCompiled.Broadcast( WidgetBlueprint() );
-
 	WidgetToMemberVariableMap.Empty();
 }
 
diff --git a/Engine/Source/Editor/UMGEditor/Private/WidgetBlueprintCompiler.h b/Engine/Source/Editor/UMGEditor/Private/WidgetBlueprintCompiler.h
--- a/Engine/Source/Editor/UMGEditor/Private/WidgetBlueprintCompiler.h
+++ b/Engine/Source/Editor/UMGEditor/Private/WidgetBlueprintCompiler.h
@@ -10,17 +10,12 @@ class UWidgetBlueprintGeneratedClass;
 //////////////////////////////////////////////////////////////////////////
 // FWidgetBlueprintCompiler
 
-DECLARE_MULTICAST_DELEGATE_OneParam( FOnWidgetBlueprintCompiled, UBlueprint* )
-
 class FWidgetBlueprintCompiler : public FKismetCompilerContext
 {
 protected:
 	typedef FKismetCompilerContext Super;
 
 public:
-	/** Delegate called when a widget blueprint is compiled */
-	static FOnWidgetBlueprintCompiled OnWidgetBlueprintCompiled;
-
 	FWidgetBlueprintCompiler(UWidgetBlueprint* SourceSketch, FCompilerResultsLog& InMessageLog, const FKismetCompilerOptions& InCompilerOptions, TArray<UObject*>* InObjLoaded);
 	virtual ~FWidgetBlueprintCompiler();
 
diff --git a/Engine/Source/Editor/UnrealEd/Classes/Editor/EditorEngine.h b/Engine/Source/Editor/UnrealEd/Classes/Editor/EditorEngine.h
--- a/Engine/Source/Editor/UnrealEd/Classes/Editor/EditorEngine.h
+++ b/Engine/Source/Editor/UnrealEd/Classes/Editor/EditorEngine.h
@@ -547,6 +547,13 @@ public:
 	/**	Broadcasts that a blueprint just finished compiling. THIS SHOULD NOT BE PUBLIC */
 	void BroadcastBlueprintCompiled() { BlueprintCompiledEvent.Broadcast(); }
 
+	/** Called by the blueprint compiler after a blueprint has been compiled and all instances replaced, but prior to garbage collection. */
+	DECLARE_EVENT(UEditorEngine, FBlueprintReinstanced);
+	FBlueprintReinstanced& OnBlueprintReinstanced() { return BlueprintReinstanced; }
+
+	/**	Broadcasts that a blueprint just finished being reinstanced. THIS SHOULD NOT BE PUBLIC */
+	void BroadcastBlueprintReinstanced() { BlueprintReinstanced.Broadcast(); }
+
 	/** Called when uobjects have been replaced to allow others a chance to fix their references. */
 	typedef TMap<UObject*, UObject*> ReplacementObjectMap;
 	DECLARE_EVENT_OneParam( UEditorEngine, FObjectsReplacedEvent, const ReplacementObjectMap& );
@@ -2595,6 +2602,9 @@ private:
 	/** Delegate broadcast when blueprint is compiled */
 	FBlueprintCompiledEvent BlueprintCompiledEvent;
 
+	/** Delegate broadcast when blueprint is reinstanced */
+	FBlueprintReinstanced BlueprintReinstanced;
+
 	/** Delegate broadcast when objects have been replaced (e.g on blueprint compile) */
 	FObjectsReplacedEvent ObjectsReplacedEvent;
 
@@ -2700,6 +2710,9 @@ protected:
 	void HandleLaunchCanceled(double TotalTime, bool bHasCode, TWeakPtr<SNotificationItem> NotificationItemPtr);
 	void HandleLaunchCompleted(bool Succeeded, double TotalTime, int32 ErrorCode, bool bHasCode, TWeakPtr<SNotificationItem> NotificationItemPtr, TSharedPtr<class FMessageLog> MessageLog);
 
+	// Handle requests from slate application to open assets.
+	bool HandleOpenAsset(UObject* Asset);
+
 public:
 	/** True if world assets are enabled */
 	static bool IsUsingWorldAssets();
@@ -2711,5 +2724,3 @@ private:
 	/** Gets the init values for worlds opened via Map_Load in the editor */
 	UWorld::InitializationValues GetEditorWorldInitializationValues() const;
 };
-
-
diff --git a/Engine/Source/Editor/UnrealEd/Private/Editor.cpp b/Engine/Source/Editor/UnrealEd/Private/Editor.cpp
--- a/Engine/Source/Editor/UnrealEd/Private/Editor.cpp
+++ b/Engine/Source/Editor/UnrealEd/Private/Editor.cpp
@@ -375,35 +375,35 @@ void UEditorEngine::InitEditor(IEngineLoop* InEngineLoop)
 	UEngine::Init(InEngineLoop);
 
 	// Specify "-ForceLauncher" on the command-line to always open the launcher, even in unusual cases.  This is useful for debugging the Launcher startup.
-	const bool bForceLauncherToOpen = FParse::Param( FCommandLine::Get(),TEXT( "ForceLauncher" ) );
+	const bool bForceLauncherToOpen = FParse::Param(FCommandLine::Get(), TEXT("ForceLauncher"));
 
-	if( bForceLauncherToOpen ||
+	if ( bForceLauncherToOpen ||
 		( !FEngineBuildSettings::IsInternalBuild() &&
-		  !FEngineBuildSettings::IsPerforceBuild() && 
-		  !FPlatformMisc::IsDebuggerPresent() &&	// Don't spawn launcher while running in the Visual Studio debugger by default
-		  !FApp::IsBenchmarking() &&
-		  !GIsDemoMode && 
-		  !IsRunningCommandlet() &&
-		  !FPlatformProcess::IsApplicationRunning(TEXT("UnrealEngineLauncher") ) &&
-		  !FPlatformProcess::IsApplicationRunning(TEXT("Unreal Engine Launcher") ) ) )
+		!FEngineBuildSettings::IsPerforceBuild() &&
+		!FPlatformMisc::IsDebuggerPresent() &&	// Don't spawn launcher while running in the Visual Studio debugger by default
+		!FApp::IsBenchmarking() &&
+		!GIsDemoMode &&
+		!IsRunningCommandlet() &&
+		!FPlatformProcess::IsApplicationRunning(TEXT("UnrealEngineLauncher")) &&
+		!FPlatformProcess::IsApplicationRunning(TEXT("Unreal Engine Launcher")) ) )
 	{
 		IDesktopPlatform* DesktopPlatform = FDesktopPlatformModule::Get();
-		if( DesktopPlatform != NULL )
+		if ( DesktopPlatform != NULL )
 		{
-			DesktopPlatform->OpenLauncher( false, TEXT("") );
+			DesktopPlatform->OpenLauncher(false, TEXT(""));
 		}
 	}
-	
+
 	// Create selection sets.
 	PrivateInitSelectedSets();
-	
+
 	// Set slate options
 	FMultiBoxSettings::UseSmallToolBarIcons = TAttribute<bool>::Create(TAttribute<bool>::FGetter::CreateStatic(&GetSmallToolBarIcons));
 	FMultiBoxSettings::DisplayMultiboxHooks = TAttribute<bool>::Create(TAttribute<bool>::FGetter::CreateStatic(&GetDisplayMultiboxHooks));
-	
-	if (FSlateApplication::IsInitialized())
+
+	if ( FSlateApplication::IsInitialized() )
 	{
-		FSlateApplication::Get().GetRenderer()->SetColorVisionDeficiencyType((uint32)(GetDefault<UEditorStyleSettings>()->ColorVisionDeficiencyPreviewType.GetValue()));
+		FSlateApplication::Get().GetRenderer()->SetColorVisionDeficiencyType((uint32)( GetDefault<UEditorStyleSettings>()->ColorVisionDeficiencyPreviewType.GetValue() ));
 		FSlateApplication::Get().EnableMenuAnimations(GetDefault<UEditorStyleSettings>()->bEnableWindowAnimations);
 	}
 
@@ -416,17 +416,17 @@ void UEditorEngine::InitEditor(IEngineLoop* InEngineLoop)
 	GEngine->SelectionHighlightIntensity = ViewportSettings->SelectionHighlightIntensity;
 	GEngine->BSPSelectionHighlightIntensity = ViewportSettings->BSPSelectionHighlightIntensity;
 	GEngine->HoverHighlightIntensity = ViewportSettings->HoverHighlightIntensity;
-	
+
 	// Set navigation system property indicating whether navigation is supposed to rebuild automatically 
 	FWorldContext &EditorContext = GEditor->GetEditorWorldContext();
-	UNavigationSystem::SetNavigationAutoUpdateEnabled(GetDefault<ULevelEditorMiscSettings>()->bNavigationAutoUpdate, EditorContext.World()->GetNavigationSystem() );
+	UNavigationSystem::SetNavigationAutoUpdateEnabled(GetDefault<ULevelEditorMiscSettings>()->bNavigationAutoUpdate, EditorContext.World()->GetNavigationSystem());
 
 	// Allocate temporary model.
-	TempModel = new UModel( FObjectInitializer(),NULL, 1 );
-	ConversionTempModel = new UModel( FObjectInitializer(),NULL, 1 );
+	TempModel = new UModel(FObjectInitializer(), NULL, 1);
+	ConversionTempModel = new UModel(FObjectInitializer(), NULL, 1);
 
 	// create the timer manager
-	TimerManager = MakeShareable( new FTimerManager() );
+	TimerManager = MakeShareable(new FTimerManager());
 
 	// Settings.
 	FBSPOps::GFastRebuild = 0;
@@ -446,6 +446,14 @@ void UEditorEngine::InitEditor(IEngineLoop* InEngineLoop)
 
 	// Load any modules that might be required by commandlets
 	FModuleManager::Get().LoadModule(TEXT("OnlineBlueprintSupport"));
+
+	// Setup a delegate to handle requests for opening assets
+	FSlateApplication::Get().SetWidgetReflectorAssetAccessDelegate(FAccessAsset::CreateUObject(this, &UEditorEngine::HandleOpenAsset));
+}
+
+bool UEditorEngine::HandleOpenAsset(UObject* Asset)
+{
+	return FAssetEditorManager::Get().OpenEditorForAsset(Asset);
 }
 
 void UEditorEngine::HandleSettingChanged( FName Name )
@@ -3992,26 +4000,26 @@ bool UEditorEngine::AttachActorToComponent(AActor* ParentActor, AActor* ChildAct
 
 	if(SkeletalMeshComponent && SkeletalMeshComponent->SkeletalMesh)
 	{
-	    USkeletalMeshSocket const* const Socket = SkeletalMeshComponent->SkeletalMesh->FindSocket(SocketName);
-	    if (Socket)
-	    {
-		    bAttachedToSocket = Socket->AttachActor(ChildActor, SkeletalMeshComponent);
-	    }
-	    else
-	    {
-		    // now search bone, if bone exists, snap to the bone
-		    int32 BoneIndex = SkeletalMeshComponent->SkeletalMesh->RefSkeleton.FindBoneIndex(SocketName);
-		    if (BoneIndex != INDEX_NONE)
-		    {
-			    ChildActor->GetRootComponent()->SnapTo(ParentActor->GetRootComponent(), SocketName);
-			    bAttachedToSocket = true;
-    
-    #if WITH_EDITOR
-			    ChildActor->PreEditChange(NULL);
-			    ChildActor->PostEditChange();
-    #endif // WITH_EDITOR
-		    }
-	    }
+		USkeletalMeshSocket const* const Socket = SkeletalMeshComponent->SkeletalMesh->FindSocket(SocketName);
+		if (Socket)
+		{
+			bAttachedToSocket = Socket->AttachActor(ChildActor, SkeletalMeshComponent);
+		}
+		else
+		{
+			// now search bone, if bone exists, snap to the bone
+			int32 BoneIndex = SkeletalMeshComponent->SkeletalMesh->RefSkeleton.FindBoneIndex(SocketName);
+			if (BoneIndex != INDEX_NONE)
+			{
+				ChildActor->GetRootComponent()->SnapTo(ParentActor->GetRootComponent(), SocketName);
+				bAttachedToSocket = true;
+	
+	#if WITH_EDITOR
+				ChildActor->PreEditChange(NULL);
+				ChildActor->PostEditChange();
+	#endif // WITH_EDITOR
+			}
+		}
 	}
 
 	return bAttachedToSocket;
@@ -4023,11 +4031,11 @@ bool UEditorEngine::AttachActorToComponent(AActor* ChildActor, UStaticMeshCompon
 
 	if(StaticMeshComponent && StaticMeshComponent->StaticMesh)
 	{
-	    UStaticMeshSocket const* const Socket = StaticMeshComponent->StaticMesh->FindSocket(SocketName);
-	    if (Socket)
-	    {
-		    bAttachedToSocket = Socket->AttachActor(ChildActor, StaticMeshComponent);
-	    }
+		UStaticMeshSocket const* const Socket = StaticMeshComponent->StaticMesh->FindSocket(SocketName);
+		if (Socket)
+		{
+			bAttachedToSocket = Socket->AttachActor(ChildActor, StaticMeshComponent);
+		}
 	}
 
 	return bAttachedToSocket;
diff --git a/Engine/Source/Editor/UnrealEd/Private/Kismet2/Kismet2.cpp b/Engine/Source/Editor/UnrealEd/Private/Kismet2/Kismet2.cpp
--- a/Engine/Source/Editor/UnrealEd/Private/Kismet2/Kismet2.cpp
+++ b/Engine/Source/Editor/UnrealEd/Private/Kismet2/Kismet2.cpp
@@ -420,6 +420,12 @@ void FKismetEditorUtilities::CompileBlueprint(UBlueprint* BlueprintObj, bool bIs
 		static const FBoolConfigValueHelper ReinstanceOnlyWhenNecessary(TEXT("Kismet"), TEXT("bReinstanceOnlyWhenNecessary"), GEngineIni);
 		ReinstanceHelper.ReinstanceObjects(!ReinstanceOnlyWhenNecessary);
 
+		// Notify everyone a blueprint has been compiled and reinstanced, but before GC so they can perform any final cleanup.
+		if ( GEditor )
+		{
+			GEditor->BroadcastBlueprintReinstanced();
+		}
+
 		if (!bSkipGarbageCollection)
 		{
 			BP_SCOPED_COMPILER_EVENT_NAME(TEXT("Garbage Collection"));
diff --git a/Engine/Source/Runtime/UMG/Private/UWidgetNavigation.cpp b/Engine/Source/Runtime/UMG/Private/UWidgetNavigation.cpp
new file mode 100644
--- /dev/null
+++ b/Engine/Source/Runtime/UMG/Private/UWidgetNavigation.cpp
@@ -0,0 +1,11 @@
+// Copyright 1998-2014 Epic Games, Inc. All Rights Reserved.
+
+#include "UMGPrivatePCH.h"
+
+/////////////////////////////////////////////////////
+// UWidgetNavigation
+
+UWidgetNavigation::UWidgetNavigation(const FObjectInitializer& ObjectInitializer)
+	: Super(ObjectInitializer)
+{
+}
diff --git a/Engine/Source/Runtime/UMG/Private/UserWidget.cpp b/Engine/Source/Runtime/UMG/Private/UserWidget.cpp
--- a/Engine/Source/Runtime/UMG/Private/UserWidget.cpp
+++ b/Engine/Source/Runtime/UMG/Private/UserWidget.cpp
@@ -46,6 +46,7 @@ void UUserWidget::Initialize()
 			BGClass->InitializeWidget(this);
 		}
 
+		// Map the named slot bindings to the available slots.
 		WidgetTree->ForEachWidget([&] (UWidget* Widget) {
 			if ( UNamedSlot* NamedWidet = Cast<UNamedSlot>(Widget) )
 			{
@@ -55,11 +56,14 @@ void UUserWidget::Initialize()
 					{
 						NamedWidet->ClearChildren();
 						NamedWidet->AddChild(Binding.Content);
-						continue;
+						return;
 					}
 				}
 			}
 		});
+
+		// Clear the named slot bindings table.
+		NamedSlotBindings.Reset();
 	}
 }
 
diff --git a/Engine/Source/Runtime/UMG/Public/Blueprint/WidgetNavigation.h b/Engine/Source/Runtime/UMG/Public/Blueprint/WidgetNavigation.h
--- a/Engine/Source/Runtime/UMG/Public/Blueprint/WidgetNavigation.h
+++ b/Engine/Source/Runtime/UMG/Public/Blueprint/WidgetNavigation.h
@@ -29,7 +29,7 @@ public:
 UCLASS()
 class UMG_API UWidgetNavigation : public UObject
 {
-	GENERATED_BODY()
+	GENERATED_UCLASS_BODY()
 	
 public:
 	/**  */
