[UETOOL-143] UMG - Fixing several crashes in the hierarchy view from edge cases not being accounted for, or results from the refactor of the components flat list being removed from the UserWidget.

diff --git a/Engine/Source/Editor/UMGEditor/Private/Hierarchy/SHierarchyViewItem.cpp b/Engine/Source/Editor/UMGEditor/Private/Hierarchy/SHierarchyViewItem.cpp
--- a/Engine/Source/Editor/UMGEditor/Private/Hierarchy/SHierarchyViewItem.cpp
+++ b/Engine/Source/Editor/UMGEditor/Private/Hierarchy/SHierarchyViewItem.cpp
@@ -36,7 +36,9 @@ class FHierarchyWidgetDragDropOp : public FDecoratedDragDropOp
 public:
 	DRAG_DROP_OPERATOR_TYPE(FHierarchyWidgetDragDropOp, FDecoratedDragDropOp)
 
-		virtual ~FHierarchyWidgetDragDropOp();
+	virtual ~FHierarchyWidgetDragDropOp();
+
+	virtual void OnDrop(bool bDropWasHandled, const FPointerEvent& MouseEvent) override;
 
 	/** The slot properties for the old slot the widget was in, is used to attempt to reapply the same layout information */
 	TMap<FName, FString> ExportedSlotProperties;
@@ -87,6 +89,14 @@ FHierarchyWidgetDragDropOp::~FHierarchyWidgetDragDropOp()
 	delete Transaction;
 }
 
+void FHierarchyWidgetDragDropOp::OnDrop(bool bDropWasHandled, const FPointerEvent& MouseEvent)
+{
+	if ( !bDropWasHandled )
+	{
+		Transaction->Cancel();
+	}
+}
+
 //////////////////////////////////////////////////////////////////////////
 
 FReply ProcessHierarchyDragDrop(const FDragDropEvent& DragDropEvent, bool bIsDrop, UWidgetBlueprint* Blueprint, FWidgetReference TargetItem)
@@ -120,7 +130,6 @@ FReply ProcessHierarchyDragDrop(const FDragDropEvent& DragDropEvent, bool bIsDro
 				FBlueprintEditorUtils::MarkBlueprintAsStructurallyModified(Blueprint);
 			}
 		}
-		//TODO Are we adding to a named slot?
 		else
 		{
 			TemplateDragDropOp->SetCursorOverride(EMouseCursor::SlashedCircle);
@@ -133,63 +142,69 @@ FReply ProcessHierarchyDragDrop(const FDragDropEvent& DragDropEvent, bool bIsDro
 	if ( HierarchyDragDropOp.IsValid() )
 	{
 		HierarchyDragDropOp->SetCursorOverride(TOptional<EMouseCursor::Type>());
+		HierarchyDragDropOp->ResetToDefaultToolTip();
 
 		// If the target item is valid we're dealing with a normal widget in the hierarchy, otherwise we should assume it's
 		// the null case and we should be adding it as the root widget.
 		if ( TargetItem.IsValid() )
 		{
-			const bool bIsPanel = TargetItem.GetTemplate()->IsA(UPanelWidget::StaticClass());
 			const bool bIsDraggedObject = TargetItem.GetTemplate() == HierarchyDragDropOp->Widget.GetTemplate();
-			const bool bIsChildOfDraggedObject = TargetItem.GetTemplate()->IsChildOf(HierarchyDragDropOp->Widget.GetTemplate());
-
-			if ( bIsPanel && !bIsDraggedObject && !bIsChildOfDraggedObject )
+			if ( bIsDraggedObject )
 			{
-				UWidget* Widget = bIsDrop ? HierarchyDragDropOp->Widget.GetTemplate() : HierarchyDragDropOp->Widget.GetPreview();
+				HierarchyDragDropOp->SetCursorOverride(EMouseCursor::SlashedCircle);
+				return FReply::Handled();
+			}
 
-				if ( Widget->GetParent() )
-				{
-					Widget->GetParent()->RemoveChild(Widget);
-				}
+			UPanelWidget* NewParent = Cast<UPanelWidget>(TargetItem.GetTemplate());
+			if ( !NewParent )
+			{
+				HierarchyDragDropOp->SetCursorOverride(EMouseCursor::SlashedCircle);
+				HierarchyDragDropOp->CurrentHoverText = LOCTEXT("CantHaveChildren", "Widget can't have children.");
+				return FReply::Handled();
+			}
 
-				UPanelWidget* NewParent = Cast<UPanelWidget>(bIsDrop ? TargetItem.GetTemplate() : TargetItem.GetPreview());
+			if ( !NewParent->CanHaveMultipleChildren() && NewParent->GetChildrenCount() > 0 )
+			{
+				HierarchyDragDropOp->SetCursorOverride(EMouseCursor::SlashedCircle);
+				HierarchyDragDropOp->CurrentHoverText = LOCTEXT("NoAdditionalChildren", "Widget can't accept additional children.");
+				return FReply::Handled();
+			}
 
-				if ( bIsDrop && NewParent != HierarchyDragDropOp->WidgetParent )
-				{
-					NewParent->SetFlags(RF_Transactional);
-					NewParent->Modify();
-				}
+			UWidget* TemplateWidget = HierarchyDragDropOp->Widget.GetTemplate();
 
-				if ( UPanelSlot* Slot = NewParent->AddChild(Widget) )
+			// Verify that the new location we're placing the widget is not inside of its existing children.
+			bool bFoundNewParentInChildSet = false;
+			Blueprint->WidgetTree->ForWidgetAndChildren(TemplateWidget, [&] (UWidget* Widget) {
+				if ( NewParent == Widget )
 				{
-					FWidgetBlueprintEditorUtils::ImportPropertiesFromText(Slot, HierarchyDragDropOp->ExportedSlotProperties);
-
-					Slot->SetDesiredPosition(FVector2D::ZeroVector);
-
-					if ( bIsDrop )
-					{
-						FBlueprintEditorUtils::MarkBlueprintAsStructurallyModified(Blueprint);
-					}
-					else
-					{
-						TSharedPtr<SWidget> CachedWidget = TargetItem.GetPreview()->GetCachedWidget();
-						if ( CachedWidget.IsValid() )
-						{
-							CachedWidget->SlatePrepass();
-						}
-					}
+					bFoundNewParentInChildSet = true;
 				}
-				else
-				{
-					HierarchyDragDropOp->SetCursorOverride(EMouseCursor::SlashedCircle);
+			});
 
-					// TODO UMG ERROR Slot can not be created because maybe the max children has been reached.
-					//          Maybe we can traverse the hierarchy and add it to the first parent that will accept it?
-				}
-			}
-			else
+			if ( bFoundNewParentInChildSet )
 			{
 				HierarchyDragDropOp->SetCursorOverride(EMouseCursor::SlashedCircle);
+				HierarchyDragDropOp->CurrentHoverText = LOCTEXT("CantMakeWidgetChildOfChildren", "Can't make widget a child of its children.");
+				return FReply::Handled();
+			}
+
+			if ( bIsDrop )
+			{
+				TemplateWidget->RemoveFromParent();
+
+				NewParent->SetFlags(RF_Transactional);
+				NewParent->Modify();
+
+				UPanelSlot* NewSlot = NewParent->AddChild(TemplateWidget);
+				check(NewSlot);
+
+				// Import the old slot properties
+				FWidgetBlueprintEditorUtils::ImportPropertiesFromText(NewSlot, HierarchyDragDropOp->ExportedSlotProperties);
+				NewSlot->SetDesiredPosition(FVector2D::ZeroVector);
+				FBlueprintEditorUtils::MarkBlueprintAsStructurallyModified(Blueprint);
 			}
+
+			return FReply::Handled();
 		}
 		else
 		{
@@ -492,8 +507,6 @@ FReply FNamedSlotModel::HandleDrop(FDragDropEvent const& DragDropEvent)
 {
 	UWidgetBlueprint* Blueprint = BlueprintEditor.Pin()->GetWidgetBlueprintObj();
 
-	bool bIsDrop = true;
-
 	// Is this a drag/drop op to create a new widget in the tree?
 	TSharedPtr<FWidgetTemplateDragDropOp> TemplateDragDropOp = DragDropEvent.GetOperationAs<FWidgetTemplateDragDropOp>();
 	if ( TemplateDragDropOp.IsValid() )
@@ -515,6 +528,23 @@ FReply FNamedSlotModel::HandleDrop(FDragDropEvent const& DragDropEvent)
 	return FReply::Unhandled();
 }
 
+FReply FNamedSlotModel::HandleDragDetected(const FGeometry& MyGeometry, const FPointerEvent& MouseEvent)
+{
+	if ( INamedSlotInterface* NamedSlotHost = Cast<INamedSlotInterface>(Item.GetTemplate()) )
+	{
+		// Only assign content to the named slot if it is null.
+		if ( UWidget* Content = NamedSlotHost->GetContentForSlot(SlotName) )
+		{
+			FWidgetReference ContentRef = BlueprintEditor.Pin()->GetReferenceFromTemplate(Content);
+			check(ContentRef.IsValid());
+
+			return FReply::Handled().BeginDragDrop(FHierarchyWidgetDragDropOp::New(BlueprintEditor.Pin()->GetWidgetBlueprintObj(), ContentRef));
+		}
+	}
+
+	return FReply::Unhandled();
+}
+
 //////////////////////////////////////////////////////////////////////////
 
 FHierarchyWidget::FHierarchyWidget(FWidgetReference InItem, TSharedPtr<FWidgetBlueprintEditor> InBlueprintEditor)
diff --git a/Engine/Source/Editor/UMGEditor/Private/Hierarchy/SHierarchyViewItem.h b/Engine/Source/Editor/UMGEditor/Private/Hierarchy/SHierarchyViewItem.h
--- a/Engine/Source/Editor/UMGEditor/Private/Hierarchy/SHierarchyViewItem.h
+++ b/Engine/Source/Editor/UMGEditor/Private/Hierarchy/SHierarchyViewItem.h
@@ -121,9 +121,9 @@ public:
 
 	virtual void OnSelection();
 	
-	virtual FReply OnDragOver(const FGeometry& MyGeometry, const FDragDropEvent& DragDropEvent);
-
-	virtual FReply HandleDrop(FDragDropEvent const& DragDropEvent);
+	virtual FReply OnDragOver(const FGeometry& MyGeometry, const FDragDropEvent& DragDropEvent) override;
+	virtual FReply HandleDrop(FDragDropEvent const& DragDropEvent) override;
+	virtual FReply HandleDragDetected(const FGeometry& MyGeometry, const FPointerEvent& MouseEvent) override;
 
 protected:
 	virtual void GetChildren(TArray< TSharedPtr<FHierarchyModel> >& Children) override;
@@ -172,7 +172,12 @@ public:
 
 	virtual bool IsVisible() const override
 	{
-		return !Item.GetTemplate()->bHiddenInDesigner;
+		if ( UWidget* TemplateWidget = Item.GetTemplate() )
+		{
+			return !TemplateWidget->bHiddenInDesigner;
+		}
+
+		return true;
 	}
 
 	virtual bool CanControlVisibility() const override
