UMG - Prevent attempting to load the widget blueprint on client and server in cooked builds.

diff --git a/Engine/Source/Editor/UMGEditor/Private/WidgetBlueprint.cpp b/Engine/Source/Editor/UMGEditor/Private/WidgetBlueprint.cpp
--- a/Engine/Source/Editor/UMGEditor/Private/WidgetBlueprint.cpp
+++ b/Engine/Source/Editor/UMGEditor/Private/WidgetBlueprint.cpp
@@ -209,15 +209,14 @@ void UWidgetBlueprint::GetReparentingRules(TSet< const UClass* >& AllowedChildre
 	AllowedChildrenOfClasses.Add( UUserWidget::StaticClass() );
 }
 
-void UWidgetBlueprint::PostLoadSubobjects(FObjectInstancingGraph* OuterInstanceGraph)
+bool UWidgetBlueprint::NeedsLoadForClient() const
 {
-	Super::PostLoadSubobjects(OuterInstanceGraph);
+	return false;
+}
 
-	UWidgetBlueprintGeneratedClass* BPGClass = Cast<UWidgetBlueprintGeneratedClass>(*GeneratedClass);
-	if ( BPGClass )
-	{
-		//BPGClass->WidgetTree = WidgetTree;
-	}
+bool UWidgetBlueprint::NeedsLoadForServer() const
+{
+	return false;
 }
 
 #undef LOCTEXT_NAMESPACE 
\ No newline at end of file
diff --git a/Engine/Source/Editor/UMGEditor/Public/WidgetBlueprint.h b/Engine/Source/Editor/UMGEditor/Public/WidgetBlueprint.h
--- a/Engine/Source/Editor/UMGEditor/Public/WidgetBlueprint.h
+++ b/Engine/Source/Editor/UMGEditor/Public/WidgetBlueprint.h
@@ -81,6 +81,14 @@ class UMGEDITOR_API UWidgetBlueprint : public UBlueprint
 	GENERATED_UCLASS_BODY()
 
 public:
+
+	// Begin UObject interface
+	virtual bool NeedsLoadForClient() const override;
+	virtual bool NeedsLoadForServer() const override;
+	// End UObject interface
+
+public:
+
 	/** A tree of the widget templates to be created */
 	UPROPERTY()
 	class UWidgetTree* WidgetTree;
@@ -94,9 +102,10 @@ public:
 	UPROPERTY()
 	TArray<UWidgetAnimation*> Animations;
 
+public:
+
 	/** UObject interface */
 	virtual void PostLoad() override;
-	virtual void PostLoadSubobjects(FObjectInstancingGraph* OuterInstanceGraph) override;
 	
 	// UBlueprint interface
 	virtual UClass* GetBlueprintClass() const override;
