UMG - Fixing unity build on mac.

diff --git a/Engine/Source/Editor/UMGEditor/Private/Designer/SDesignerView.cpp b/Engine/Source/Editor/UMGEditor/Private/Designer/SDesignerView.cpp
--- a/Engine/Source/Editor/UMGEditor/Private/Designer/SDesignerView.cpp
+++ b/Engine/Source/Editor/UMGEditor/Private/Designer/SDesignerView.cpp
@@ -40,6 +40,7 @@
 #include "ScopedTransaction.h"
 #include "Settings/WidgetDesignerSettings.h"
 #include "Components/CanvasPanelSlot.h"
+#include "Components/NamedSlot.h"
 
 #define LOCTEXT_NAMESPACE "UMG"
 
