UMG - Fixing a crash when the sequencer tab is removed from UMG, future runs will fail to initialize the binding object and you'll get a null deref crash.

diff --git a/Engine/Source/Editor/UMGEditor/Private/WidgetBlueprintEditor.cpp b/Engine/Source/Editor/UMGEditor/Private/WidgetBlueprintEditor.cpp
--- a/Engine/Source/Editor/UMGEditor/Private/WidgetBlueprintEditor.cpp
+++ b/Engine/Source/Editor/UMGEditor/Private/WidgetBlueprintEditor.cpp
@@ -68,7 +68,11 @@ void FWidgetBlueprintEditor::InitWidgetBlueprintEditor(const EToolkitMode::Type
 
 	UpdatePreview(GetWidgetBlueprintObj(), true);
 
-	SequencerObjectBindingManager->InitPreviewObjects();
+	// If the user has close the sequencer tab, this will not be initialized.
+	if ( SequencerObjectBindingManager.IsValid() )
+	{
+		SequencerObjectBindingManager->InitPreviewObjects();
+	}
 
 	DesignerCommandList = MakeShareable(new FUICommandList);
 
