UMG - Selecting widgets now after placing them.

diff --git a/Engine/Source/Editor/UMGEditor/Private/Designer/SDesignerView.cpp b/Engine/Source/Editor/UMGEditor/Private/Designer/SDesignerView.cpp
--- a/Engine/Source/Editor/UMGEditor/Private/Designer/SDesignerView.cpp
+++ b/Engine/Source/Editor/UMGEditor/Private/Designer/SDesignerView.cpp
@@ -1680,11 +1680,16 @@ FReply SDesignerView::OnDrop(const FGeometry& MyGeometry, const FDragDropEvent&
 	}
 	
 	const bool bIsPreview = false;
-	UWidget* Widget = ProcessDropAndAddWidget(MyGeometry, DragDropEvent, bIsPreview);
-	if ( Widget )
+	UWidget* NewTemplateWidget = ProcessDropAndAddWidget(MyGeometry, DragDropEvent, bIsPreview);
+	if ( NewTemplateWidget )
 	{
 		FBlueprintEditorUtils::MarkBlueprintAsStructurallyModified(BP);
 
+		TSet<FWidgetReference> SelectedTemplates;
+		SelectedTemplates.Add(BlueprintEditor.Pin()->GetReferenceFromTemplate(NewTemplateWidget));
+		 
+		BlueprintEditor.Pin()->SelectWidgets(SelectedTemplates, false);
+
 		// Regenerate extension widgets now that we've finished moving or placing the widget.
 		CreateExtensionWidgetsForSelection();
 
diff --git a/Engine/Source/Editor/UMGEditor/Private/Hierarchy/SHierarchyViewItem.cpp b/Engine/Source/Editor/UMGEditor/Private/Hierarchy/SHierarchyViewItem.cpp
--- a/Engine/Source/Editor/UMGEditor/Private/Hierarchy/SHierarchyViewItem.cpp
+++ b/Engine/Source/Editor/UMGEditor/Private/Hierarchy/SHierarchyViewItem.cpp
@@ -99,7 +99,7 @@ void FHierarchyWidgetDragDropOp::OnDrop(bool bDropWasHandled, const FPointerEven
 
 //////////////////////////////////////////////////////////////////////////
 
-FReply ProcessHierarchyDragDrop(const FDragDropEvent& DragDropEvent, bool bIsDrop, UWidgetBlueprint* Blueprint, FWidgetReference TargetItem)
+FReply ProcessHierarchyDragDrop(const FDragDropEvent& DragDropEvent, bool bIsDrop, TSharedPtr<FWidgetBlueprintEditor> BlueprintEditor, UWidgetBlueprint* Blueprint, FWidgetReference TargetItem)
 {
 	// Is this a drag/drop op to create a new widget in the tree?
 	TSharedPtr<FWidgetTemplateDragDropOp> TemplateDragDropOp = DragDropEvent.GetOperationAs<FWidgetTemplateDragDropOp>();
@@ -202,6 +202,11 @@ FReply ProcessHierarchyDragDrop(const FDragDropEvent& DragDropEvent, bool bIsDro
 				FWidgetBlueprintEditorUtils::ImportPropertiesFromText(NewSlot, HierarchyDragDropOp->ExportedSlotProperties);
 				NewSlot->SetDesiredPosition(FVector2D::ZeroVector);
 				FBlueprintEditorUtils::MarkBlueprintAsStructurallyModified(Blueprint);
+
+				TSet<FWidgetReference> SelectedTemplates;
+				SelectedTemplates.Add(BlueprintEditor->GetReferenceFromTemplate(TemplateWidget));
+
+				BlueprintEditor->SelectWidgets(SelectedTemplates, false);
 			}
 
 			return FReply::Handled();
@@ -383,14 +388,14 @@ FReply FHierarchyRoot::OnDragOver(const FGeometry& MyGeometry, const FDragDropEv
 {
 	UWidgetBlueprint* Blueprint = BlueprintEditor.Pin()->GetWidgetBlueprintObj();
 	bool bIsDrop = false;
-	return ProcessHierarchyDragDrop(DragDropEvent, bIsDrop, Blueprint, FWidgetReference());
+	return ProcessHierarchyDragDrop(DragDropEvent, bIsDrop, BlueprintEditor.Pin(), Blueprint, FWidgetReference());
 }
 
 FReply FHierarchyRoot::HandleDrop(FDragDropEvent const& DragDropEvent)
 {
 	UWidgetBlueprint* Blueprint = BlueprintEditor.Pin()->GetWidgetBlueprintObj();
 	bool bIsDrop = true;
-	return ProcessHierarchyDragDrop(DragDropEvent, bIsDrop, Blueprint, FWidgetReference());
+	return ProcessHierarchyDragDrop(DragDropEvent, bIsDrop, BlueprintEditor.Pin(), Blueprint, FWidgetReference());
 }
 
 //////////////////////////////////////////////////////////////////////////
@@ -519,6 +524,11 @@ FReply FNamedSlotModel::HandleDrop(FDragDropEvent const& DragDropEvent)
 				UWidget* Widget = TemplateDragDropOp->Template->Create(Blueprint->WidgetTree);
 				NamedSlotHost->SetContentForSlot(SlotName, Widget);
 				FBlueprintEditorUtils::MarkBlueprintAsStructurallyModified(Blueprint);
+
+				TSet<FWidgetReference> SelectedTemplates;
+				SelectedTemplates.Add(BlueprintEditor.Pin()->GetReferenceFromTemplate(Widget));
+
+				BlueprintEditor.Pin()->SelectWidgets(SelectedTemplates, false);
 			}
 
 			return FReply::Handled();
@@ -643,7 +653,7 @@ FReply FHierarchyWidget::OnDragOver(const FGeometry& MyGeometry, const FDragDrop
 	{
 		bool bIsDrop = false;
 		UWidgetBlueprint* Blueprint = BlueprintEditor.Pin()->GetWidgetBlueprintObj();
-		return ProcessHierarchyDragDrop(DragDropEvent, bIsDrop, Blueprint, Item);
+		return ProcessHierarchyDragDrop(DragDropEvent, bIsDrop, BlueprintEditor.Pin(), Blueprint, Item);
 	}
 
 	return FReply::Unhandled();
@@ -675,7 +685,7 @@ FReply FHierarchyWidget::HandleDrop(const FDragDropEvent& DragDropEvent)
 {
 	UWidgetBlueprint* Blueprint = BlueprintEditor.Pin()->GetWidgetBlueprintObj();
 	bool bIsDrop = true;
-	return ProcessHierarchyDragDrop(DragDropEvent, bIsDrop, Blueprint, Item);
+	return ProcessHierarchyDragDrop(DragDropEvent, bIsDrop, BlueprintEditor.Pin(), Blueprint, Item);
 }
 
 bool FHierarchyWidget::OnVerifyNameTextChanged(const FText& InText, FText& OutErrorMessage)
