UMG - Details panel now correctly shows the current selection details after recreating the view.

diff --git a/Engine/Source/Editor/UMGEditor/Private/Details/SWidgetDetailsView.cpp b/Engine/Source/Editor/UMGEditor/Private/Details/SWidgetDetailsView.cpp
--- a/Engine/Source/Editor/UMGEditor/Private/Details/SWidgetDetailsView.cpp
+++ b/Engine/Source/Editor/UMGEditor/Private/Details/SWidgetDetailsView.cpp
@@ -149,6 +149,9 @@ void SWidgetDetailsView::Construct(const FArguments& InArgs, TSharedPtr<FWidgetB
 	BlueprintEditor.Pin()->OnSelectedWidgetsChanged.AddRaw(this, &SWidgetDetailsView::OnEditorSelectionChanged);
 
 	RegisterCustomizations();
+	
+	// Refresh the selection in the details panel.
+	OnEditorSelectionChanged();
 }
 
 SWidgetDetailsView::~SWidgetDetailsView()
