UMG - Moving the code to check if a widget class is usable into a common area.  Having more code use it when serving up lists of widget classes, like in the "wrap with" command.  Adding a Replace With command, lets you swap out one multiitem container for another.  Expansion state is now stored in the widgets themselves just like hidden in the designer, neither cause the asset to be dirtied, but if you save the asset that information will be saved with it now.

diff --git a/Engine/Source/Editor/UMGEditor/Private/Hierarchy/SHierarchyView.cpp b/Engine/Source/Editor/UMGEditor/Private/Hierarchy/SHierarchyView.cpp
--- a/Engine/Source/Editor/UMGEditor/Private/Hierarchy/SHierarchyView.cpp
+++ b/Engine/Source/Editor/UMGEditor/Private/Hierarchy/SHierarchyView.cpp
@@ -102,8 +102,6 @@ void SHierarchyView::Tick(const FGeometry& AllottedGeometry, const double InCurr
 			RebuildTreeView();
 		}
 
-		SaveExpandedItems();
-
 		RefreshTree();
 
 		RestoreExpandedItems();
@@ -237,6 +235,11 @@ void SHierarchyView::WidgetHierarchy_OnSelectionChanged(TSharedPtr<FHierarchyMod
 	}
 }
 
+void SHierarchyView::WidgetHierarchy_OnExpansionChanged(TSharedPtr<FHierarchyModel> Item, bool bExpanded)
+{
+	Item->SetExpanded(bExpanded);
+}
+
 FReply SHierarchyView::HandleDeleteSelected()
 {
 	TSet<FWidgetReference> SelectedWidgets = BlueprintEditor.Pin()->GetSelectedWidgets();
@@ -262,6 +265,7 @@ void SHierarchyView::RebuildTreeView()
 		.OnGetChildren(FilterHandler.ToSharedRef(), &TreeFilterHandler< TSharedPtr<FHierarchyModel> >::OnGetFilteredChildren)
 		.OnGenerateRow(this, &SHierarchyView::WidgetHierarchy_OnGenerateRow)
 		.OnSelectionChanged(this, &SHierarchyView::WidgetHierarchy_OnSelectionChanged)
+		.OnExpansionChanged(this, &SHierarchyView::WidgetHierarchy_OnExpansionChanged)
 		.OnContextMenuOpening(this, &SHierarchyView::WidgetHierarchy_OnContextMenuOpening)
 		.TreeItemsSource(&TreeRootWidgets);
 
@@ -280,23 +284,6 @@ void SHierarchyView::OnObjectsReplaced(const TMap<UObject*, UObject*>& Replaceme
 	{
 		bRefreshRequested = true;
 		bRebuildTreeRequested = true;
-
-		// We save the expanded items immediately because they're potentially about to become invalid.
-		SaveExpandedItems();
-	}
-}
-
-void SHierarchyView::SaveExpandedItems()
-{
-	if ( ExpandedItems.Num() == 0 )
-	{
-		TSet < TSharedPtr<FHierarchyModel> > ExpandedModels;
-		WidgetTreeView->GetExpandedItems(ExpandedModels);
-
-		for ( TSharedPtr<FHierarchyModel>& Model : ExpandedModels )
-		{
-			ExpandedItems.Add(Model->GetUniqueName());
-		}
 	}
 }
 
@@ -306,13 +293,11 @@ void SHierarchyView::RestoreExpandedItems()
 	{
 		RecursiveExpand(Model);
 	}
-
-	ExpandedItems.Empty();
 }
 
 void SHierarchyView::RecursiveExpand(TSharedPtr<FHierarchyModel>& Model)
 {
-	if ( ExpandedItems.Contains(Model->GetUniqueName()) )
+	if ( Model->IsExpanded() )
 	{
 		WidgetTreeView->SetItemExpansion(Model, true);
 
diff --git a/Engine/Source/Editor/UMGEditor/Private/Hierarchy/SHierarchyView.h b/Engine/Source/Editor/UMGEditor/Private/Hierarchy/SHierarchyView.h
--- a/Engine/Source/Editor/UMGEditor/Private/Hierarchy/SHierarchyView.h
+++ b/Engine/Source/Editor/UMGEditor/Private/Hierarchy/SHierarchyView.h
@@ -45,6 +45,7 @@ private:
 	void WidgetHierarchy_OnGetChildren(TSharedPtr<FHierarchyModel> InParent, TArray< TSharedPtr<FHierarchyModel> >& OutChildren);
 	TSharedRef< ITableRow > WidgetHierarchy_OnGenerateRow(TSharedPtr<FHierarchyModel> InItem, const TSharedRef<STableViewBase>& OwnerTable);
 	void WidgetHierarchy_OnSelectionChanged(TSharedPtr<FHierarchyModel> SelectedItem, ESelectInfo::Type SelectInfo);
+	void WidgetHierarchy_OnExpansionChanged(TSharedPtr<FHierarchyModel> Item, bool bExpanded);
 
 private:
 	/** @returns the current blueprint being edited */
@@ -86,9 +87,6 @@ private:
 	/** Called when a Blueprint is recompiled and live objects are swapped out for replacements */
 	void OnObjectsReplaced(const TMap<UObject*, UObject*>& ReplacementMap);
 
-	/** Saves the expansion of items into a list based on FName */
-	void SaveExpandedItems();
-
 	/** Restores the state of expanded items based on the saved expanded item state, then clears the expanded state cache. */
 	void RestoreExpandedItems();
 
@@ -124,9 +122,6 @@ private:
 	/** The filter used by the search box */
 	TSharedPtr<WidgetTextFilter> SearchBoxWidgetFilter;
 
-	/** Temporary expanded item state, used to restore expansion after tree rebuilds. */
-	TArray<FName> ExpandedItems;
-
 	/** Has a full refresh of the tree been requested?  This happens when the user is filtering the tree */
 	bool bRefreshRequested;
 
diff --git a/Engine/Source/Editor/UMGEditor/Private/Hierarchy/SHierarchyViewItem.h b/Engine/Source/Editor/UMGEditor/Private/Hierarchy/SHierarchyViewItem.h
--- a/Engine/Source/Editor/UMGEditor/Private/Hierarchy/SHierarchyViewItem.h
+++ b/Engine/Source/Editor/UMGEditor/Private/Hierarchy/SHierarchyViewItem.h
@@ -58,6 +58,9 @@ public:
 	virtual bool CanControlVisibility() const { return false; }
 	virtual void SetIsVisible(bool IsVisible) { }
 
+	virtual bool IsExpanded() const { return true; }
+	virtual void SetExpanded(bool bIsExpanded) { }
+
 protected:
 	virtual void GetChildren(TArray< TSharedPtr<FHierarchyModel> >& Children) = 0;
 	virtual void UpdateSelection() = 0;
@@ -193,9 +196,16 @@ public:
 		{
 			PreviewWidget->bHiddenInDesigner = !IsVisible;
 		}
+	}
+
+	virtual bool IsExpanded() const override
+	{
+		return Item.GetTemplate()->bExpandedInDesigner;
+	}
 
-		// Mark the blueprint as modified
-		FBlueprintEditorUtils::MarkBlueprintAsModified(BlueprintEditor.Pin()->GetBlueprintObj());
+	virtual void SetExpanded(bool bIsExpanded) override
+	{
+		Item.GetTemplate()->bExpandedInDesigner = bIsExpanded;
 	}
 
 protected:
diff --git a/Engine/Source/Editor/UMGEditor/Private/Palette/SPaletteView.cpp b/Engine/Source/Editor/UMGEditor/Private/Palette/SPaletteView.cpp
--- a/Engine/Source/Editor/UMGEditor/Private/Palette/SPaletteView.cpp
+++ b/Engine/Source/Editor/UMGEditor/Private/Palette/SPaletteView.cpp
@@ -345,23 +345,12 @@ void SPaletteView::BuildClassWidgetList()
 	{
 		UClass* WidgetClass = *ClassIt;
 
-		if ( WidgetClass->IsChildOf(UWidget::StaticClass()) )
-		{	
-			// We aren't interested in classes that are experimental or cannot be instantiated
-			bool bIsExperimental, bIsEarlyAccess;
-			FObjectEditorUtils::GetClassDevelopmentStatus(WidgetClass, bIsExperimental, bIsEarlyAccess);
-			const bool bIsInvalid = WidgetClass->HasAnyClassFlags(CLASS_Abstract | CLASS_Deprecated | CLASS_NewerVersionExists);
-			if ( bIsExperimental || bIsEarlyAccess || bIsInvalid )
-			{
-				continue;
-			}
-
-			// Don't include skeleton classes or the same class as the widget being edited
-			const bool bIsSkeletonClass = WidgetClass->HasAnyFlags(RF_Transient) && WidgetClass->HasAnyClassFlags(CLASS_CompiledFromBlueprint);
+		if ( FWidgetBlueprintEditorUtils::IsUsableWidgetClass(WidgetClass) )
+		{
 			const bool bIsSameClass = WidgetClass->GetFName() == ActiveWidgetBlueprintClassName;
 
 			// Check that the asset that generated this class is valid (necessary b/c of a larger issue wherein force delete does not wipe the generated class object)
-			if ( bIsSkeletonClass || bIsSameClass )
+			if ( bIsSameClass )
 			{
 				continue;
 			}
diff --git a/Engine/Source/Editor/UMGEditor/Private/WidgetBlueprintEditorUtils.cpp b/Engine/Source/Editor/UMGEditor/Private/WidgetBlueprintEditorUtils.cpp
--- a/Engine/Source/Editor/UMGEditor/Private/WidgetBlueprintEditorUtils.cpp
+++ b/Engine/Source/Editor/UMGEditor/Private/WidgetBlueprintEditorUtils.cpp
@@ -105,45 +105,11 @@ bool FWidgetBlueprintEditorUtils::RenameWidget(TSharedRef<FWidgetBlueprintEditor
 		// Find and update all variable references in the graph
 
 		Widget->Rename(*NewNameStr);
-	
-		// Update Variable References
-		TArray<UK2Node_Variable*> WidgetVarNodes;
-		FBlueprintEditorUtils::GetAllNodesOfClass<UK2Node_Variable>(Blueprint, WidgetVarNodes);
-		for ( UK2Node_Variable* TestNode : WidgetVarNodes )
-		{
-			if ( TestNode && ( OldName == TestNode->GetVarName() ) )
-			{
-				UEdGraphPin* TestPin = TestNode->FindPin(OldNameStr);
-				if ( TestPin && ( Widget->GetClass() == TestPin->PinType.PinSubCategoryObject.Get() ) )
-				{
-					TestNode->Modify();
-					if ( TestNode->VariableReference.IsSelfContext() )
-					{
-						TestNode->VariableReference.SetSelfMember(NewName);
-					}
-					else
-					{
-						//TODO:
-						UClass* ParentClass = TestNode->VariableReference.GetMemberParentClass((UClass*)NULL);
-						TestNode->VariableReference.SetExternalMember(NewName, ParentClass);
-					}
-					TestPin->Modify();
-					TestPin->PinName = NewNameStr;
-				}
-			}
-		}
 
+		// Update Variable References and
 		// Update Event References to member variables
-		TArray<UK2Node_ComponentBoundEvent*> EventNodes;
-		FBlueprintEditorUtils::GetAllNodesOfClass<UK2Node_ComponentBoundEvent>(Blueprint, EventNodes);
-		for ( UK2Node_ComponentBoundEvent* EventNode : EventNodes )
-		{
-			if ( EventNode->ComponentPropertyName == OldName )
-			{
-				EventNode->ComponentPropertyName = NewName;
-			}
-		}
-
+		FBlueprintEditorUtils::ReplaceVariableReferences(Blueprint, OldName, NewName);
+		
 		// Find and update all binding references in the widget blueprint
 		for ( FDelegateEditorBinding& Binding : Blueprint->Bindings )
 		{
@@ -216,6 +182,15 @@ void FWidgetBlueprintEditorUtils::CreateWidgetContextMenu(FMenuBuilder& MenuBuil
 			LOCTEXT("WidgetTree_WrapWithToolTip", "Wraps the currently selected widgets inside of another container widget"),
 			FNewMenuDelegate::CreateStatic(&FWidgetBlueprintEditorUtils::BuildWrapWithMenu, BP, Widgets)
 			);
+
+		if ( Widgets.Num() == 1 )
+		{
+			MenuBuilder.AddSubMenu(
+				LOCTEXT("WidgetTree_ReplaceWith", "Replace With..."),
+				LOCTEXT("WidgetTree_ReplaceWithToolTip", "Replaces the currently selected panel widget, with another panel widget"),
+				FNewMenuDelegate::CreateStatic(&FWidgetBlueprintEditorUtils::BuildReplaceWithMenu, BP, Widgets)
+				);
+		}
 	}
 	MenuBuilder.EndSection();
 
@@ -339,17 +314,20 @@ void FWidgetBlueprintEditorUtils::BuildWrapWithMenu(FMenuBuilder& Menu, UWidgetB
 		for ( TObjectIterator<UClass> ClassIt; ClassIt; ++ClassIt )
 		{
 			UClass* WidgetClass = *ClassIt;
-			if ( WidgetClass->IsChildOf(UPanelWidget::StaticClass()) && WidgetClass->HasAnyClassFlags(CLASS_Abstract) == false )
+			if ( FWidgetBlueprintEditorUtils::IsUsableWidgetClass(WidgetClass) )
 			{
-				Menu.AddMenuEntry(
-					WidgetClass->GetDisplayNameText(),
-					FText::GetEmpty(),
-					FSlateIcon(),
-					FUIAction(
-					FExecuteAction::CreateStatic(&FWidgetBlueprintEditorUtils::WrapWidgets, BP, Widgets, WidgetClass),
-					FCanExecuteAction()
-					)
-					);
+				if ( WidgetClass->IsChildOf(UPanelWidget::StaticClass()) && WidgetClass->HasAnyClassFlags(CLASS_Abstract) == false )
+				{
+					Menu.AddMenuEntry(
+						WidgetClass->GetDisplayNameText(),
+						FText::GetEmpty(),
+						FSlateIcon(),
+						FUIAction(
+						FExecuteAction::CreateStatic(&FWidgetBlueprintEditorUtils::WrapWidgets, BP, Widgets, WidgetClass),
+						FCanExecuteAction()
+						)
+						);
+				}
 			}
 		}
 	}
@@ -358,11 +336,12 @@ void FWidgetBlueprintEditorUtils::BuildWrapWithMenu(FMenuBuilder& Menu, UWidgetB
 
 void FWidgetBlueprintEditorUtils::WrapWidgets(UWidgetBlueprint* BP, TSet<FWidgetReference> Widgets, UClass* WidgetClass)
 {
+	const FScopedTransaction Transaction(LOCTEXT("WrapWidgets", "Wrap Widgets"));
+
 	TSharedPtr<FWidgetTemplateClass> Template = MakeShareable(new FWidgetTemplateClass(WidgetClass));
 
 	UPanelWidget* NewWrapperWidget = CastChecked<UPanelWidget>(Template->Create(BP->WidgetTree));
 
-	//TODO UMG When wrapping multiple widgets, how will that work?
 	for ( FWidgetReference& Item : Widgets )
 	{
 		int32 OutIndex;
@@ -372,7 +351,7 @@ void FWidgetBlueprintEditorUtils::WrapWidgets(UWidgetBlueprint* BP, TSet<FWidget
 			CurrentParent->Modify();
 			CurrentParent->ReplaceChildAt(OutIndex, NewWrapperWidget);
 		}
-		else
+		else if ( Item.GetTemplate() == BP->WidgetTree->RootWidget )
 		{
 			BP->WidgetTree->Modify();
 
@@ -385,6 +364,79 @@ void FWidgetBlueprintEditorUtils::WrapWidgets(UWidgetBlueprint* BP, TSet<FWidget
 	FBlueprintEditorUtils::MarkBlueprintAsStructurallyModified(BP);
 }
 
+void FWidgetBlueprintEditorUtils::BuildReplaceWithMenu(FMenuBuilder& Menu, UWidgetBlueprint* BP, TSet<FWidgetReference> Widgets)
+{
+	Menu.BeginSection("ReplaceWith", LOCTEXT("WidgetTree_ReplaceWith", "Replace With..."));
+	{
+		for ( TObjectIterator<UClass> ClassIt; ClassIt; ++ClassIt )
+		{
+			UClass* WidgetClass = *ClassIt;
+			if ( FWidgetBlueprintEditorUtils::IsUsableWidgetClass(WidgetClass) )
+			{
+				if ( WidgetClass->IsChildOf(UPanelWidget::StaticClass()) && WidgetClass->HasAnyClassFlags(CLASS_Abstract) == false )
+				{
+					// Only allow replacement with panels that accept multiple children
+					if ( WidgetClass->GetDefaultObject<UPanelWidget>()->CanHaveMultipleChildren() )
+					{
+						Menu.AddMenuEntry(
+							WidgetClass->GetDisplayNameText(),
+							FText::GetEmpty(),
+							FSlateIcon(),
+							FUIAction(
+							FExecuteAction::CreateStatic(&FWidgetBlueprintEditorUtils::ReplaceWidgets, BP, Widgets, WidgetClass),
+							FCanExecuteAction()
+							)
+							);
+					}
+				}
+			}
+		}
+	}
+	Menu.EndSection();
+}
+
+void FWidgetBlueprintEditorUtils::ReplaceWidgets(UWidgetBlueprint* BP, TSet<FWidgetReference> Widgets, UClass* WidgetClass)
+{
+	const FScopedTransaction Transaction(LOCTEXT("ReplaceWidgets", "Replace Widgets"));
+
+	TSharedPtr<FWidgetTemplateClass> Template = MakeShareable(new FWidgetTemplateClass(WidgetClass));
+
+	UPanelWidget* NewReplacementWidget = CastChecked<UPanelWidget>(Template->Create(BP->WidgetTree));
+
+	for ( FWidgetReference& Item : Widgets )
+	{
+		if ( UPanelWidget* ExistingPanel = Cast<UPanelWidget>(Item.GetTemplate()) )
+		{
+			ExistingPanel->Modify();
+
+			if ( UPanelWidget* CurrentParent = ExistingPanel->GetParent() )
+			{
+				CurrentParent->Modify();
+				CurrentParent->ReplaceChild(ExistingPanel, NewReplacementWidget);
+			}
+			else if ( Item.GetTemplate() == BP->WidgetTree->RootWidget )
+			{
+				BP->WidgetTree->Modify();
+				BP->WidgetTree->RootWidget = NewReplacementWidget;
+			}
+			else
+			{
+				continue;
+			}
+
+			while ( ExistingPanel->GetChildrenCount() > 0 )
+			{
+				UWidget* Widget = ExistingPanel->GetChildAt(0);
+				Widget->Modify();
+
+				NewReplacementWidget->AddChild(Widget);
+			}
+		}
+	}
+
+	FBlueprintEditorUtils::MarkBlueprintAsStructurallyModified(BP);
+}
+
 void FWidgetBlueprintEditorUtils::CutWidgets(UWidgetBlueprint* BP, TSet<FWidgetReference> Widgets)
 {
 	CopyWidgets(BP, Widgets);
@@ -601,4 +653,32 @@ void FWidgetBlueprintEditorUtils::ImportPropertiesFromText(UObject* Object, cons
 	}
 }
 
+bool FWidgetBlueprintEditorUtils::IsUsableWidgetClass(UClass* WidgetClass)
+{
+	if ( WidgetClass->IsChildOf(UWidget::StaticClass()) )
+	{
+		// We aren't interested in classes that are experimental or cannot be instantiated
+		bool bIsExperimental, bIsEarlyAccess;
+		FObjectEditorUtils::GetClassDevelopmentStatus(WidgetClass, bIsExperimental, bIsEarlyAccess);
+		const bool bIsInvalid = WidgetClass->HasAnyClassFlags(CLASS_Abstract | CLASS_Deprecated | CLASS_NewerVersionExists);
+		if ( bIsExperimental || bIsEarlyAccess || bIsInvalid )
+		{
+			return false;
+		}
+
+		// Don't include skeleton classes or the same class as the widget being edited
+		const bool bIsSkeletonClass = WidgetClass->HasAnyFlags(RF_Transient) && WidgetClass->HasAnyClassFlags(CLASS_CompiledFromBlueprint);
+
+		// Check that the asset that generated this class is valid (necessary b/c of a larger issue wherein force delete does not wipe the generated class object)
+		if ( bIsSkeletonClass )
+		{
+			return false;
+		}
+
+		return true;
+	}
+
+	return false;
+}
+
 #undef LOCTEXT_NAMESPACE
diff --git a/Engine/Source/Editor/UMGEditor/Private/WidgetBlueprintEditorUtils.h b/Engine/Source/Editor/UMGEditor/Private/WidgetBlueprintEditorUtils.h
--- a/Engine/Source/Editor/UMGEditor/Private/WidgetBlueprintEditorUtils.h
+++ b/Engine/Source/Editor/UMGEditor/Private/WidgetBlueprintEditorUtils.h
@@ -24,6 +24,8 @@ public:
 
 	static void CutWidgets(UWidgetBlueprint* BP, TSet<FWidgetReference> Widgets);
 
+	static bool IsUsableWidgetClass(UClass* WidgetClass);
+
 public:
 	static void ExportWidgetsToText(TSet<UWidget*> WidgetsToExport, /*out*/ FString& ExportedText);
 
@@ -46,4 +48,8 @@ private:
 	static void BuildWrapWithMenu(FMenuBuilder& Menu, UWidgetBlueprint* BP, TSet<FWidgetReference> Widgets);
 
 	static void WrapWidgets(UWidgetBlueprint* BP, TSet<FWidgetReference> Widgets, UClass* WidgetClass);
+
+	static void BuildReplaceWithMenu(FMenuBuilder& Menu, UWidgetBlueprint* BP, TSet<FWidgetReference> Widgets);
+
+	static void ReplaceWidgets(UWidgetBlueprint* BP, TSet<FWidgetReference> Widgets, UClass* WidgetClass);
 };
diff --git a/Engine/Source/Runtime/UMG/Public/Components/Widget.h b/Engine/Source/Runtime/UMG/Public/Components/Widget.h
--- a/Engine/Source/Runtime/UMG/Public/Components/Widget.h
+++ b/Engine/Source/Runtime/UMG/Public/Components/Widget.h
@@ -125,6 +125,10 @@ public:
 	UPROPERTY()
 	bool bHiddenInDesigner;
 
+	/** Stores the design time flag setting if the widget is expanded inside the designer */
+	UPROPERTY()
+	bool bExpandedInDesigner;
+
 	/** Stores a reference to the asset responsible for this widgets construction. */
 	UPROPERTY(Transient)
 	UObject* WidgetGeneratedBy;
