Non-unity build fix.

diff --git a/Engine/Source/Editor/UMGEditor/Private/Palette/SPaletteView.cpp b/Engine/Source/Editor/UMGEditor/Private/Palette/SPaletteView.cpp
--- a/Engine/Source/Editor/UMGEditor/Private/Palette/SPaletteView.cpp
+++ b/Engine/Source/Editor/UMGEditor/Private/Palette/SPaletteView.cpp
@@ -23,6 +23,7 @@
 #include "SSearchBox.h"
 
 #include "WidgetBlueprintCompiler.h"
+#include "WidgetBlueprintEditorUtils.h"
 
 #include "Blueprint/UserWidget.h"
 #include "WidgetBlueprint.h"
