UMG - Fixing a bug with the named slots disappearing when you click PIE.

diff --git a/Engine/Source/Editor/UMGEditor/Private/Hierarchy/SHierarchyViewItem.cpp b/Engine/Source/Editor/UMGEditor/Private/Hierarchy/SHierarchyViewItem.cpp
--- a/Engine/Source/Editor/UMGEditor/Private/Hierarchy/SHierarchyViewItem.cpp
+++ b/Engine/Source/Editor/UMGEditor/Private/Hierarchy/SHierarchyViewItem.cpp
@@ -421,7 +421,7 @@ FName FNamedSlotModel::GetUniqueName() const
 
 FText FNamedSlotModel::GetText() const
 {
-	if ( INamedSlotInterface* NamedSlotHost = Cast<INamedSlotInterface>(Item.GetPreview()) )
+	if ( INamedSlotInterface* NamedSlotHost = Cast<INamedSlotInterface>(Item.GetTemplate()) )
 	{
 		TSet<FWidgetReference> SelectedWidgets;
 		if ( UWidget* SlotContent = NamedSlotHost->GetContentForSlot(SlotName) )
@@ -732,7 +732,7 @@ void FHierarchyWidget::GetChildren(TArray< TSharedPtr<FHierarchyModel> >& Childr
 	TSharedPtr<FWidgetBlueprintEditor> BPEd = BlueprintEditor.Pin();
 
 	// Check for named slots
-	if ( INamedSlotInterface* NamedSlotHost = Cast<INamedSlotInterface>(Item.GetPreview()) )
+	if ( INamedSlotInterface* NamedSlotHost = Cast<INamedSlotInterface>(Item.GetTemplate()) )
 	{
 		TArray<FName> SlotNames;
 		NamedSlotHost->GetSlotNames(SlotNames);
