UMG - Binding options no longer available on the Root of a widget (it never actually worked), now the UI properly hides to show that fact.

diff --git a/Engine/Source/Editor/UMGEditor/Private/Details/DetailWidgetExtensionHandler.cpp b/Engine/Source/Editor/UMGEditor/Private/Details/DetailWidgetExtensionHandler.cpp
--- a/Engine/Source/Editor/UMGEditor/Private/Details/DetailWidgetExtensionHandler.cpp
+++ b/Engine/Source/Editor/UMGEditor/Private/Details/DetailWidgetExtensionHandler.cpp
@@ -14,6 +14,15 @@ bool FDetailWidgetExtensionHandler::IsPropertyExtenable(const UClass* InObjectCl
 	// TODO UMG make this work for multiple widgets.
 	if ( InPropertyHandle.GetNumOuterObjects() == 1 )
 	{
+		TArray<UObject*> Objects;
+		InPropertyHandle.GetOuterObjects(Objects);
+
+		// We don't allow bindings on the CDO.
+		if ( Objects[0]->HasAnyFlags(RF_ClassDefaultObject) )
+		{
+			return false;
+		}
+
 		UProperty* Property = InPropertyHandle.GetProperty();
 		FString DelegateName = Property->GetName() + "Delegate";
 
