UMG - The play PIE button is now visibile in the designer.

diff --git a/Engine/Source/Editor/UMGEditor/Private/BlueprintModes/WidgetDesignerApplicationMode.cpp b/Engine/Source/Editor/UMGEditor/Private/BlueprintModes/WidgetDesignerApplicationMode.cpp
--- a/Engine/Source/Editor/UMGEditor/Private/BlueprintModes/WidgetDesignerApplicationMode.cpp
+++ b/Engine/Source/Editor/UMGEditor/Private/BlueprintModes/WidgetDesignerApplicationMode.cpp
@@ -124,9 +124,11 @@ FWidgetDesignerApplicationMode::FWidgetDesignerApplicationMode(TSharedPtr<FWidge
 	ToolbarExtender = MakeShareable(new FExtender);
 	InWidgetEditor->GetWidgetToolbarBuilder()->AddWidgetBlueprintEditorModesToolbar(ToolbarExtender);
 	InWidgetEditor->GetToolbarBuilder()->AddCompileToolbar(ToolbarExtender);
-	//InPersona->GetToolbarBuilder()->AddScriptingToolbar(ToolbarExtender);
-	//InPersona->GetToolbarBuilder()->AddBlueprintGlobalOptionsToolbar(ToolbarExtender);
-	//InPersona->GetToolbarBuilder()->AddDebuggingToolbar(ToolbarExtender);
+	//InWidgetEditor->GetToolbarBuilder()->AddComponentsToolbar(ToolbarExtender);
+	InWidgetEditor->GetToolbarBuilder()->AddDebuggingToolbar(ToolbarExtender);
+
+	//InWidgetEditor->GetToolbarBuilder()->AddScriptingToolbar(ToolbarExtender);
+	//InWidgetEditor->GetToolbarBuilder()->AddBlueprintGlobalOptionsToolbar(ToolbarExtender);
 }
 
 void FWidgetDesignerApplicationMode::RegisterTabFactories(TSharedPtr<FTabManager> InTabManager)
