UMG - Designer now changes visualization when in PIE mode.  User widgets now remove old binding data that's no longer applicable if some of their named slots are removed, they now also remove the bindings on save.

diff --git a/Engine/Source/Editor/UMGEditor/Private/Designer/SDesignerView.cpp b/Engine/Source/Editor/UMGEditor/Private/Designer/SDesignerView.cpp
--- a/Engine/Source/Editor/UMGEditor/Private/Designer/SDesignerView.cpp
+++ b/Engine/Source/Editor/UMGEditor/Private/Designer/SDesignerView.cpp
@@ -244,6 +244,29 @@ void SDesignerView::Construct(const FArguments& InArgs, TSharedPtr<FWidgetBluepr
 				.OnPaintHandler(this, &SDesignerView::HandleEffectsPainting)
 			]
 
+			// Top-right corner text indicating PIE is active
+			+ SOverlay::Slot()
+			.Padding(0)
+			.VAlign(VAlign_Fill)
+			.HAlign(HAlign_Fill)
+			[
+				SNew(SImage)
+				.Visibility(this, &SDesignerView::PIENotification)
+				.Image(FEditorStyle::GetBrush(TEXT("Graph.PlayInEditor")))
+			]
+
+			// Top-right corner text indicating PIE is active
+			+ SOverlay::Slot()
+			.Padding(20)
+			.VAlign(VAlign_Top)
+			.HAlign(HAlign_Right)
+			[
+				SNew(STextBlock)
+				.Visibility(this, &SDesignerView::PIENotification)
+				.TextStyle( FEditorStyle::Get(), "Graph.SimulatingText" )
+				.Text( LOCTEXT("SIMULATING", "SIMULATING") )
+			]
+
 			// A layer in the overlay where we put all the user intractable widgets, like the reorder widgets.
 			+ SOverlay::Slot()
 			.HAlign(HAlign_Fill)
@@ -1751,6 +1774,16 @@ EVisibility SDesignerView::GetResolutionTextVisibility() const
 	return EVisibility::SelfHitTestInvisible;
 }
 
+EVisibility SDesignerView::PIENotification() const
+{
+	if ( GEditor->bIsSimulatingInEditor || GEditor->PlayWorld != nullptr )
+	{
+		return EVisibility::HitTestInvisible;
+	}
+
+	return EVisibility::Hidden;
+}
+
 FReply SDesignerView::HandleDPISettingsClicked()
 {
 	FModuleManager::LoadModuleChecked<ISettingsModule>("Settings").ShowViewer("Project", "Engine", "UI");
diff --git a/Engine/Source/Editor/UMGEditor/Private/Designer/SDesignerView.h b/Engine/Source/Editor/UMGEditor/Private/Designer/SDesignerView.h
--- a/Engine/Source/Editor/UMGEditor/Private/Designer/SDesignerView.h
+++ b/Engine/Source/Editor/UMGEditor/Private/Designer/SDesignerView.h
@@ -133,6 +133,8 @@ private:
 	void HandleOnCustomResolutionSelected();
 	TSharedRef<SWidget> GetAspectMenu();
 
+	EVisibility PIENotification() const;
+
 	// Handles drawing selection and other effects a SPaintSurface widget injected into the hierarchy.
 	int32 HandleEffectsPainting(const FOnPaintHandlerParams& PaintArgs);
 	FReply HandleDPISettingsClicked();
diff --git a/Engine/Source/Editor/UMGEditor/Private/WidgetBlueprintEditor.cpp b/Engine/Source/Editor/UMGEditor/Private/WidgetBlueprintEditor.cpp
--- a/Engine/Source/Editor/UMGEditor/Private/WidgetBlueprintEditor.cpp
+++ b/Engine/Source/Editor/UMGEditor/Private/WidgetBlueprintEditor.cpp
@@ -28,7 +28,10 @@ FWidgetBlueprintEditor::FWidgetBlueprintEditor()
 	: PreviewScene(FPreviewScene::ConstructionValues().AllowAudioPlayback(true).ShouldSimulatePhysics(true))
 	, PreviewBlueprint(nullptr)
 	, HoverTime(0)
+	, bIsSimulateEnabled(false)
+	, bIsRealTime(true)
 {
+	PreviewScene.GetWorld()->bBegunPlay = false;
 }
 
 FWidgetBlueprintEditor::~FWidgetBlueprintEditor()
@@ -59,7 +62,7 @@ void FWidgetBlueprintEditor::InitWidgetBlueprintEditor(const EToolkitMode::Type
 	UWidgetBlueprint* Blueprint = GetWidgetBlueprintObj();
 
 	// If this blueprint is empty, add a canvas panel as the root widget.
-	if ( Blueprint->WidgetTree->RootWidget == NULL )
+	if ( Blueprint->WidgetTree->RootWidget == nullptr )
 	{
 		UWidget* RootWidget = Blueprint->WidgetTree->ConstructWidget<UCanvasPanel>(UCanvasPanel::StaticClass());
 		RootWidget->SetIsDesignTime(true);
@@ -318,12 +321,12 @@ bool FWidgetBlueprintEditor::CanPasteWidgets()
 	if ( Widgets.Num() == 1 )
 	{
 		FWidgetReference Target = *Widgets.CreateIterator();
-		const bool bIsPanel = Cast<UPanelWidget>(Target.GetTemplate()) != NULL;
+		const bool bIsPanel = Cast<UPanelWidget>(Target.GetTemplate()) != nullptr;
 		return bIsPanel;
 	}
 	else if ( Widgets.Num() == 0 )
 	{
-		if ( GetWidgetBlueprintObj()->WidgetTree->RootWidget == NULL )
+		if ( GetWidgetBlueprintObj()->WidgetTree->RootWidget == nullptr )
 		{
 			return true;
 		}
@@ -351,11 +354,8 @@ void FWidgetBlueprintEditor::Tick(float DeltaTime)
 	// Tick the preview scene world.
 	if ( !GIntraFrameDebuggingGameThread )
 	{
-		bool bIsSimulateEnabled = true;
-		bool bIsRealTime = true;
-
 		// Allow full tick only if preview simulation is enabled and we're not currently in an active SIE or PIE session
-		if ( bIsSimulateEnabled && GEditor->PlayWorld == NULL && !GEditor->bIsSimulatingInEditor )
+		if ( bIsSimulateEnabled && GEditor->PlayWorld == nullptr && !GEditor->bIsSimulatingInEditor )
 		{
 			PreviewScene.GetWorld()->Tick(bIsRealTime ? LEVELTICK_All : LEVELTICK_TimeOnly, DeltaTime);
 		}
@@ -378,7 +378,7 @@ static bool MigratePropertyValue(UObject* SourceObject, UObject* DestinationObje
 {
 	UProperty* CurrentProperty = PropertyChainNode->GetValue();
 
-	if ( PropertyChainNode->GetNextNode() == NULL )
+	if ( PropertyChainNode->GetNextNode() == nullptr )
 	{
 		if ( bIsModify )
 		{
@@ -390,7 +390,7 @@ static bool MigratePropertyValue(UObject* SourceObject, UObject* DestinationObje
 			// Check to see if there's an edit condition property we also need to migrate.
 			bool bDummyNegate = false;
 			UBoolProperty* EditConditionProperty = PropertyCustomizationHelpers::GetEditConditionProperty(MemberProperty, bDummyNegate);
-			if ( EditConditionProperty != NULL )
+			if ( EditConditionProperty != nullptr )
 			{
 				FObjectEditorUtils::MigratePropertyValue(SourceObject, EditConditionProperty, DestinationObject, EditConditionProperty);
 			}
@@ -428,7 +428,7 @@ void FWidgetBlueprintEditor::MigrateFromChain(FEditPropertyChain* PropertyThatCh
 	UWidgetBlueprint* Blueprint = GetWidgetBlueprintObj();
 
 	UUserWidget* PreviewActor = GetPreview();
-	if ( PreviewActor != NULL )
+	if ( PreviewActor != nullptr )
 	{
 		for ( FWidgetReference& WidgetRef : SelectedWidgets )
 		{
@@ -489,12 +489,27 @@ UUserWidget* FWidgetBlueprintEditor::GetPreview() const
 {
 	if ( PreviewWidgetPtr.IsStale(true) )
 	{
-		return NULL;
+		return nullptr;
 	}
 
 	return PreviewWidgetPtr.Get();
 }
 
+FPreviewScene* FWidgetBlueprintEditor::GetPreviewScene()
+{
+	return &PreviewScene;
+}
+
+bool FWidgetBlueprintEditor::IsSimulating() const
+{
+	return bIsSimulateEnabled;
+}
+
+void FWidgetBlueprintEditor::SetIsSimulating(bool bSimulating)
+{
+	bIsSimulateEnabled = bSimulating;
+}
+
 FWidgetReference FWidgetBlueprintEditor::GetReferenceFromTemplate(UWidget* TemplateWidget)
 {
 	TSharedRef<FWidgetHandle> Reference = MakeShareable(new FWidgetHandle(TemplateWidget));
@@ -612,7 +627,7 @@ void FWidgetBlueprintEditor::RefreshPreview()
 void FWidgetBlueprintEditor::DestroyPreview()
 {
 	UUserWidget* PreviewActor = GetPreview();
-	if ( PreviewActor != NULL )
+	if ( PreviewActor != nullptr )
 	{
 		check(PreviewScene.GetWorld());
 
@@ -625,7 +640,7 @@ void FWidgetBlueprintEditor::UpdatePreview(UBlueprint* InBlueprint, bool bInForc
 	UUserWidget* PreviewActor = GetPreview();
 
 	// Signal that we're going to be constructing editor components
-	if ( InBlueprint != NULL && InBlueprint->SimpleConstructionScript != NULL )
+	if ( InBlueprint != nullptr && InBlueprint->SimpleConstructionScript != nullptr )
 	{
 		InBlueprint->SimpleConstructionScript->BeginEditorComponentConstruction();
 	}
@@ -656,7 +671,6 @@ void FWidgetBlueprintEditor::UpdatePreview(UBlueprint* InBlueprint, bool bInForc
 	OnWidgetPreviewUpdated.Broadcast();
 }
 
-
 FGraphAppearanceInfo FWidgetBlueprintEditor::GetGraphAppearance() const
 {
 	FGraphAppearanceInfo AppearanceInfo = FBlueprintEditor::GetGraphAppearance();
diff --git a/Engine/Source/Editor/UMGEditor/Private/WidgetBlueprintEditor.h b/Engine/Source/Editor/UMGEditor/Private/WidgetBlueprintEditor.h
--- a/Engine/Source/Editor/UMGEditor/Private/WidgetBlueprintEditor.h
+++ b/Engine/Source/Editor/UMGEditor/Private/WidgetBlueprintEditor.h
@@ -43,6 +43,15 @@ public:
 	/** @return The preview widget. */
 	UUserWidget* GetPreview() const;
 
+	/** @return The preview scene that owns the preview widget. */
+	FPreviewScene* GetPreviewScene();
+
+	/**  */
+	bool IsSimulating() const;
+
+	/**  */
+	void SetIsSimulating(bool bSimulating);
+
 	/** Causes the preview to be destroyed and a new one to be created next tick */
 	void InvalidatePreview();
 
@@ -196,4 +205,10 @@ private:
 
 	/** The preview becomes invalid and needs to be rebuilt on the next tick. */
 	bool bPreviewInvalidated;
+
+	/**  */
+	bool bIsSimulateEnabled;
+
+	/**  */
+	bool bIsRealTime;
 };
diff --git a/Engine/Source/Runtime/UMG/Private/UserWidget.cpp b/Engine/Source/Runtime/UMG/Private/UserWidget.cpp
--- a/Engine/Source/Runtime/UMG/Private/UserWidget.cpp
+++ b/Engine/Source/Runtime/UMG/Private/UserWidget.cpp
@@ -70,11 +70,15 @@ void UUserWidget::Initialize()
 
 void UUserWidget::PostEditImport()
 {
+	Super::PostEditImport();
+
 	Initialize();
 }
 
 void UUserWidget::PostDuplicate(bool bDuplicateForPIE)
 {
+	Super::PostDuplicate(bDuplicateForPIE);
+
 	Initialize();
 }
 
@@ -151,6 +155,11 @@ UWorld* UUserWidget::GetWorld() const
 		}
 	}
 
+	if ( UWorld* World = Cast<UWorld>(GetOuter()) )
+	{
+		return World;
+	}
+
 	return nullptr;
 }
 
@@ -444,17 +453,19 @@ UWidget* UUserWidget::GetWidgetFromName(const FName& Name) const
 void UUserWidget::GetSlotNames(TArray<FName>& SlotNames) const
 {
 	// Only do this if this widget is of a blueprint class
-	UWidgetBlueprintGeneratedClass* BGClass = Cast<UWidgetBlueprintGeneratedClass>(GetClass());
-	if ( BGClass != nullptr )
+	if ( UWidgetBlueprintGeneratedClass* BGClass = Cast<UWidgetBlueprintGeneratedClass>(GetClass()) )
 	{
 		SlotNames.Append(BGClass->NamedSlots);
 	}
-
-	// Also add any existing bindings uniquely, templates don't have any components
-	// yet because they're never initialized because of their CDO status.
-	for ( const FNamedSlotBinding& Binding : NamedSlotBindings )
+	else // For non-blueprint widget blueprints we have to go through the widget tree to locate the named slots dynamically.
 	{
-		SlotNames.AddUnique( Binding.Name );
+		TArray<FName> NamedSlots;
+		WidgetTree->ForEachWidget([&] (UWidget* Widget) {
+			if ( Widget && Widget->IsA<UNamedSlot>() )
+			{
+				NamedSlots.Add(Widget->GetFName());
+			}
+		});
 	}
 }
 
@@ -660,8 +671,8 @@ void UUserWidget::SetPositionInViewport(FVector2D Position)
 {
 	float Scale = UWidgetLayoutLibrary::GetViewportScale(this);
 
-	ViewportOffsets.Left = Position.X * ( 1.0f / Scale );
-	ViewportOffsets.Top = Position.Y * ( 1.0f / Scale );
+	ViewportOffsets.Left = Position.X / Scale;
+	ViewportOffsets.Top = Position.Y / Scale;
 
 	ViewportAnchors = FAnchors(0, 0);
 }
@@ -710,6 +721,26 @@ FVector2D UUserWidget::GetFullScreenAlignment() const
 	return ViewportAlignment;
 }
 
+void UUserWidget::PreSave()
+{
+	Super::PreSave();
+
+	// Remove bindings that are no longer contained in the class.
+	if ( UWidgetBlueprintGeneratedClass* BGClass = Cast<UWidgetBlueprintGeneratedClass>(GetClass()) )
+	{
+		for ( int32 BindingIndex = 0; BindingIndex < NamedSlotBindings.Num(); BindingIndex++ )
+		{
+			const FNamedSlotBinding& Binding = NamedSlotBindings[BindingIndex];
+
+			if ( !BGClass->NamedSlots.Contains(Binding.Name) )
+			{
+				NamedSlotBindings.RemoveAt(BindingIndex);
+				BindingIndex--;
+			}
+		}
+	}
+}
+
 #if WITH_EDITOR
 
 const FSlateBrush* UUserWidget::GetEditorIcon()
diff --git a/Engine/Source/Runtime/UMG/Private/WidgetLayoutLibrary.cpp b/Engine/Source/Runtime/UMG/Private/WidgetLayoutLibrary.cpp
--- a/Engine/Source/Runtime/UMG/Private/WidgetLayoutLibrary.cpp
+++ b/Engine/Source/Runtime/UMG/Private/WidgetLayoutLibrary.cpp
@@ -37,8 +37,8 @@ bool UWidgetLayoutLibrary::GetMousePositionScaledByDPI(APlayerController* Player
 	if ( Player && Player->GetMousePosition(LocationX, LocationY) )
 	{
 		float Scale = UWidgetLayoutLibrary::GetViewportScale(Player);
-		LocationX = LocationX * ( 1.0f / Scale );
-		LocationY = LocationY * ( 1.0f / Scale );
+		LocationX = LocationX / Scale;
+		LocationY = LocationY / Scale;
 
 		return true;
 	}
diff --git a/Engine/Source/Runtime/UMG/Public/Blueprint/UserWidget.h b/Engine/Source/Runtime/UMG/Public/Blueprint/UserWidget.h
--- a/Engine/Source/Runtime/UMG/Public/Blueprint/UserWidget.h
+++ b/Engine/Source/Runtime/UMG/Public/Blueprint/UserWidget.h
@@ -579,6 +579,10 @@ public:
 	/** Ticks this widget and forwards to the underlying widget to tick */
 	void NativeTick(const FGeometry& MyGeometry, float InDeltaTime );
 
+	// Begin UObject interface
+	virtual void PreSave() override;
+	// End UObject interface
+
 #if WITH_EDITOR
 	// UWidget interface
 	virtual const FSlateBrush* GetEditorIcon() override;
