UMG - Fixing spelling of Visibility on widgets.  Now using a class enum with a specified size, this appears to fix the issue with visibility bindings returning bogus values on the stack, likely due to it being seen as an int16/32 vs. int8 sized value in different places.  The event binding customization now also uses the new SPropertyBinding widget instead of the old customization code path.  It's now possible to bind widgets to parent functions and parent properties of the user widget, including ones defined in native code.

diff --git a/Engine/Config/BaseEngine.ini b/Engine/Config/BaseEngine.ini
--- a/Engine/Config/BaseEngine.ini
+++ b/Engine/Config/BaseEngine.ini
@@ -598,6 +598,7 @@ bUseStreamingPause=false
 ; Widget Blueprint Redirects
 +ActiveClassRedirects=(OldClassName="WidgetBlueprint",NewClassName="/Script/UMGEditor.WidgetBlueprint")
 +ActiveClassRedirects=(OldClassName="SReply",NewClassName="EventReply")
++K2FieldRedirects=(OldFieldName="Widget.Visiblity", NewFieldName="Widget.Visibility")
 
 ;EAxis rename
 +EnumRedirects=(EnumName="EAxis",OldEnumEntry="AXIS_X",NewEnumEntry="EAxis::X")
diff --git a/Engine/Source/Editor/UMGEditor/Private/Customizations/DetailCustomizations.cpp b/Engine/Source/Editor/UMGEditor/Private/Customizations/DetailCustomizations.cpp
--- a/Engine/Source/Editor/UMGEditor/Private/Customizations/DetailCustomizations.cpp
+++ b/Engine/Source/Editor/UMGEditor/Private/Customizations/DetailCustomizations.cpp
@@ -17,6 +17,7 @@
 #include "EdGraphSchema_K2_Actions.h"
 #include "WidgetBlueprintEditor.h"
 #include "Components/PanelSlot.h"
+#include "SPropertyBinding.h"
 
 #define LOCTEXT_NAMESPACE "UMG"
 
@@ -28,7 +29,7 @@ public:
 	SLATE_BEGIN_ARGS(SGraphSchemaActionButton) {}
 	SLATE_END_ARGS()
 
-		void Construct(const FArguments& InArgs, TSharedPtr<FWidgetBlueprintEditor> InEditor, TSharedPtr<FEdGraphSchemaAction> InClickAction)
+	void Construct(const FArguments& InArgs, TSharedPtr<FWidgetBlueprintEditor> InEditor, TSharedPtr<FEdGraphSchemaAction> InClickAction)
 	{
 		Editor = InEditor;
 		Action = InClickAction;
@@ -81,110 +82,9 @@ private:
 	TSharedPtr<FEdGraphSchemaAction> Action;
 };
 
-
-void FBlueprintWidgetCustomization::CreateDelegateCustomization( IDetailLayoutBuilder& DetailLayout, UDelegateProperty* Property, UWidget* Widget )
-{
-	FString ConstPropertyName = Property->GetName();
-	ConstPropertyName.RemoveFromEnd(TEXT("Delegate"));
-
-	//TSharedRef<IPropertyHandle> DelegatePropertyHandle = DetailLayout.GetProperty(Property->GetFName(), PropertyClass);
-	TSharedRef<IPropertyHandle> ConstPropertyHandle = DetailLayout.GetProperty(FName(*ConstPropertyName), CastChecked<UClass>(Property->GetOuter()));
-
-	const bool bHasValidHandle = ConstPropertyHandle->IsValidHandle();
-	if(!bHasValidHandle)
-	{
-		return;
-	}
-
-	UProperty* ConstProperty = ConstPropertyHandle->GetProperty();
-
-	const bool bIsEditable = ConstProperty->HasAnyPropertyFlags(CPF_Edit | CPF_EditConst);
-	const bool bDoSignaturesMatch = Property->SignatureFunction->GetReturnProperty()->SameType(ConstProperty);
-
-	if(!(bIsEditable && bDoSignaturesMatch))
-	{
-		return;
-	}
-
-	IDetailCategoryBuilder& PropertyCategory = DetailLayout.EditCategory(FObjectEditorUtils::GetCategoryFName(ConstProperty));
-
-	IDetailPropertyRow& PropertyRow = PropertyCategory.AddProperty(ConstPropertyHandle);
-
-	TSharedPtr<SWidget> NameWidget;
-	TSharedPtr<SWidget> ValueWidget;
-	FDetailWidgetRow Row;
-	PropertyRow.GetDefaultWidgets(NameWidget, ValueWidget, Row);
-
-	TSharedRef<SComboButton> DelegateButton = SNew(SComboButton)
-		.OnGetMenuContent(this, &FBlueprintWidgetCustomization::OnGenerateDelegateMenu, Widget, ConstPropertyHandle, Property->SignatureFunction, true)
-		.ContentPadding(1)
-		.ButtonContent()
-		[
-			SNew(SHorizontalBox)
-			+ SHorizontalBox::Slot()
-			.AutoWidth()
-			.VAlign(VAlign_Center)
-			[
-				SNew(SImage)
-				.Image(this, &FBlueprintWidgetCustomization::GetCurrentBindingImage, ConstPropertyHandle)
-				.ColorAndOpacity(FLinearColor(0.25f, 0.25f, 0.25f))
-			]
-			+ SHorizontalBox::Slot()
-			.AutoWidth()
-			.VAlign(VAlign_Center)
-			.Padding(4, 1, 0, 0)
-			[
-				SNew(STextBlock)
-				.Text(this, &FBlueprintWidgetCustomization::GetCurrentBindingText, ConstPropertyHandle)
-				.Font(IDetailLayoutBuilder::GetDetailFont())
-			]
-		];
-
-	const bool bShowChildren = true;
-	PropertyRow.CustomWidget(bShowChildren)
-		.NameContent()
-		.MinDesiredWidth(Row.NameWidget.MinWidth)
-		.MaxDesiredWidth(Row.NameWidget.MaxWidth)
-		[
-			NameWidget.ToSharedRef()
-		]
-		.ValueContent()
-		.MinDesiredWidth(Row.ValueWidget.MinWidth * 2)
-		.MaxDesiredWidth(Row.ValueWidget.MaxWidth * 2)
-		[
-			SNew(SHorizontalBox)
-			+ SHorizontalBox::Slot()
-			.FillWidth(1)
-			.Padding(FMargin(0, 0, 4, 0))
-			[
-				ValueWidget.ToSharedRef()
-			]
-			+ SHorizontalBox::Slot()
-			.AutoWidth()
-			.VAlign(VAlign_Center)
-			[
-				DelegateButton
-			]
-			+ SHorizontalBox::Slot()
-			.AutoWidth()
-			[
-				SNew(SButton)
-				.ButtonStyle(FEditorStyle::Get(), "HoverHintOnly")
-				.Visibility(this, &FBlueprintWidgetCustomization::GetGotoBindingVisibility, ConstPropertyHandle)
-				.OnClicked(this, &FBlueprintWidgetCustomization::HandleGotoBindingClicked, ConstPropertyHandle)
-				.VAlign(VAlign_Center)
-				.ToolTipText(LOCTEXT("GotoFunction", "Goto Function"))
-				[
-					SNew(SImage)
-					.Image(FEditorStyle::GetBrush("PropertyWindow.Button_Browse"))
-				]
-			]
-		];
-}
-
 void FBlueprintWidgetCustomization::CreateEventCustomization( IDetailLayoutBuilder& DetailLayout, UDelegateProperty* Property, UWidget* Widget )
 {
-	if(!Property->GetName().EndsWith(TEXT("Event")))
+	if ( !Property->GetName().EndsWith(TEXT("Event")) )
 	{
 		return;
 	}
@@ -201,21 +101,6 @@ void FBlueprintWidgetCustomization::CreateEventCustomization( IDetailLayoutBuild
 
 	IDetailPropertyRow& PropertyRow = PropertyCategory.AddProperty(DelegatePropertyHandle);
 
-	TSharedPtr<SWidget> NameWidget;
-	TSharedPtr<SWidget> ValueWidget;
-	FDetailWidgetRow Row;
-	PropertyRow.GetDefaultWidgets(NameWidget, ValueWidget, Row);
-
-	TSharedRef<SComboButton> DelegateButton = SNew(SComboButton)
-		.OnGetMenuContent(this, &FBlueprintWidgetCustomization::OnGenerateDelegateMenu, Widget, DelegatePropertyHandle, Property->SignatureFunction, false)
-		.ContentPadding(1)
-		.ButtonContent()
-		[
-			SNew(STextBlock)
-			.Text(this, &FBlueprintWidgetCustomization::GetCurrentBindingText, DelegatePropertyHandle)
-			.Font(IDetailLayoutBuilder::GetDetailFont())
-		];
-
 	FString LabelStr = Property->GetName();
 	LabelStr.RemoveFromEnd(TEXT("Event"));
 
@@ -224,8 +109,6 @@ void FBlueprintWidgetCustomization::CreateEventCustomization( IDetailLayoutBuild
 	const bool bShowChildren = true;
 	PropertyRow.CustomWidget(bShowChildren)
 		.NameContent()
-		.MinDesiredWidth(Row.NameWidget.MinWidth)
-		.MaxDesiredWidth(Row.NameWidget.MaxWidth)
 		[
 			SNew(SHorizontalBox)
 
@@ -249,29 +132,8 @@ void FBlueprintWidgetCustomization::CreateEventCustomization( IDetailLayoutBuild
 		.MinDesiredWidth(200)
 		.MaxDesiredWidth(200)
 		[
-			SNew(SHorizontalBox)
-
-			+ SHorizontalBox::Slot()
-			.FillWidth(1)
-			.VAlign(VAlign_Center)
-			[
-				DelegateButton
-			]
-
-			+ SHorizontalBox::Slot()
-			.AutoWidth()
-			[
-				SNew(SButton)
-				.ButtonStyle(FEditorStyle::Get(), "HoverHintOnly")
-				.Visibility(this, &FBlueprintWidgetCustomization::GetGotoBindingVisibility, DelegatePropertyHandle)
-				.OnClicked(this, &FBlueprintWidgetCustomization::HandleGotoBindingClicked, DelegatePropertyHandle)
-				.VAlign(VAlign_Center)
-				.ToolTipText(LOCTEXT("GotoFunction", "Goto Function"))
-				[
-					SNew(SImage)
-					.Image(FEditorStyle::GetBrush("PropertyWindow.Button_Browse"))
-				]
-			]
+			SNew(SPropertyBinding, Editor.Pin().ToSharedRef(), Property, DelegatePropertyHandle)
+			.GeneratePureBindings(false)
 		];
 }
 
@@ -404,11 +266,7 @@ void FBlueprintWidgetCustomization::PerformBindingCustomization(IDetailLayoutBui
 
 			if ( UDelegateProperty* DelegateProperty = Cast<UDelegateProperty>(*PropertyIt) )
 			{
-				if ( DelegateProperty->GetName().EndsWith(TEXT("Delegate")) )
-				{
-					//CreateDelegateCustomization(DetailLayout, DelegateProperty, Widget);
-				}
-				else if ( DelegateProperty->GetName().EndsWith(TEXT("Event")) )
+				if ( DelegateProperty->GetName().EndsWith(TEXT("Event")) )
 				{
 					CreateEventCustomization(DetailLayout, DelegateProperty, Widget);
 				}
@@ -421,455 +279,4 @@ void FBlueprintWidgetCustomization::PerformBindingCustomization(IDetailLayoutBui
 	}
 }
 
-void FBlueprintWidgetCustomization::RefreshBlueprintMemberCache(const UFunction* DelegateSignature)
-{
-	const UWidgetGraphSchema* Schema = GetDefault<UWidgetGraphSchema>();
-	const FSlateFontInfo DetailFontInfo = IDetailLayoutBuilder::GetDetailFont();
-
-	BlueprintFunctionCache.Reset();
-	BlueprintPropertyCache.Reset();
-
-	// Get the current skeleton class, think header for the blueprint.
-	UBlueprintGeneratedClass* SkeletonClass = Cast<UBlueprintGeneratedClass>(Blueprint->SkeletonGeneratedClass);
-
-	// Grab functions implemented by the blueprint
-	for ( int32 i = 0; i < Blueprint->FunctionGraphs.Num(); i++ )
-	{
-		UEdGraph* Graph = Blueprint->FunctionGraphs[i];
-		check(Graph);
-
-		// Ignore the construction script, we just want user created functions.
-		if ( Graph->GetFName() == Schema->FN_UserConstructionScript )
-		{
-			continue;
-		}
-
-		FName FunctionFName = Graph->GetFName();
-
-		UFunction* Function = SkeletonClass->FindFunctionByName(FunctionFName, EIncludeSuperFlag::IncludeSuper);
-		if ( Function == NULL )
-		{
-			continue;
-		}
-
-		// We ignore CPF_ReturnParm because all that matters for binding to script functions is that the number of out parameters match.
-		if ( Function->IsSignatureCompatibleWith(DelegateSignature, UFunction::GetDefaultIgnoredSignatureCompatibilityFlags() | CPF_ReturnParm) )
-		{
-			FGraphDisplayInfo DisplayInfo;
-			Graph->GetSchema()->GetGraphDisplayInformation(*Graph, DisplayInfo);
-
-			TSharedPtr<FunctionInfo> NewFuncAction = MakeShareable(new FunctionInfo());
-			NewFuncAction->DisplayName = DisplayInfo.PlainName;
-			NewFuncAction->Tooltip = DisplayInfo.Tooltip;
-			NewFuncAction->FuncName = FunctionFName;
-			NewFuncAction->EdGraph = Graph;
-
-			BlueprintFunctionCache.Add(NewFuncAction);
-		}
-	}
-
-	// Grab functions implemented by the blueprint
-	for ( TFieldIterator<UProperty> PropIt(SkeletonClass, EFieldIteratorFlags::ExcludeSuper); PropIt; ++PropIt )
-	{
-		UProperty* Prop = *PropIt;
-
-		if ( UProperty* ReturnProperty = DelegateSignature->GetReturnProperty() )
-		{
-			if ( ReturnProperty->SameType(Prop) )
-			{
-				if ( Prop->HasAnyPropertyFlags(UP::BlueprintReadWrite) )
-				{
-					BlueprintPropertyCache.Add(Prop);
-				}
-			}
-		}
-	}
-}
-
-TSharedRef<SWidget> FBlueprintWidgetCustomization::OnGenerateDelegateMenu(UWidget* Widget, TSharedRef<IPropertyHandle> PropertyHandle, UFunction* DelegateSignature, bool bIsPure)
-{
-	RefreshBlueprintMemberCache(DelegateSignature);
-
-	const bool bInShouldCloseWindowAfterMenuSelection = true;
-	FMenuBuilder MenuBuilder(bInShouldCloseWindowAfterMenuSelection, NULL);
-
-	static FName PropertyIcon(TEXT("Kismet.Tabs.Variables"));
-	static FName FunctionIcon(TEXT("GraphEditor.Function_16x"));
-
-	//FLinearColor ColorOut;
-	//const UClass* VarClass = PropertyHandle->GetPropertyClass();
-	//BrushOut = FBlueprintEditor::GetVarIconAndColor(VarClass, PropertyHandle->GetProperty()->GetFName(), ColorOut);
-
-	if ( CanRemoveBinding(PropertyHandle) )
-	{
-		MenuBuilder.BeginSection("RemoveBinding");
-		{
-			MenuBuilder.AddMenuEntry(
-				LOCTEXT("RemoveBinding", "Remove Binding"),
-				LOCTEXT("RemoveBindingToolTip", "Removes the current binding"),
-				FSlateIcon(),
-				FUIAction(FExecuteAction::CreateSP(this, &FBlueprintWidgetCustomization::HandleRemoveBinding, PropertyHandle))
-				);
-		}
-		MenuBuilder.EndSection(); //RemoveBinding
-	}
-
-	MenuBuilder.BeginSection("CreateBinding");
-	{
-		MenuBuilder.AddMenuEntry(
-			LOCTEXT("CreateBinding", "Create Binding"),
-			LOCTEXT("CreateBindingToolTip", "Creates a new function for this property to be bound to"),
-			FSlateIcon(),
-			FUIAction(FExecuteAction::CreateSP(this, &FBlueprintWidgetCustomization::HandleCreateAndAddBinding, Widget, PropertyHandle, DelegateSignature, bIsPure))
-			);
-	}
-	MenuBuilder.EndSection(); //CreateBinding
-
-	MenuBuilder.BeginSection("Functions", LOCTEXT("Functions", "Functions"));
-	{
-		for ( TSharedPtr<FunctionInfo>& Function : BlueprintFunctionCache )
-		{
-			MenuBuilder.AddMenuEntry(
-				Function->DisplayName,
-				FText::FromString(Function->Tooltip),
-				FSlateIcon(FEditorStyle::GetStyleSetName(), FunctionIcon),
-				FUIAction(FExecuteAction::CreateSP(this, &FBlueprintWidgetCustomization::HandleAddFunctionBinding, PropertyHandle, Function))
-				);
-		}
-	}
-	MenuBuilder.EndSection(); //Functions
-
-	MenuBuilder.BeginSection("Properties", LOCTEXT("Properties", "Properties"));
-	{
-		for ( UProperty* ExistingProperty : BlueprintPropertyCache )
-		{
-			MenuBuilder.AddMenuEntry(
-				ExistingProperty->GetDisplayNameText(),
-				ExistingProperty->GetToolTipText(),
-				FSlateIcon(FEditorStyle::GetStyleSetName(), PropertyIcon),
-				FUIAction(FExecuteAction::CreateSP(this, &FBlueprintWidgetCustomization::HandleAddPropertyBinding, PropertyHandle, ExistingProperty))
-				);
-		}
-	}
-	MenuBuilder.EndSection(); //Functions
-
-
-	FDisplayMetrics DisplayMetrics;
-	FSlateApplication::Get().GetDisplayMetrics(DisplayMetrics);
-
-	return
-		SNew(SVerticalBox)
-
-		+ SVerticalBox::Slot()
-		.MaxHeight(DisplayMetrics.PrimaryDisplayHeight * 0.5)
-		[
-			MenuBuilder.MakeWidget()
-		];
-}
-
-const FSlateBrush* FBlueprintWidgetCustomization::GetCurrentBindingImage(TSharedRef<IPropertyHandle> PropertyHandle) const
-{
-	static FName PropertyIcon(TEXT("Kismet.Tabs.Variables"));
-	static FName FunctionIcon(TEXT("GraphEditor.Function_16x"));
-
-	TArray<UObject*> OuterObjects;
-	PropertyHandle->GetOuterObjects(OuterObjects);
-
-	//TODO UMG O(N) Isn't good for this, needs to be map, but map isn't serialized, need cached runtime map for fast lookups.
-
-	FName PropertyName = PropertyHandle->GetProperty()->GetFName();
-	for ( int32 ObjectIndex = 0; ObjectIndex < OuterObjects.Num(); ObjectIndex++ )
-	{
-		// Ignore null outer objects
-		if ( OuterObjects[ObjectIndex] == NULL )
-		{
-			continue;
-		}
-
-		//TODO UMG handle multiple things selected
-
-		for ( const FDelegateEditorBinding& Binding : Blueprint->Bindings )
-		{
-			if ( Binding.ObjectName == OuterObjects[ObjectIndex]->GetName() && Binding.PropertyName == PropertyName )
-			{
-				if ( Binding.Kind == EBindingKind::Function )
-				{
-					return FEditorStyle::GetBrush(FunctionIcon);
-				}
-				else // Property
-				{
-					return FEditorStyle::GetBrush(PropertyIcon);
-				}
-			}
-		}
-	}
-
-	return NULL;
-}
-
-FText FBlueprintWidgetCustomization::GetCurrentBindingText(TSharedRef<IPropertyHandle> PropertyHandle) const
-{
-	TArray<UObject*> OuterObjects;
-	PropertyHandle->GetOuterObjects(OuterObjects);
-
-	//TODO UMG O(N) Isn't good for this, needs to be map, but map isn't serialized, need cached runtime map for fast lookups.
-
-	FName PropertyName = PropertyHandle->GetProperty()->GetFName();
-	for ( int32 ObjectIndex = 0; ObjectIndex < OuterObjects.Num(); ObjectIndex++ )
-	{
-		// Ignore null outer objects
-		if ( OuterObjects[ObjectIndex] == NULL )
-		{
-			continue;
-		}
-
-		//TODO UMG handle multiple things selected
-
-		for ( const FDelegateEditorBinding& Binding : Blueprint->Bindings )
-		{
-			if ( Binding.ObjectName == OuterObjects[ObjectIndex]->GetName() && Binding.PropertyName == PropertyName )
-			{
-				if ( Binding.Kind == EBindingKind::Function )
-				{
-					FName FoundName = Blueprint->GetFieldNameFromClassByGuid<UFunction>(Blueprint->GeneratedClass, Binding.MemberGuid);
-					return FText::FromString(FName::NameToDisplayString(FoundName.ToString(), false));
-				}
-				else // Property
-				{
-					FName FoundName = Blueprint->GetFieldNameFromClassByGuid<UProperty>(Blueprint->GeneratedClass, Binding.MemberGuid);
-					return FText::FromString( FName::NameToDisplayString(FoundName.ToString(), false) );
-				}
-			}
-		}
-
-		//TODO UMG Do something about missing functions, little exclamation points if they're missing and such.
-
-		break;
-	}
-
-	return LOCTEXT("Bind", "Bind");
-}
-
-bool FBlueprintWidgetCustomization::CanRemoveBinding(TSharedRef<IPropertyHandle> PropertyHandle)
-{
-	FName PropertyName = PropertyHandle->GetProperty()->GetFName();
-
-	TArray<UObject*> OuterObjects;
-	PropertyHandle->GetOuterObjects(OuterObjects);
-	for ( UObject* SelectedObject : OuterObjects )
-	{
-		for ( const FDelegateEditorBinding& Binding : Blueprint->Bindings )
-		{
-			if ( Binding.ObjectName == SelectedObject->GetName() && Binding.PropertyName == PropertyName )
-			{
-				return true;
-			}
-		}
-	}
-
-	return false;
-}
-
-void FBlueprintWidgetCustomization::HandleRemoveBinding(TSharedRef<IPropertyHandle> PropertyHandle)
-{
-	const FScopedTransaction Transaction(LOCTEXT("BindDelegate", "Remove Binding"));
-
-	Blueprint->Modify();
-
-	TArray<UObject*> OuterObjects;
-	PropertyHandle->GetOuterObjects(OuterObjects);
-	for ( UObject* SelectedObject : OuterObjects )
-	{
-		FDelegateEditorBinding Binding;
-		Binding.ObjectName = SelectedObject->GetName();
-		Binding.PropertyName = PropertyHandle->GetProperty()->GetFName();
-
-		Blueprint->Bindings.Remove(Binding);
-	}
-
-	FBlueprintEditorUtils::MarkBlueprintAsStructurallyModified(Blueprint);
-}
-
-void FBlueprintWidgetCustomization::HandleAddFunctionBinding(TSharedRef<IPropertyHandle> PropertyHandle, TSharedPtr<FunctionInfo> SelectedFunction)
-{
-	const FScopedTransaction Transaction(LOCTEXT("BindDelegate", "Set Binding"));
-
-	Blueprint->Modify();
-
-	TArray<UObject*> OuterObjects;
-	PropertyHandle->GetOuterObjects(OuterObjects);
-	for ( UObject* SelectedObject : OuterObjects )
-	{
-		FDelegateEditorBinding Binding;
-		Binding.ObjectName = SelectedObject->GetName();
-		Binding.PropertyName = PropertyHandle->GetProperty()->GetFName();
-		Binding.FunctionName = SelectedFunction->FuncName;
-		Binding.MemberGuid = SelectedFunction->EdGraph->GraphGuid;
-		Binding.Kind = EBindingKind::Function;
-
-		Blueprint->Bindings.Remove(Binding);
-		Blueprint->Bindings.AddUnique(Binding);
-	}
-
-	FBlueprintEditorUtils::MarkBlueprintAsStructurallyModified(Blueprint);
-}
-
-void FBlueprintWidgetCustomization::HandleAddPropertyBinding(TSharedRef<IPropertyHandle> PropertyHandle, UProperty* SelectedProperty)
-{
-	const FScopedTransaction Transaction(LOCTEXT("BindDelegate", "Set Binding"));
-
-	// Get the current skeleton class, think header for the blueprint.
-	UBlueprintGeneratedClass* SkeletonClass = Cast<UBlueprintGeneratedClass>(Blueprint->SkeletonGeneratedClass);
-
-	Blueprint->Modify();
-
-	FGuid MemberGuid;
-	UBlueprint::GetGuidFromClassByFieldName<UProperty>(SkeletonClass, SelectedProperty->GetFName(), MemberGuid);
-
-	TArray<UObject*> OuterObjects;
-	PropertyHandle->GetOuterObjects(OuterObjects);
-	for ( UObject* SelectedObject : OuterObjects )
-	{
-		FDelegateEditorBinding Binding;
-		Binding.ObjectName = SelectedObject->GetName();
-		Binding.PropertyName = PropertyHandle->GetProperty()->GetFName();
-		Binding.SourceProperty = SelectedProperty->GetFName();
-		Binding.MemberGuid = MemberGuid;
-		Binding.Kind = EBindingKind::Property;
-
-		Blueprint->Bindings.Remove(Binding);
-		Blueprint->Bindings.AddUnique(Binding);
-	}
-
-	FBlueprintEditorUtils::MarkBlueprintAsStructurallyModified(Blueprint);
-}
-
-void FBlueprintWidgetCustomization::HandleCreateAndAddBinding(UWidget* Widget, TSharedRef<IPropertyHandle> PropertyHandle, UFunction* DelegateSignature, bool bIsPure)
-{
-	const FScopedTransaction Transaction(LOCTEXT("CreateDelegate", "Create Binding"));
-
-	Blueprint->Modify();
-
-	FString Pre = bIsPure ? FString(TEXT("Get")) : FString(TEXT("On"));
-
-	FString WidgetName;
-	if ( Widget && !Widget->IsGeneratedName() )
-	{
-		WidgetName = TEXT("_") + Widget->GetName() + TEXT("_");
-	}
-
-	FString Post = PropertyHandle->GetProperty()->GetName();
-	Post.RemoveFromStart(TEXT("On"));
-	Post.RemoveFromEnd(TEXT("Event"));
-
-	// Create the function graph.
-	FString FunctionName = Pre + WidgetName + Post;
-	UEdGraph* FunctionGraph = FBlueprintEditorUtils::CreateNewGraph(Blueprint, FBlueprintEditorUtils::FindUniqueKismetName(Blueprint, FunctionName), UEdGraph::StaticClass(), UEdGraphSchema_K2::StaticClass());
-
-	// Only mark bindings as pure that need to be.
-	if ( bIsPure )
-	{
-		const UEdGraphSchema_K2* Schema_K2 = Cast<UEdGraphSchema_K2>(FunctionGraph->GetSchema());
-		Schema_K2->AddExtraFunctionFlags(FunctionGraph, FUNC_BlueprintPure);
-	}
-
-	// Add the binding to the blueprint
-	TSharedPtr<FunctionInfo> SelectedFunction = MakeShareable(new FunctionInfo());
-	SelectedFunction->FuncName = FunctionGraph->GetFName();
-	SelectedFunction->EdGraph = FunctionGraph;
-
-	HandleAddFunctionBinding(PropertyHandle, SelectedFunction);
-
-	bool bUserCreated = true;
-	FBlueprintEditorUtils::AddFunctionGraph(Blueprint, FunctionGraph, bUserCreated, DelegateSignature);
-
-	GotoFunction(FunctionGraph);
-}
-
-EVisibility FBlueprintWidgetCustomization::GetGotoBindingVisibility(TSharedRef<IPropertyHandle> PropertyHandle) const
-{
-	TArray<UObject*> OuterObjects;
-	PropertyHandle->GetOuterObjects(OuterObjects);
-
-	//TODO UMG O(N) Isn't good for this, needs to be map, but map isn't serialized, need cached runtime map for fast lookups.
-
-	FName PropertyName = PropertyHandle->GetProperty()->GetFName();
-	for ( int32 ObjectIndex = 0; ObjectIndex < OuterObjects.Num(); ObjectIndex++ )
-	{
-		// Ignore null outer objects
-		if ( OuterObjects[ObjectIndex] == NULL )
-		{
-			continue;
-		}
-
-		//TODO UMG handle multiple things selected
-
-		for ( const FDelegateEditorBinding& Binding : Blueprint->Bindings )
-		{
-			if ( Binding.ObjectName == OuterObjects[ObjectIndex]->GetName() && Binding.PropertyName == PropertyName )
-			{
-				if ( Binding.Kind == EBindingKind::Function )
-				{
-					return EVisibility::Visible;
-				}
-			}
-		}
-	}
-
-	return EVisibility::Collapsed;
-}
-
-FReply FBlueprintWidgetCustomization::HandleGotoBindingClicked(TSharedRef<IPropertyHandle> PropertyHandle)
-{
-	TArray<UObject*> OuterObjects;
-	PropertyHandle->GetOuterObjects(OuterObjects);
-
-	//TODO UMG O(N) Isn't good for this, needs to be map, but map isn't serialized, need cached runtime map for fast lookups.
-
-	FName PropertyName = PropertyHandle->GetProperty()->GetFName();
-	for ( int32 ObjectIndex = 0; ObjectIndex < OuterObjects.Num(); ObjectIndex++ )
-	{
-		// Ignore null outer objects
-		if ( OuterObjects[ObjectIndex] == NULL )
-		{
-			continue;
-		}
-
-		//TODO UMG handle multiple things selected
-
-		for ( const FDelegateEditorBinding& Binding : Blueprint->Bindings )
-		{
-			if ( Binding.ObjectName == OuterObjects[ObjectIndex]->GetName() && Binding.PropertyName == PropertyName )
-			{
-				if ( Binding.Kind == EBindingKind::Function )
-				{
-					TArray<UEdGraph*> AllGraphs;
-					Blueprint->GetAllGraphs(AllGraphs);
-
-					for ( UEdGraph* Graph : AllGraphs )
-					{
-						if ( Graph->GraphGuid == Binding.MemberGuid )
-						{
-							GotoFunction(Graph);
-						}
-					}
-
-
-					// Either way return
-					return FReply::Handled();
-				}
-			}
-		}
-	}
-
-	return FReply::Unhandled();
-}
-
-void FBlueprintWidgetCustomization::GotoFunction(UEdGraph* FunctionGraph)
-{
-	Editor.Pin()->SetCurrentMode(FWidgetBlueprintApplicationModes::GraphMode);
-
-	Editor.Pin()->OpenDocument(FunctionGraph, FDocumentTracker::OpenNewDocument);
-}
-
 #undef LOCTEXT_NAMESPACE
diff --git a/Engine/Source/Editor/UMGEditor/Private/Details/DetailWidgetExtensionHandler.cpp b/Engine/Source/Editor/UMGEditor/Private/Details/DetailWidgetExtensionHandler.cpp
--- a/Engine/Source/Editor/UMGEditor/Private/Details/DetailWidgetExtensionHandler.cpp
+++ b/Engine/Source/Editor/UMGEditor/Private/Details/DetailWidgetExtensionHandler.cpp
@@ -54,5 +54,6 @@ TSharedRef<SWidget> FDetailWidgetExtensionHandler::GenerateExtensionWidget(const
 		return SNullWidget::NullWidget;
 	}
 
-	return SNew(SPropertyBinding, BlueprintEditor.Pin().ToSharedRef(), DelegateProperty, InPropertyHandle.ToSharedRef());
+	return SNew(SPropertyBinding, BlueprintEditor.Pin().ToSharedRef(), DelegateProperty, InPropertyHandle.ToSharedRef())
+		.GeneratePureBindings(true);
 }
diff --git a/Engine/Source/Editor/UMGEditor/Private/Details/SPropertyBinding.cpp b/Engine/Source/Editor/UMGEditor/Private/Details/SPropertyBinding.cpp
--- a/Engine/Source/Editor/UMGEditor/Private/Details/SPropertyBinding.cpp
+++ b/Engine/Source/Editor/UMGEditor/Private/Details/SPropertyBinding.cpp
@@ -129,18 +129,53 @@ void SPropertyBinding::RefreshBlueprintMemberCache(const UFunction* DelegateSign
 		}
 	}
 
-	// Grab functions implemented by the blueprint
-	for ( TFieldIterator<UProperty> PropIt(SkeletonClass, EFieldIteratorFlags::ExcludeSuper); PropIt; ++PropIt )
+	// Walk up class hierarchy for native functions and properties
+	for ( UClass* Class = SkeletonClass; Class != UUserWidget::StaticClass(); Class = Class->GetSuperClass() )
 	{
-		UProperty* Prop = *PropIt;
+		// Add matching native functions
+		for ( TFieldIterator<UFunction> FuncIt(Class, EFieldIteratorFlags::ExcludeSuper); FuncIt; ++FuncIt )
+		{
+			UFunction* Func = *FuncIt;
+
+			if ( !Func->HasAnyFunctionFlags(UF::BlueprintCallable | UF::BlueprintPure) )
+			{
+				continue;
+			}
+
+			FName FunctionFName = Func->GetFName();
+
+			UFunction* Function = Class->FindFunctionByName(FunctionFName, EIncludeSuperFlag::IncludeSuper);
+			if ( Function == nullptr )
+			{
+				continue;
+			}
+
+			// We ignore CPF_ReturnParm because all that matters for binding to script functions is that the number of out parameters match.
+			if ( Function->IsSignatureCompatibleWith(DelegateSignature, UFunction::GetDefaultIgnoredSignatureCompatibilityFlags() | CPF_ReturnParm) )
+			{
+				TSharedPtr<FunctionInfo> NewFuncAction = MakeShareable(new FunctionInfo());
+				NewFuncAction->DisplayName = FText::FromName(Func->GetFName());
+				NewFuncAction->Tooltip = Func->GetMetaData("Tooltip");
+				NewFuncAction->FuncName = FunctionFName;
+				NewFuncAction->EdGraph = nullptr;
+
+				BlueprintFunctionCache.Add(NewFuncAction);
+			}
+		}
 
-		if ( UProperty* ReturnProperty = DelegateSignature->GetReturnProperty() )
+		// Grab functions implemented by the blueprint
+		for ( TFieldIterator<UProperty> PropIt(SkeletonClass, EFieldIteratorFlags::ExcludeSuper); PropIt; ++PropIt )
 		{
-			if ( ReturnProperty->SameType(Prop) )
+			UProperty* Prop = *PropIt;
+
+			if ( UProperty* ReturnProperty = DelegateSignature->GetReturnProperty() )
 			{
-				if ( Prop->HasAnyPropertyFlags(UP::BlueprintReadWrite) )
+				if ( ReturnProperty->SameType(Prop) )
 				{
-					BlueprintPropertyCache.Add(Prop);
+					if ( Prop->HasAnyPropertyFlags(UP::BlueprintReadWrite) )
+					{
+						BlueprintPropertyCache.Add(Prop);
+					}
 				}
 			}
 		}
@@ -152,7 +187,7 @@ TSharedRef<SWidget> SPropertyBinding::OnGenerateDelegateMenu(UWidget* Widget, TS
 	RefreshBlueprintMemberCache(DelegateSignature, bIsPure);
 
 	const bool bInShouldCloseWindowAfterMenuSelection = true;
-	FMenuBuilder MenuBuilder(bInShouldCloseWindowAfterMenuSelection, NULL);
+	FMenuBuilder MenuBuilder(bInShouldCloseWindowAfterMenuSelection, nullptr);
 
 	static FName PropertyIcon(TEXT("Kismet.Tabs.Variables"));
 	static FName FunctionIcon(TEXT("GraphEditor.Function_16x"));
@@ -265,7 +300,7 @@ const FSlateBrush* SPropertyBinding::GetCurrentBindingImage(TSharedRef<IProperty
 		}
 	}
 
-	return NULL;
+	return nullptr;
 }
 
 FText SPropertyBinding::GetCurrentBindingText(TSharedRef<IPropertyHandle> PropertyHandle) const
@@ -279,7 +314,7 @@ FText SPropertyBinding::GetCurrentBindingText(TSharedRef<IPropertyHandle> Proper
 	for ( int32 ObjectIndex = 0; ObjectIndex < OuterObjects.Num(); ObjectIndex++ )
 	{
 		// Ignore null outer objects
-		if ( OuterObjects[ObjectIndex] == NULL )
+		if ( OuterObjects[ObjectIndex] == nullptr )
 		{
 			continue;
 		}
@@ -292,13 +327,30 @@ FText SPropertyBinding::GetCurrentBindingText(TSharedRef<IPropertyHandle> Proper
 			{
 				if ( Binding.Kind == EBindingKind::Function )
 				{
-					FName FoundName = Blueprint->GetFieldNameFromClassByGuid<UFunction>(Blueprint->GeneratedClass, Binding.MemberGuid);
-					return FText::FromString(FName::NameToDisplayString(FoundName.ToString(), false));
+					if ( Binding.MemberGuid.IsValid() )
+					{
+						// Graph function, look up by Guid
+						FName FoundName = Blueprint->GetFieldNameFromClassByGuid<UFunction>(Blueprint->GeneratedClass, Binding.MemberGuid);
+						return FText::FromString(FName::NameToDisplayString(FoundName.ToString(), false));
+					}
+					else
+					{
+						// No GUID, native function, return function name.
+						return FText::FromName(Binding.FunctionName);
+					}
 				}
 				else // Property
 				{
-					FName FoundName = Blueprint->GetFieldNameFromClassByGuid<UProperty>(Blueprint->GeneratedClass, Binding.MemberGuid);
-					return FText::FromString(FName::NameToDisplayString(FoundName.ToString(), false));
+					if ( Binding.MemberGuid.IsValid() )
+					{
+						FName FoundName = Blueprint->GetFieldNameFromClassByGuid<UProperty>(Blueprint->GeneratedClass, Binding.MemberGuid);
+						return FText::FromString(FName::NameToDisplayString(FoundName.ToString(), false));
+					}
+					else
+					{
+						// No GUID, native property, return source property.
+						return FText::FromName(Binding.SourceProperty);
+					}
 				}
 			}
 		}
@@ -365,7 +417,7 @@ void SPropertyBinding::HandleAddFunctionBinding(TSharedRef<IPropertyHandle> Prop
 		Binding.ObjectName = SelectedObject->GetName();
 		Binding.PropertyName = PropertyHandle->GetProperty()->GetFName();
 		Binding.FunctionName = SelectedFunction->FuncName;
-		Binding.MemberGuid = SelectedFunction->EdGraph->GraphGuid;
+		Binding.MemberGuid = ( SelectedFunction->EdGraph ) ? SelectedFunction->EdGraph->GraphGuid : FGuid();
 		Binding.Kind = EBindingKind::Function;
 
 		Blueprint->Bindings.Remove(Binding);
@@ -462,7 +514,7 @@ EVisibility SPropertyBinding::GetGotoBindingVisibility(TSharedRef<IPropertyHandl
 	for ( int32 ObjectIndex = 0; ObjectIndex < OuterObjects.Num(); ObjectIndex++ )
 	{
 		// Ignore null outer objects
-		if ( OuterObjects[ObjectIndex] == NULL )
+		if ( OuterObjects[ObjectIndex] == nullptr )
 		{
 			continue;
 		}
@@ -495,7 +547,7 @@ FReply SPropertyBinding::HandleGotoBindingClicked(TSharedRef<IPropertyHandle> Pr
 	for ( int32 ObjectIndex = 0; ObjectIndex < OuterObjects.Num(); ObjectIndex++ )
 	{
 		// Ignore null outer objects
-		if ( OuterObjects[ObjectIndex] == NULL )
+		if ( OuterObjects[ObjectIndex] == nullptr )
 		{
 			continue;
 		}
diff --git a/Engine/Source/Editor/UMGEditor/Private/WidgetBlueprint.cpp b/Engine/Source/Editor/UMGEditor/Private/WidgetBlueprint.cpp
--- a/Engine/Source/Editor/UMGEditor/Private/WidgetBlueprint.cpp
+++ b/Engine/Source/Editor/UMGEditor/Private/WidgetBlueprint.cpp
@@ -77,7 +77,7 @@ FDelegateRuntimeBinding FDelegateEditorBinding::ToRuntimeBinding(UWidgetBlueprin
 	Binding.PropertyName = PropertyName;
 	if ( Kind == EBindingKind::Function )
 	{
-		Binding.FunctionName = Blueprint->GetFieldNameFromClassByGuid<UFunction>(Blueprint->SkeletonGeneratedClass, MemberGuid);
+		Binding.FunctionName = ( MemberGuid.IsValid() ) ? Blueprint->GetFieldNameFromClassByGuid<UFunction>(Blueprint->SkeletonGeneratedClass, MemberGuid) : FunctionName;
 	}
 	else
 	{
@@ -140,6 +140,20 @@ void UWidgetBlueprint::PostLoad()
 
 		AnimationData_DEPRECATED.Empty();
 	}
+
+	if ( GetLinkerUE4Version() < VER_UE4_RENAME_WIDGET_VISIBILITY )
+	{
+		static const FName Visiblity(TEXT("Visiblity"));
+		static const FName Visibility(TEXT("Visibility"));
+
+		for ( FDelegateEditorBinding& Binding : Bindings )
+		{
+			if ( Binding.PropertyName == Visiblity )
+			{
+				Binding.PropertyName = Visibility;
+			}
+		}
+	}
 }
 
 UClass* UWidgetBlueprint::GetBlueprintClass() const
diff --git a/Engine/Source/Editor/UMGEditor/Public/WidgetBlueprint.h b/Engine/Source/Editor/UMGEditor/Public/WidgetBlueprint.h
--- a/Engine/Source/Editor/UMGEditor/Public/WidgetBlueprint.h
+++ b/Engine/Source/Editor/UMGEditor/Public/WidgetBlueprint.h
@@ -15,18 +15,27 @@ struct UMGEDITOR_API FDelegateEditorBinding
 {
 	GENERATED_USTRUCT_BODY()
 
+	/** The member widget the binding is on, must be a direct variable of the UUserWidget. */
 	UPROPERTY()
 	FString ObjectName;
 
+	/** The property on the ObjectName that we are binding to. */
 	UPROPERTY()
 	FName PropertyName;
 
+	/** The optional source object we are bindings to, this could be a potential model object on the widget. */
+	UPROPERTY()
+	FName DataSource;
+
+	/** The function that was generated to return the SourceProperty */
 	UPROPERTY()
 	FName FunctionName;
 
+	/** The property we are bindings to directly on the source object. */
 	UPROPERTY()
 	FName SourceProperty;
 
+	/** If it's an actual Function Graph in the blueprint that we're bound to, there's a GUID we can use to lookup that function, to deal with renames better.  This is that GUID. */
 	UPROPERTY()
 	FGuid MemberGuid;
 
diff --git a/Engine/Source/Runtime/Core/Public/UObject/ObjectVersion.h b/Engine/Source/Runtime/Core/Public/UObject/ObjectVersion.h
--- a/Engine/Source/Runtime/Core/Public/UObject/ObjectVersion.h
+++ b/Engine/Source/Runtime/Core/Public/UObject/ObjectVersion.h
@@ -644,6 +644,8 @@ enum EUnrealEngineObjectUE4Version
 	VER_UE4_GRAPH_INTERACTIVE_COMMENTBUBBLES,
 	// Landscape serializes physical materials for collision objects 
 	VER_UE4_LANDSCAPE_SERIALIZE_PHYSICS_MATERIALS,
+	// Rename Visiblity on widgets to Visibility
+	VER_UE4_RENAME_WIDGET_VISIBILITY,
 	// -----<new versions can be added before this line>-------------------------------------------------
 	// - this needs to be the last line (see note below)
 	VER_UE4_AUTOMATIC_VERSION_PLUS_ONE,
diff --git a/Engine/Source/Runtime/UMG/Private/Components/CanvasPanel.cpp b/Engine/Source/Runtime/UMG/Private/Components/CanvasPanel.cpp
--- a/Engine/Source/Runtime/UMG/Private/Components/CanvasPanel.cpp
+++ b/Engine/Source/Runtime/UMG/Private/Components/CanvasPanel.cpp
@@ -13,7 +13,7 @@ UCanvasPanel::UCanvasPanel(const FObjectInitializer& ObjectInitializer)
 	bIsVariable = false;
 
 	SConstraintCanvas::FArguments Defaults;
-	Visiblity = UWidget::ConvertRuntimeToSerializedVisiblity(Defaults._Visibility.Get());
+	Visibility = UWidget::ConvertRuntimeToSerializedVisibility(Defaults._Visibility.Get());
 }
 
 void UCanvasPanel::ReleaseSlateResources(bool bReleaseChildren)
@@ -66,6 +66,11 @@ TSharedRef<SWidget> UCanvasPanel::RebuildWidget()
 	return BuildDesignTimeWidget( MyCanvas.ToSharedRef() );
 }
 
+UCanvasPanelSlot* UCanvasPanel::AddChildToCanvas(UWidget* Content)
+{
+	return Cast<UCanvasPanelSlot>( Super::AddChild(Content) );
+}
+
 TSharedPtr<SConstraintCanvas> UCanvasPanel::GetCanvasWidget() const
 {
 	return MyCanvas;
@@ -79,7 +84,7 @@ bool UCanvasPanel::GetGeometryForSlot(int32 SlotIndex, FGeometry& ArrangedGeomet
 
 bool UCanvasPanel::GetGeometryForSlot(UCanvasPanelSlot* Slot, FGeometry& ArrangedGeometry) const
 {
-	if ( Slot->Content == NULL )
+	if ( Slot->Content == nullptr )
 	{
 		return false;
 	}
diff --git a/Engine/Source/Runtime/UMG/Private/Components/GridPanel.cpp b/Engine/Source/Runtime/UMG/Private/Components/GridPanel.cpp
--- a/Engine/Source/Runtime/UMG/Private/Components/GridPanel.cpp
+++ b/Engine/Source/Runtime/UMG/Private/Components/GridPanel.cpp
@@ -13,7 +13,7 @@ UGridPanel::UGridPanel(const FObjectInitializer& ObjectInitializer)
 	bIsVariable = false;
 
 	SGridPanel::FArguments Defaults;
-	Visiblity = UWidget::ConvertRuntimeToSerializedVisiblity(Defaults._Visibility.Get());
+	Visibility = UWidget::ConvertRuntimeToSerializedVisibility(Defaults._Visibility.Get());
 }
 
 void UGridPanel::ReleaseSlateResources(bool bReleaseChildren)
diff --git a/Engine/Source/Runtime/UMG/Private/Components/HorizontalBox.cpp b/Engine/Source/Runtime/UMG/Private/Components/HorizontalBox.cpp
--- a/Engine/Source/Runtime/UMG/Private/Components/HorizontalBox.cpp
+++ b/Engine/Source/Runtime/UMG/Private/Components/HorizontalBox.cpp
@@ -13,7 +13,7 @@ UHorizontalBox::UHorizontalBox(const FObjectInitializer& ObjectInitializer)
 	bIsVariable = false;
 
 	SHorizontalBox::FArguments Defaults;
-	Visiblity = UWidget::ConvertRuntimeToSerializedVisiblity(Defaults._Visibility.Get());
+	Visibility = UWidget::ConvertRuntimeToSerializedVisibility(Defaults._Visibility.Get());
 }
 
 void UHorizontalBox::ReleaseSlateResources(bool bReleaseChildren)
diff --git a/Engine/Source/Runtime/UMG/Private/Components/Overlay.cpp b/Engine/Source/Runtime/UMG/Private/Components/Overlay.cpp
--- a/Engine/Source/Runtime/UMG/Private/Components/Overlay.cpp
+++ b/Engine/Source/Runtime/UMG/Private/Components/Overlay.cpp
@@ -13,7 +13,7 @@ UOverlay::UOverlay(const FObjectInitializer& ObjectInitializer)
 	bIsVariable = false;
 
 	SOverlay::FArguments Defaults;
-	Visiblity = UWidget::ConvertRuntimeToSerializedVisiblity(Defaults._Visibility.Get());
+	Visibility = UWidget::ConvertRuntimeToSerializedVisibility(Defaults._Visibility.Get());
 }
 
 void UOverlay::ReleaseSlateResources(bool bReleaseChildren)
diff --git a/Engine/Source/Runtime/UMG/Private/Components/ScrollBox.cpp b/Engine/Source/Runtime/UMG/Private/Components/ScrollBox.cpp
--- a/Engine/Source/Runtime/UMG/Private/Components/ScrollBox.cpp
+++ b/Engine/Source/Runtime/UMG/Private/Components/ScrollBox.cpp
@@ -15,7 +15,7 @@ UScrollBox::UScrollBox(const FObjectInitializer& ObjectInitializer)
 	Orientation = Orient_Vertical;
 
 	SScrollBox::FArguments Defaults;
-	Visiblity = UWidget::ConvertRuntimeToSerializedVisiblity(Defaults._Visibility.Get());
+	Visibility = UWidget::ConvertRuntimeToSerializedVisibility(Defaults._Visibility.Get());
 
 	WidgetStyle = *Defaults._Style;
 	WidgetBarStyle = *Defaults._ScrollBarStyle;
diff --git a/Engine/Source/Runtime/UMG/Private/Components/UniformGridPanel.cpp b/Engine/Source/Runtime/UMG/Private/Components/UniformGridPanel.cpp
--- a/Engine/Source/Runtime/UMG/Private/Components/UniformGridPanel.cpp
+++ b/Engine/Source/Runtime/UMG/Private/Components/UniformGridPanel.cpp
@@ -13,7 +13,7 @@ UUniformGridPanel::UUniformGridPanel(const FObjectInitializer& ObjectInitializer
 	bIsVariable = false;
 
 	SUniformGridPanel::FArguments Defaults;
-	Visiblity = UWidget::ConvertRuntimeToSerializedVisiblity(Defaults._Visibility.Get());
+	Visibility = UWidget::ConvertRuntimeToSerializedVisibility(Defaults._Visibility.Get());
 }
 
 void UUniformGridPanel::ReleaseSlateResources(bool bReleaseChildren)
diff --git a/Engine/Source/Runtime/UMG/Private/Components/VerticalBox.cpp b/Engine/Source/Runtime/UMG/Private/Components/VerticalBox.cpp
--- a/Engine/Source/Runtime/UMG/Private/Components/VerticalBox.cpp
+++ b/Engine/Source/Runtime/UMG/Private/Components/VerticalBox.cpp
@@ -13,7 +13,7 @@ UVerticalBox::UVerticalBox(const FObjectInitializer& ObjectInitializer)
 	bIsVariable = false;
 
 	SVerticalBox::FArguments Defaults;
-	Visiblity = UWidget::ConvertRuntimeToSerializedVisiblity( Defaults._Visibility.Get() );
+	Visibility = UWidget::ConvertRuntimeToSerializedVisibility(Defaults._Visibility.Get());
 }
 
 void UVerticalBox::ReleaseSlateResources(bool bReleaseChildren)
diff --git a/Engine/Source/Runtime/UMG/Private/Components/Widget.cpp b/Engine/Source/Runtime/UMG/Private/Components/Widget.cpp
--- a/Engine/Source/Runtime/UMG/Private/Components/Widget.cpp
+++ b/Engine/Source/Runtime/UMG/Private/Components/Widget.cpp
@@ -13,7 +13,7 @@ UWidget::UWidget(const FObjectInitializer& ObjectInitializer)
 	bIsEnabled = true;
 	bIsVariable = true;
 	bDesignTime = false;
-	Visiblity = ESlateVisibility::Visible;
+	Visibility = ESlateVisibility::Visible;
 	RenderTransformPivot = FVector2D(0.5f, 0.5f);
 
 	//TODO UMG ToolTipWidget
@@ -106,20 +106,20 @@ bool UWidget::IsVisible() const
 	return false;
 }
 
-TEnumAsByte<ESlateVisibility::Type> UWidget::GetVisibility()
+ESlateVisibility UWidget::GetVisibility() const
 {
 	TSharedPtr<SWidget> SafeWidget = GetCachedWidget();
 	if (SafeWidget.IsValid())
 	{
-		return UWidget::ConvertRuntimeToSerializedVisiblity(SafeWidget->GetVisibility());
+		return UWidget::ConvertRuntimeToSerializedVisibility(SafeWidget->GetVisibility());
 	}
 
-	return Visiblity;
+	return Visibility;
 }
 
-void UWidget::SetVisibility(TEnumAsByte<ESlateVisibility::Type> InVisibility)
+void UWidget::SetVisibility(ESlateVisibility InVisibility)
 {
-	Visiblity = InVisibility;
+	Visibility = InVisibility;
 
 	TSharedPtr<SWidget> SafeWidget = GetCachedWidget();
 	if (SafeWidget.IsValid())
@@ -458,7 +458,7 @@ void UWidget::SynchronizeProperties()
 #endif
 	{
 		SafeWidget->SetEnabled(OPTIONAL_BINDING(bool, bIsEnabled));
-		SafeWidget->SetVisibility(OPTIONAL_BINDING_CONVERT(ESlateVisibility::Type, Visiblity, EVisibility, ConvertVisibility));
+		SafeWidget->SetVisibility(OPTIONAL_BINDING_CONVERT(ESlateVisibility, Visibility, EVisibility, ConvertVisibility));
 	}
 
 	UpdateRenderTransform();
@@ -496,7 +496,17 @@ void UWidget::SetIsDesignTime(bool bInDesignTime)
 	bDesignTime = bInDesignTime;
 }
 
-EVisibility UWidget::ConvertSerializedVisibilityToRuntime(ESlateVisibility::Type Input)
+void UWidget::PostLoad()
+{
+	Super::PostLoad();
+
+	if ( GetLinkerUE4Version() < VER_UE4_RENAME_WIDGET_VISIBILITY )
+	{
+		Visibility = Visiblity_DEPRECATED;
+	}
+}
+
+EVisibility UWidget::ConvertSerializedVisibilityToRuntime(ESlateVisibility Input)
 {
 	switch ( Input )
 	{
@@ -511,12 +521,12 @@ EVisibility UWidget::ConvertSerializedVisibilityToRuntime(ESlateVisibility::Type
 	case ESlateVisibility::SelfHitTestInvisible:
 		return EVisibility::SelfHitTestInvisible;
 	default:
-		//check(false);
+		check(false);
 		return EVisibility::Visible;
 	}
 }
 
-ESlateVisibility::Type UWidget::ConvertRuntimeToSerializedVisiblity(const EVisibility& Input)
+ESlateVisibility UWidget::ConvertRuntimeToSerializedVisibility(const EVisibility& Input)
 {
 	if ( Input == EVisibility::Visible )
 	{
@@ -540,7 +550,7 @@ ESlateVisibility::Type UWidget::ConvertRuntimeToSerializedVisiblity(const EVisib
 	}
 	else
 	{
-		//check(false);
+		check(false);
 		return ESlateVisibility::Visible;
 	}
 }
diff --git a/Engine/Source/Runtime/UMG/Private/Components/WidgetSwitcher.cpp b/Engine/Source/Runtime/UMG/Private/Components/WidgetSwitcher.cpp
--- a/Engine/Source/Runtime/UMG/Private/Components/WidgetSwitcher.cpp
+++ b/Engine/Source/Runtime/UMG/Private/Components/WidgetSwitcher.cpp
@@ -14,7 +14,7 @@ UWidgetSwitcher::UWidgetSwitcher(const FObjectInitializer& ObjectInitializer)
 	bIsVariable = true;
 
 	SWidgetSwitcher::FArguments Defaults;
-	Visiblity = UWidget::ConvertRuntimeToSerializedVisiblity(Defaults._Visibility.Get());
+	Visibility = UWidget::ConvertRuntimeToSerializedVisibility(Defaults._Visibility.Get());
 }
 
 void UWidgetSwitcher::ReleaseSlateResources(bool bReleaseChildren)
diff --git a/Engine/Source/Runtime/UMG/Private/Components/WrapBox.cpp b/Engine/Source/Runtime/UMG/Private/Components/WrapBox.cpp
--- a/Engine/Source/Runtime/UMG/Private/Components/WrapBox.cpp
+++ b/Engine/Source/Runtime/UMG/Private/Components/WrapBox.cpp
@@ -13,7 +13,7 @@ UWrapBox::UWrapBox(const FObjectInitializer& ObjectInitializer)
 	bIsVariable = false;
 
 	SWrapBox::FArguments Defaults;
-	Visiblity = UWidget::ConvertRuntimeToSerializedVisiblity( Defaults._Visibility.Get() );
+	Visibility = UWidget::ConvertRuntimeToSerializedVisibility(Defaults._Visibility.Get());
 }
 
 void UWrapBox::ReleaseSlateResources(bool bReleaseChildren)
diff --git a/Engine/Source/Runtime/UMG/Private/UserWidget.cpp b/Engine/Source/Runtime/UMG/Private/UserWidget.cpp
--- a/Engine/Source/Runtime/UMG/Private/UserWidget.cpp
+++ b/Engine/Source/Runtime/UMG/Private/UserWidget.cpp
@@ -18,7 +18,7 @@ UUserWidget::UUserWidget(const FObjectInitializer& ObjectInitializer)
 	: Super(ObjectInitializer)
 {
 	ViewportAnchors = FAnchors(0, 0, 1, 1);
-	Visiblity = ESlateVisibility::SelfHitTestInvisible;
+	Visibility = ESlateVisibility::SelfHitTestInvisible;
 
 	bInitialized = false;
 	bSupportsKeyboardFocus = true;
@@ -741,6 +741,11 @@ void UUserWidget::PreSave()
 	}
 }
 
+void UUserWidget::PostLoad()
+{
+	Super::PostLoad();
+}
+
 #if WITH_EDITOR
 
 const FSlateBrush* UUserWidget::GetEditorIcon()
diff --git a/Engine/Source/Runtime/UMG/Private/WidgetBlueprintGeneratedClass.cpp b/Engine/Source/Runtime/UMG/Private/WidgetBlueprintGeneratedClass.cpp
--- a/Engine/Source/Runtime/UMG/Private/WidgetBlueprintGeneratedClass.cpp
+++ b/Engine/Source/Runtime/UMG/Private/WidgetBlueprintGeneratedClass.cpp
@@ -106,4 +106,23 @@ void UWidgetBlueprintGeneratedClass::InitializeWidget(UUserWidget* UserWidget) c
 	}
 }
 
+void UWidgetBlueprintGeneratedClass::PostLoad()
+{
+	Super::PostLoad();
+
+	if ( GetLinkerUE4Version() < VER_UE4_RENAME_WIDGET_VISIBILITY )
+	{
+		static const FName Visiblity(TEXT("Visiblity"));
+		static const FName Visibility(TEXT("Visibility"));
+
+		for ( FDelegateRuntimeBinding& Binding : Bindings )
+		{
+			if ( Binding.PropertyName == Visiblity )
+			{
+				Binding.PropertyName = Visibility;
+			}
+		}
+	}
+}
+
 #undef LOCTEXT_NAMESPACE
diff --git a/Engine/Source/Runtime/UMG/Private/WidgetBlueprintLibrary.cpp b/Engine/Source/Runtime/UMG/Private/WidgetBlueprintLibrary.cpp
--- a/Engine/Source/Runtime/UMG/Private/WidgetBlueprintLibrary.cpp
+++ b/Engine/Source/Runtime/UMG/Private/WidgetBlueprintLibrary.cpp
@@ -17,7 +17,12 @@ UWidgetBlueprintLibrary::UWidgetBlueprintLibrary(const FObjectInitializer& Objec
 
 UUserWidget* UWidgetBlueprintLibrary::Create(UObject* WorldContextObject, TSubclassOf<UUserWidget> WidgetType, APlayerController* OwningPlayer)
 {
-	if ( OwningPlayer == NULL )
+	if ( WidgetType == nullptr || WidgetType->HasAnyClassFlags(CLASS_Abstract) )
+	{
+		return nullptr;
+	}
+
+	if ( OwningPlayer == nullptr )
 	{
 		UWorld* World = GEngine->GetWorldFromContextObject(WorldContextObject);
 		return CreateWidget<UUserWidget>(World, WidgetType);
diff --git a/Engine/Source/Runtime/UMG/Public/Blueprint/UserWidget.h b/Engine/Source/Runtime/UMG/Public/Blueprint/UserWidget.h
--- a/Engine/Source/Runtime/UMG/Public/Blueprint/UserWidget.h
+++ b/Engine/Source/Runtime/UMG/Public/Blueprint/UserWidget.h
@@ -105,7 +105,6 @@ class UUMGSequencePlayer;
 //TODO UMG If you want to host a widget that's full screen there may need to be a SWindow equivalent that you spawn it into.
 
 DECLARE_DYNAMIC_MULTICAST_DELEGATE(FOnConstructEvent);
-DECLARE_DYNAMIC_MULTICAST_DELEGATE_OneParam(FOnVisibilityChangedEvent, ESlateVisibility::Type, Visibility);
 
 /**
  * The user widget is extensible by users through the WidgetBlueprint.
@@ -583,6 +582,7 @@ public:
 
 	// Begin UObject interface
 	virtual void PreSave() override;
+	virtual void PostLoad() override;
 	// End UObject interface
 
 #if WITH_EDITOR
diff --git a/Engine/Source/Runtime/UMG/Public/Blueprint/WidgetBlueprintGeneratedClass.h b/Engine/Source/Runtime/UMG/Public/Blueprint/WidgetBlueprintGeneratedClass.h
--- a/Engine/Source/Runtime/UMG/Public/Blueprint/WidgetBlueprintGeneratedClass.h
+++ b/Engine/Source/Runtime/UMG/Public/Blueprint/WidgetBlueprintGeneratedClass.h
@@ -24,23 +24,21 @@ struct FDelegateRuntimeBinding
 {
 	GENERATED_USTRUCT_BODY()
 
+	/** The widget that will be bound to the live data. */
 	UPROPERTY()
 	FString ObjectName;
 
+	/** The property on the widget that will have a binding placed on it. */
 	UPROPERTY()
 	FName PropertyName;
 
+	/** The function or property we're binding to on the source object. */
 	UPROPERTY()
 	FName FunctionName;
 
+	/** The kind of binding we're performing, are we binding to a property or a function. */
 	UPROPERTY()
 	TEnumAsByte<EBindingKind::Type> Kind;
-	
-	/*UFunction* UK2Node_CallFunction::GetTargetFunction() const
-	{
-	UFunction* Function = FunctionReference.ResolveMember<UFunction>(this);
-	return Function;
-	}*/
 };
 
 
@@ -72,6 +70,8 @@ public:
 
 public:
 
+	virtual void PostLoad() override;
+
 	/**
 	 * This is the function that makes UMG work.  Once a user widget is constructed, it will post load
 	 * call into its generated class and ask to be initialized.  The class will perform all the delegate
diff --git a/Engine/Source/Runtime/UMG/Public/Components/CanvasPanel.h b/Engine/Source/Runtime/UMG/Public/Components/CanvasPanel.h
--- a/Engine/Source/Runtime/UMG/Public/Components/CanvasPanel.h
+++ b/Engine/Source/Runtime/UMG/Public/Components/CanvasPanel.h
@@ -23,6 +23,10 @@ class UMG_API UCanvasPanel : public UPanelWidget
 
 public:
 
+	/**  */
+	UFUNCTION(BlueprintCallable, Category="Canvas Panel")
+	UCanvasPanelSlot* AddChildToCanvas(UWidget* Content);
+
 	/** Gets the underlying native canvas widget if it has been constructed */
 	TSharedPtr<class SConstraintCanvas> GetCanvasWidget() const;
 
diff --git a/Engine/Source/Runtime/UMG/Public/Components/ScrollBox.h b/Engine/Source/Runtime/UMG/Public/Components/ScrollBox.h
--- a/Engine/Source/Runtime/UMG/Public/Components/ScrollBox.h
+++ b/Engine/Source/Runtime/UMG/Public/Components/ScrollBox.h
@@ -34,7 +34,7 @@ public:
 
 	/** Visibility */
 	UPROPERTY(EditDefaultsOnly, BlueprintReadOnly, Category="Scroll")
-	TEnumAsByte<ESlateVisibility::Type> ScrollBarVisibility;
+	TEnumAsByte<ESlateVisibility> ScrollBarVisibility;
 
 	/**  */
 	UPROPERTY(EditDefaultsOnly, BlueprintReadOnly, Category="Scroll")
diff --git a/Engine/Source/Runtime/UMG/Public/Components/SlateWrapperTypes.h b/Engine/Source/Runtime/UMG/Public/Components/SlateWrapperTypes.h
--- a/Engine/Source/Runtime/UMG/Public/Components/SlateWrapperTypes.h
+++ b/Engine/Source/Runtime/UMG/Public/Components/SlateWrapperTypes.h
@@ -13,22 +13,19 @@
 
 /** Is an entity visible? */
 UENUM(BlueprintType)
-namespace ESlateVisibility
+enum class ESlateVisibility : uint8
 {
-	enum Type
-	{
-		/** Default widget visibility - visible and can interactive with the cursor */
-		Visible,
-		/** Not visible and takes up no space in the layout; can never be clicked on because it takes up no space. */
-		Collapsed,
-		/** Not visible, but occupies layout space. Not interactive for obvious reasons. */
-		Hidden,
-		/** Visible to the user, but only as art. The cursors hit tests will never see this widget. */
-		HitTestInvisible,
-		/** Same as HitTestInvisible, but doesn't apply to child widgets. */
-		SelfHitTestInvisible
-	};
-}
+	/** Default widget visibility - visible and can interactive with the cursor */
+	Visible,
+	/** Not visible and takes up no space in the layout; can never be clicked on because it takes up no space. */
+	Collapsed,
+	/** Not visible, but occupies layout space. Not interactive for obvious reasons. */
+	Hidden,
+	/** Visible to the user, but only as art. The cursors hit tests will never see this widget. */
+	HitTestInvisible,
+	/** Same as HitTestInvisible, but doesn't apply to child widgets. */
+	SelfHitTestInvisible
+};
 
 /** The sizing options of UWidgets */
 UENUM(BlueprintType)
diff --git a/Engine/Source/Runtime/UMG/Public/Components/Widget.h b/Engine/Source/Runtime/UMG/Public/Components/Widget.h
--- a/Engine/Source/Runtime/UMG/Public/Components/Widget.h
+++ b/Engine/Source/Runtime/UMG/Public/Components/Widget.h
@@ -41,7 +41,7 @@ public:
 	DECLARE_DYNAMIC_DELEGATE_RetVal(FMargin, FGetMargin);
 	DECLARE_DYNAMIC_DELEGATE_RetVal(FSlateColor, FGetSlateColor);
 	DECLARE_DYNAMIC_DELEGATE_RetVal(FLinearColor, FGetLinearColor);
-	DECLARE_DYNAMIC_DELEGATE_RetVal(ESlateVisibility::Type, FGetSlateVisibility);
+	DECLARE_DYNAMIC_DELEGATE_RetVal(ESlateVisibility, FGetSlateVisibility);
 	DECLARE_DYNAMIC_DELEGATE_RetVal(EMouseCursor::Type, FGetMouseCursor);
 	DECLARE_DYNAMIC_DELEGATE_RetVal(USlateBrushAsset*, FGetSlateBrushAsset);
 	DECLARE_DYNAMIC_DELEGATE_RetVal(FSlateBrush, FGetSlateBrush);
@@ -89,12 +89,16 @@ public:
 	FGetText ToolTipTextDelegate;
 
 	/** The visibility of the widget */
-	UPROPERTY(EditDefaultsOnly, BlueprintReadOnly, Category=Behavior)
-	TEnumAsByte<ESlateVisibility::Type> Visiblity;
+	UPROPERTY()
+	TEnumAsByte<ESlateVisibility> Visiblity_DEPRECATED;
+
+	/** The visibility of the widget */
+	UPROPERTY(EditDefaultsOnly, Category=Behavior)
+	ESlateVisibility Visibility;
 
 	/** A bindable delegate for Visibility */
 	UPROPERTY()
-	FGetSlateVisibility VisiblityDelegate;
+	FGetSlateVisibility VisibilityDelegate;
 
 	/** The cursor to show when the mouse is over the widget */
 	UPROPERTY(EditDefaultsOnly, BlueprintReadOnly, Category=Behavior, AdvancedDisplay)
@@ -179,11 +183,11 @@ public:
 
 	/** Gets the current visibility of the widget. */
 	UFUNCTION(BlueprintCallable, Category="Widget")
-	TEnumAsByte<ESlateVisibility::Type> GetVisibility();
+	ESlateVisibility GetVisibility() const;
 
 	/** Sets the visibility of the widget. */
 	UFUNCTION(BlueprintCallable, Category="Widget")
-	void SetVisibility(TEnumAsByte<ESlateVisibility::Type> InVisibility);
+	void SetVisibility(ESlateVisibility InVisibility);
 
 	/** Gets if the button is currently being hovered by the mouse */
 	UFUNCTION(BlueprintCallable, Category="Widget")
@@ -267,6 +271,8 @@ public:
 	 * @return true if this widget is a child of the PossibleParent
 	 */
 	bool IsChildOf(UWidget* PossibleParent);
+
+	virtual void PostLoad() override;
 	
 #if WITH_EDITOR
 	/** Is the label generated or provided by the user? */
@@ -317,8 +323,8 @@ public:
 
 	// Utility methods
 	//@TODO UMG: Should move elsewhere
-	static EVisibility ConvertSerializedVisibilityToRuntime(ESlateVisibility::Type Input);
-	static ESlateVisibility::Type ConvertRuntimeToSerializedVisiblity(const EVisibility& Input);
+	static EVisibility ConvertSerializedVisibilityToRuntime(ESlateVisibility Input);
+	static ESlateVisibility ConvertRuntimeToSerializedVisibility(const EVisibility& Input);
 
 	static FSizeParam ConvertSerializedSizeParamToRuntime(const FSlateChildSize& Input);
 
@@ -337,9 +343,9 @@ protected:
 protected:
 	//TODO UMG Consider moving conversion functions into another class.
 	// Conversion functions
-	EVisibility ConvertVisibility(TAttribute<ESlateVisibility::Type> SerializedType) const
+	EVisibility ConvertVisibility(TAttribute<ESlateVisibility> SerializedType) const
 	{
-		ESlateVisibility::Type SlateVisibility = SerializedType.Get();
+		ESlateVisibility SlateVisibility = SerializedType.Get();
 		return ConvertSerializedVisibilityToRuntime(SlateVisibility);
 	}
 
