UMG - The variables generated by the blueprint compiler for UMG for widgets and animations are now marked as transient.  This should greatly simplify complexity of serializing them, as those fields are all generated at init time, no sense in attempting to serialize them.

diff --git a/Engine/Source/Editor/UMGEditor/Private/WidgetBlueprintCompiler.cpp b/Engine/Source/Editor/UMGEditor/Private/WidgetBlueprintCompiler.cpp
--- a/Engine/Source/Editor/UMGEditor/Private/WidgetBlueprintCompiler.cpp
+++ b/Engine/Source/Editor/UMGEditor/Private/WidgetBlueprintCompiler.cpp
@@ -193,6 +193,7 @@ void FWidgetBlueprintCompiler::CreateClassVariablesFromBlueprint()
 		Widgets.Sort(FWidgetSorter());
 	}
 
+	// Add widget variables
 	for ( UWidget* Widget : Widgets )
 	{
 		// Skip non-variable widgets
@@ -201,15 +202,15 @@ void FWidgetBlueprintCompiler::CreateClassVariablesFromBlueprint()
 			continue;
 		}
 
-		FString VariableName = Widget->GetName();
-
 		FEdGraphPinType WidgetPinType(Schema->PC_Object, TEXT(""), Widget->GetClass(), false, false);
-		UProperty* WidgetProperty = CreateVariable(FName(*VariableName), WidgetPinType);
-		if ( WidgetProperty != NULL )
+		UProperty* WidgetProperty = CreateVariable(Widget->GetFName(), WidgetPinType);
+		if ( WidgetProperty != nullptr )
 		{
 			WidgetProperty->SetMetaData(TEXT("Category"), *Blueprint->GetName());
-			//WidgetProperty->SetMetaData(TEXT("Category"), TEXT("Widget"));
+			
 			WidgetProperty->SetPropertyFlags(CPF_BlueprintVisible);
+			WidgetProperty->SetPropertyFlags(CPF_Transient);
+			WidgetProperty->SetPropertyFlags(CPF_RepSkip);
 
 			WidgetToMemberVariableMap.Add(Widget, WidgetProperty);
 		}
@@ -218,16 +219,16 @@ void FWidgetBlueprintCompiler::CreateClassVariablesFromBlueprint()
 	// Add movie scenes variables here
 	for(UWidgetAnimation* Animation : Blueprint->Animations)
 	{
-		FString VariableName = Animation->GetName();
-
 		FEdGraphPinType WidgetPinType(Schema->PC_Object, TEXT(""), Animation->GetClass(), false, true);
-		UProperty* WidgetProperty = CreateVariable(FName(*VariableName), WidgetPinType);
+		UProperty* AnimationProperty = CreateVariable(Animation->GetFName(), WidgetPinType);
 
-		if(WidgetProperty != NULL)
+		if ( AnimationProperty != nullptr )
 		{
-			WidgetProperty->SetMetaData(TEXT("Category"), TEXT("Animations") );
+			AnimationProperty->SetMetaData(TEXT("Category"), TEXT("Animations"));
 
-			WidgetProperty->SetPropertyFlags(CPF_BlueprintVisible);
+			AnimationProperty->SetPropertyFlags(CPF_BlueprintVisible);
+			AnimationProperty->SetPropertyFlags(CPF_Transient);
+			AnimationProperty->SetPropertyFlags(CPF_RepSkip);
 		}
 	}
 }
