UMG - Adding readonly properties to the set of bindable ones, also properly checking parent classes now.

diff --git a/Engine/Source/Editor/UMGEditor/Private/Details/SPropertyBinding.cpp b/Engine/Source/Editor/UMGEditor/Private/Details/SPropertyBinding.cpp
--- a/Engine/Source/Editor/UMGEditor/Private/Details/SPropertyBinding.cpp
+++ b/Engine/Source/Editor/UMGEditor/Private/Details/SPropertyBinding.cpp
@@ -164,15 +164,16 @@ void SPropertyBinding::RefreshBlueprintMemberCache(const UFunction* DelegateSign
 		}
 
 		// Grab functions implemented by the blueprint
-		for ( TFieldIterator<UProperty> PropIt(SkeletonClass, EFieldIteratorFlags::ExcludeSuper); PropIt; ++PropIt )
+		for ( TFieldIterator<UProperty> PropIt(Class, EFieldIteratorFlags::ExcludeSuper); PropIt; ++PropIt )
 		{
 			UProperty* Prop = *PropIt;
 
+			// Add matching properties.
 			if ( UProperty* ReturnProperty = DelegateSignature->GetReturnProperty() )
 			{
 				if ( ReturnProperty->SameType(Prop) )
 				{
-					if ( Prop->HasAnyPropertyFlags(UP::BlueprintReadWrite) )
+					if ( Prop->HasAnyPropertyFlags(UP::BlueprintReadWrite | UP::BlueprintReadOnly) )
 					{
 						BlueprintPropertyCache.Add(Prop);
 					}
