UMG - Fixing an intermittent crash in the canvas slot customization that occurs during compile, the object in the property grid can become null for a frame.

diff --git a/Engine/Source/Editor/UMGEditor/Private/Customizations/CanvasSlotCustomization.cpp b/Engine/Source/Editor/UMGEditor/Private/Customizations/CanvasSlotCustomization.cpp
--- a/Engine/Source/Editor/UMGEditor/Private/Customizations/CanvasSlotCustomization.cpp
+++ b/Engine/Source/Editor/UMGEditor/Private/Customizations/CanvasSlotCustomization.cpp
@@ -267,21 +267,23 @@ FText FCanvasSlotCustomization::GetOffsetLabel(TSharedPtr<IPropertyHandle> Prope
 	TArray<UObject*> Objects;
 	PropertyHandle->GetOuterObjects(Objects);
 	
-	if ( Objects.Num() > 1 )
+	if ( Objects.Num() == 1 && Objects[0] != nullptr )
 	{
-		return StretchingLabel;
-	}
+		TArray<void*> RawData;
+		PropertyHandle->AccessRawData(RawData);
 
-	TArray<void*> RawData;
-	PropertyHandle->AccessRawData(RawData);
+		FAnchorData* AnchorData = reinterpret_cast<FAnchorData*>( RawData[0] );
 
-	FAnchorData* AnchorData = reinterpret_cast<FAnchorData*>( RawData[0] );
+		const bool bStretching =
+			( Orientation == Orient_Horizontal && AnchorData->Anchors.IsStretchedHorizontal() ) ||
+			( Orientation == Orient_Vertical && AnchorData->Anchors.IsStretchedVertical() );
 
-	const bool bStretching =
-		( Orientation == Orient_Horizontal && AnchorData->Anchors.IsStretchedHorizontal() ) ||
-		( Orientation == Orient_Vertical && AnchorData->Anchors.IsStretchedVertical() );
-
-	return bStretching ? StretchingLabel : NonStretchingLabel;
+		return bStretching ? StretchingLabel : NonStretchingLabel;
+	}
+	else
+	{
+		return StretchingLabel;
+	}
 }
 
 void FCanvasSlotCustomization::CustomizeAnchors(TSharedPtr<IPropertyHandle> PropertyHandle, IDetailChildrenBuilder& ChildBuilder, IPropertyTypeCustomizationUtils& CustomizationUtils)
