Fixed umg widget recompiles not updating 3d components in the world when not in realtime rendering mode

diff --git a/Engine/Source/Editor/UMGEditor/Private/UMGEditorModule.cpp b/Engine/Source/Editor/UMGEditor/Private/UMGEditorModule.cpp
--- a/Engine/Source/Editor/UMGEditor/Private/UMGEditorModule.cpp
+++ b/Engine/Source/Editor/UMGEditor/Private/UMGEditorModule.cpp
@@ -80,11 +80,15 @@ public:
 	{
 		if ( UWidgetBlueprint* WidgetBlueprint = Cast<UWidgetBlueprint>(Blueprint) )
 		{
-			TComponentReregisterContext<UWidgetComponent> ComponentReregisterContext;
+			{
+				TComponentReregisterContext<UWidgetComponent> ComponentReregisterContext;
+
+				FWidgetBlueprintCompiler Compiler(WidgetBlueprint, Results, CompileOptions, ObjLoaded);
+				Compiler.Compile();
+				check(Compiler.NewClass);
+			}
 
-			FWidgetBlueprintCompiler Compiler(WidgetBlueprint, Results, CompileOptions, ObjLoaded);
-			Compiler.Compile();
-			check(Compiler.NewClass);
+			GEditor->RedrawAllViewports(true);
 
 			return FReply::Handled();
 		}
