#UE4

- Fixing crash with "-game" when running the Editor

diff --git a/Engine/Source/Editor/UMGEditor/Private/UMGEditorModule.cpp b/Engine/Source/Editor/UMGEditor/Private/UMGEditorModule.cpp
--- a/Engine/Source/Editor/UMGEditor/Private/UMGEditorModule.cpp
+++ b/Engine/Source/Editor/UMGEditor/Private/UMGEditorModule.cpp
@@ -88,7 +88,10 @@ public:
 				check(Compiler.NewClass);
 			}
 
-			GEditor->RedrawAllViewports(true);
+			if (GIsEditor)
+			{
+				GEditor->RedrawAllViewports(true);
+			}
 
 			return FReply::Handled();
 		}
