Cleaned up some duplicated code in the navigation customization.

diff --git a/Engine/Source/Editor/UMGEditor/Private/Customizations/WidgetNavigationCustomization.cpp b/Engine/Source/Editor/UMGEditor/Private/Customizations/WidgetNavigationCustomization.cpp
--- a/Engine/Source/Editor/UMGEditor/Private/Customizations/WidgetNavigationCustomization.cpp
+++ b/Engine/Source/Editor/UMGEditor/Private/Customizations/WidgetNavigationCustomization.cpp
@@ -51,99 +51,21 @@ void FWidgetNavigationCustomization::CustomizeHeader(TSharedRef<IPropertyHandle>
 		.AutoHeight()
 		.Padding(2, 4, 2, 4)
 		[
-			SNew(SHorizontalBox)
-
-			+ SHorizontalBox::Slot()
-			.AutoWidth()
-			.Padding(0)
-			[
-				SNew(STextBlock)
-				.Text(LOCTEXT("RightNavigation", "Right"))
-			]
-
-			+ SHorizontalBox::Slot()
-			.AutoWidth()
-			.Padding(0)
-			[
-				SNew(SComboButton)
-				.HAlign(HAlign_Center)
-				.ButtonContent()
-				[
-					SNew(STextBlock)
-					.Text(this, &FWidgetNavigationCustomization::HandleNavigationText, PropertyHandlePtr, EUINavigation::Right)
-				]
-				.ContentPadding(FMargin(4.0f, 2.0f))
-				.MenuContent()
-				[
-					MakeNavMenu(PropertyHandlePtr, EUINavigation::Right)
-				]
-			]
+			MakeNavRow(PropertyHandlePtr, EUINavigation::Right, LOCTEXT("RightNavigation", "Right"))
 		]
 
 		+ SVerticalBox::Slot()
 		.AutoHeight()
 		.Padding(2, 4, 2, 4)
 		[
-			SNew(SHorizontalBox)
-
-			+ SHorizontalBox::Slot()
-			.AutoWidth()
-			.Padding(0)
-			[
-				SNew(STextBlock)
-				.Text(LOCTEXT("UpNavigation", "Up"))
-			]
-
-			+ SHorizontalBox::Slot()
-			.AutoWidth()
-			.Padding(0)
-			[
-				SNew(SComboButton)
-				.HAlign(HAlign_Center)
-				.ButtonContent()
-				[
-					SNew(STextBlock)
-					.Text(this, &FWidgetNavigationCustomization::HandleNavigationText, PropertyHandlePtr, EUINavigation::Up)
-				]
-				.ContentPadding(FMargin(4.0f, 2.0f))
-				.MenuContent()
-				[
-					MakeNavMenu(PropertyHandlePtr, EUINavigation::Up)
-				]
-			]
+			MakeNavRow(PropertyHandlePtr, EUINavigation::Up, LOCTEXT("UpNavigation", "Up"))
 		]
 
 		+ SVerticalBox::Slot()
 		.AutoHeight()
 		.Padding(2, 4, 2, 4)
 		[
-			SNew(SHorizontalBox)
-
-			+ SHorizontalBox::Slot()
-			.AutoWidth()
-			.Padding(0)
-			[
-				SNew(STextBlock)
-				.Text(LOCTEXT("DownNavigation", "Down"))
-			]
-
-			+ SHorizontalBox::Slot()
-			.AutoWidth()
-			.Padding(0)
-			[
-				SNew(SComboButton)
-				.HAlign(HAlign_Center)
-				.ButtonContent()
-				[
-					SNew(STextBlock)
-					.Text(this, &FWidgetNavigationCustomization::HandleNavigationText, PropertyHandlePtr, EUINavigation::Down)
-				]
-				.ContentPadding(FMargin(4.0f, 2.0f))
-				.MenuContent()
-				[
-					MakeNavMenu(PropertyHandlePtr, EUINavigation::Down)
-				]
-			]
+			MakeNavRow(PropertyHandlePtr, EUINavigation::Down, LOCTEXT("DownNavigation", "Down"))
 		]
 	];
 }
@@ -252,7 +174,7 @@ TSharedRef<class SWidget> FWidgetNavigationCustomization::MakeNavRow(TWeakPtr<IP
 			.Text(NavName)
 		]
 
-	+ SHorizontalBox::Slot()
+		+ SHorizontalBox::Slot()
 		.AutoWidth()
 		.Padding(0)
 		[
@@ -264,10 +186,10 @@ TSharedRef<class SWidget> FWidgetNavigationCustomization::MakeNavRow(TWeakPtr<IP
 				.Text(this, &FWidgetNavigationCustomization::HandleNavigationText, PropertyHandle, Nav)
 			]
 			.ContentPadding(FMargin(4.0f, 2.0f))
-				.MenuContent()
-				[
-					MakeNavMenu(PropertyHandle, Nav)
-				]
+			.MenuContent()
+			[
+				MakeNavMenu(PropertyHandle, Nav)
+			]
 		];
 
 	return Row;
