Fix for renaming in UMG Hierarchy to used name causing a assert and crash.

diff --git a/Engine/Source/Editor/UMGEditor/Private/Hierarchy/SHierarchyViewItem.cpp b/Engine/Source/Editor/UMGEditor/Private/Hierarchy/SHierarchyViewItem.cpp
--- a/Engine/Source/Editor/UMGEditor/Private/Hierarchy/SHierarchyViewItem.cpp
+++ b/Engine/Source/Editor/UMGEditor/Private/Hierarchy/SHierarchyViewItem.cpp
@@ -698,7 +698,7 @@ bool FHierarchyWidget::OnVerifyNameTextChanged(const FText& InText, FText& OutEr
 	bool bIsSameWidget = false;
 	if ( ExistingTemplate != NULL )
 	{
-		if ( Item.GetTemplate() == ExistingTemplate )
+		if ( Item.GetTemplate() != ExistingTemplate )
 		{
 			OutErrorMessage = LOCTEXT("ExistingWidgetName", "Existing Widget Name");
 			return false;
