Fixes for odd / incorrect errors when renaming Widgets.

diff --git a/Engine/Source/Editor/UMGEditor/Private/Details/SWidgetDetailsView.cpp b/Engine/Source/Editor/UMGEditor/Private/Details/SWidgetDetailsView.cpp
--- a/Engine/Source/Editor/UMGEditor/Private/Details/SWidgetDetailsView.cpp
+++ b/Engine/Source/Editor/UMGEditor/Private/Details/SWidgetDetailsView.cpp
@@ -360,10 +360,16 @@ bool SWidgetDetailsView::HandleVerifyNameTextChanged(const FText& InText, FText&
 {
 	if ( SelectedObjects.Num() == 1 )
 	{
-		UWidget* PreviewWidget = Cast<UWidget>(SelectedObjects[0].Get());
-
 		FString NewName = InText.ToString();
 
+		if (NewName.IsEmpty())
+		{
+			OutErrorMessage = LOCTEXT("EmptyWidgetName", "Empty Widget Name");
+			return false;
+		}
+
+		UWidget* PreviewWidget = Cast<UWidget>(SelectedObjects[0].Get());
+
 		UWidgetBlueprint* Blueprint = BlueprintEditor.Pin()->GetWidgetBlueprintObj();
 		UWidget* TemplateWidget = Blueprint->WidgetTree->FindWidget( FName(*NewName) );
 
@@ -416,6 +422,11 @@ void SWidgetDetailsView::HandleNameTextCommitted(const FText& Text, ETextCommit:
 			}
 		}
 		IsReentrant = false;
+
+		if (CommitType == ETextCommit::OnUserMovedFocus || CommitType == ETextCommit::OnCleared)
+		{
+			NameTextBox->SetError(FText::GetEmpty());
+		}
 	}
 }
 
diff --git a/Engine/Source/Editor/UMGEditor/Private/Hierarchy/SHierarchyViewItem.cpp b/Engine/Source/Editor/UMGEditor/Private/Hierarchy/SHierarchyViewItem.cpp
--- a/Engine/Source/Editor/UMGEditor/Private/Hierarchy/SHierarchyViewItem.cpp
+++ b/Engine/Source/Editor/UMGEditor/Private/Hierarchy/SHierarchyViewItem.cpp
@@ -692,6 +692,12 @@ bool FHierarchyWidget::OnVerifyNameTextChanged(const FText& InText, FText& OutEr
 {
 	FString NewName = InText.ToString();
 
+	if (NewName.IsEmpty())
+	{
+		OutErrorMessage = LOCTEXT("EmptyWidgetName", "Empty Widget Name");
+		return false;
+	}
+
 	UWidgetBlueprint* Blueprint = BlueprintEditor.Pin()->GetWidgetBlueprintObj();
 	UWidget* ExistingTemplate = Blueprint->WidgetTree->FindWidget( FName(*NewName) );
 
@@ -870,7 +876,7 @@ bool SHierarchyViewItem::OnVerifyNameTextChanged(const FText& InText, FText& Out
 
 void SHierarchyViewItem::OnNameTextCommited(const FText& InText, ETextCommit::Type CommitInfo)
 {
-	return Model->OnNameTextCommited(InText, CommitInfo);
+	Model->OnNameTextCommited(InText, CommitInfo);
 }
 
 FSlateFontInfo SHierarchyViewItem::GetItemFont() const
