In UMG you can no longer rename a node the the "Nice name" this would almost certainly not be desired.

diff --git a/Engine/Source/Editor/UMGEditor/Private/Hierarchy/SHierarchyViewItem.cpp b/Engine/Source/Editor/UMGEditor/Private/Hierarchy/SHierarchyViewItem.cpp
--- a/Engine/Source/Editor/UMGEditor/Private/Hierarchy/SHierarchyViewItem.cpp
+++ b/Engine/Source/Editor/UMGEditor/Private/Hierarchy/SHierarchyViewItem.cpp
@@ -876,6 +876,13 @@ bool SHierarchyViewItem::OnVerifyNameTextChanged(const FText& InText, FText& Out
 
 void SHierarchyViewItem::OnNameTextCommited(const FText& InText, ETextCommit::Type CommitInfo)
 {
+	// The model can return nice names "Border_53" becomes [Border] in some cases
+	// This check makes sure we don't rename the object internally to that nice name.
+	// Most common case would be the user enters edit mode by accident then just moves focus away.
+	if (Model->GetText().EqualToCaseIgnored(InText))
+	{
+		return;
+	}
 	Model->OnNameTextCommited(InText, CommitInfo);
 }
 
