UMG - Preventing crash on recompile.  Doesn't refresh all viewports unless the editor is available.

diff --git a/Engine/Source/Editor/UMGEditor/Private/UMGEditorModule.cpp b/Engine/Source/Editor/UMGEditor/Private/UMGEditorModule.cpp
--- a/Engine/Source/Editor/UMGEditor/Private/UMGEditorModule.cpp
+++ b/Engine/Source/Editor/UMGEditor/Private/UMGEditorModule.cpp
@@ -88,7 +88,7 @@ public:
 				check(Compiler.NewClass);
 			}
 
-			if (GIsEditor)
+			if ( GIsEditor && GEditor )
 			{
 				GEditor->RedrawAllViewports(true);
 			}
