UMG - Feature: Upgrading the binding system.  Now supports binding to sub properties of objects on the widget.  Bindings are based on a dynamically evaluated set of fields, does not use the VM to allow bindings to be as fast as possible.  Null entries are ignored and in those cases the default values are returned to avoid crashes.  Also allows binding to functions on sub objects, but the functions must be the terminal bit of the binding.  So GetMyObject().GetMyValue() is not permitted, however MyObject.GetMyValue() is.  Function bindings must be pure or const.  Trivial conversions are also handled, for the purposes of binding: LinearColor / SlateColor convert, String / Text convert, and Texture2D can be automatically converted to a Brush.

diff --git a/Engine/Source/Editor/UMGEditor/Private/Customizations/DetailCustomizations.h b/Engine/Source/Editor/UMGEditor/Private/Customizations/DetailCustomizations.h
--- a/Engine/Source/Editor/UMGEditor/Private/Customizations/DetailCustomizations.h
+++ b/Engine/Source/Editor/UMGEditor/Private/Customizations/DetailCustomizations.h
@@ -30,50 +30,14 @@ public:
 	virtual void CustomizeDetails( IDetailLayoutBuilder& DetailLayout ) override;
 
 private:
-	struct FunctionInfo
-	{
-		FText DisplayName;
-		FString Tooltip;
-
-		FName FuncName;
-		UEdGraph* EdGraph;
-	};
-
 	void PerformBindingCustomization(IDetailLayoutBuilder& DetailLayout);
 
-	void RefreshBlueprintMemberCache(const UFunction* DelegateSignature);
-
-	TSharedRef<SWidget> OnGenerateDelegateMenu(UWidget* Widget, TSharedRef<IPropertyHandle> PropertyHandle, UFunction* DelegateSignature, bool bIsPure);
-
-	const FSlateBrush* GetCurrentBindingImage(TSharedRef<IPropertyHandle> PropertyHandle) const;
-	FText GetCurrentBindingText(TSharedRef<IPropertyHandle> PropertyHandle) const;
-
-	bool CanRemoveBinding(TSharedRef<IPropertyHandle> PropertyHandle);
-	void HandleRemoveBinding(TSharedRef<IPropertyHandle> PropertyHandle);
-
-	void HandleAddFunctionBinding(TSharedRef<IPropertyHandle> PropertyHandle, TSharedPtr<FunctionInfo> SelectedFunction);
-	void HandleAddPropertyBinding(TSharedRef<IPropertyHandle> PropertyHandle, UProperty* SelectedProperty);
-
-	void HandleCreateAndAddBinding(UWidget* Widget, TSharedRef<IPropertyHandle> PropertyHandle, UFunction* DelegateSignature, bool bIsPure);
-	void GotoFunction(UEdGraph* FunctionGraph);
-
-	EVisibility GetGotoBindingVisibility(TSharedRef<IPropertyHandle> PropertyHandle) const;
-
-	FReply HandleGotoBindingClicked(TSharedRef<IPropertyHandle> PropertyHandle);
-
-	void CreateDelegateCustomization( IDetailLayoutBuilder& DetailLayout, UDelegateProperty* Property, UWidget* Widget );
-
 	void CreateEventCustomization( IDetailLayoutBuilder& DetailLayout, UDelegateProperty* Property, UWidget* Widget );
 
 	void CreateMulticastEventCustomization(IDetailLayoutBuilder& DetailLayout, FName ThisComponentName, UClass* PropertyClass, UMulticastDelegateProperty* Property);
-
-	FReply AddOrViewEventBinding(TSharedPtr<FEdGraphSchemaAction> Action);
 private:
 
 	TWeakPtr<FWidgetBlueprintEditor> Editor;
 
 	UWidgetBlueprint* Blueprint;
-
-	TArray< TSharedPtr<FunctionInfo> > BlueprintFunctionCache;
-	TArray< UProperty* > BlueprintPropertyCache;
 };
diff --git a/Engine/Source/Editor/UMGEditor/Private/Details/SPropertyBinding.cpp b/Engine/Source/Editor/UMGEditor/Private/Details/SPropertyBinding.cpp
--- a/Engine/Source/Editor/UMGEditor/Private/Details/SPropertyBinding.cpp
+++ b/Engine/Source/Editor/UMGEditor/Private/Details/SPropertyBinding.cpp
@@ -8,6 +8,7 @@
 #include "WidgetGraphSchema.h"
 #include "BlueprintEditorUtils.h"
 #include "ScopedTransaction.h"
+#include "PropertyBinding.h"
 
 
 #define LOCTEXT_NAMESPACE "UMG"
@@ -21,6 +22,9 @@ void SPropertyBinding::Construct(const FArguments& InArgs, TSharedRef<FWidgetBlu
 	Editor = InEditor;
 	Blueprint = InEditor->GetWidgetBlueprintObj();
 
+	GeneratePureBindings = InArgs._GeneratePureBindings;
+	BindableSignature = DelegateProperty->SignatureFunction;
+
 	TArray<UObject*> Objects;
 	Property->GetOuterObjects(Objects);
 
@@ -34,7 +38,7 @@ void SPropertyBinding::Construct(const FArguments& InArgs, TSharedRef<FWidgetBlu
 		.AutoWidth()
 		[
 			SNew(SComboButton)
-			.OnGetMenuContent(this, &SPropertyBinding::OnGenerateDelegateMenu, Widget, Property, DelegateProperty->SignatureFunction, InArgs._GeneratePureBindings)
+			.OnGetMenuContent(this, &SPropertyBinding::OnGenerateDelegateMenu, Widget, Property)
 			.ContentPadding(1)
 			.ButtonContent()
 			[
@@ -76,180 +80,178 @@ void SPropertyBinding::Construct(const FArguments& InArgs, TSharedRef<FWidgetBlu
 	];
 }
 
-void SPropertyBinding::RefreshBlueprintMemberCache(const UFunction* DelegateSignature, bool bIsPure)
+static bool IsClassBlackListed(UClass* OwnerClass)
+{
+	if ( OwnerClass == UUserWidget::StaticClass() ||
+		OwnerClass == AActor::StaticClass() ||
+		OwnerClass == APawn::StaticClass() ||
+		OwnerClass == UObject::StaticClass() ||
+		OwnerClass == UPrimitiveComponent::StaticClass() ||
+		OwnerClass == USceneComponent::StaticClass() ||
+		OwnerClass == UActorComponent::StaticClass() ||
+		OwnerClass == UWidgetComponent::StaticClass() ||
+		OwnerClass == UStaticMeshComponent::StaticClass() ||
+		OwnerClass == UWidgetAnimation::StaticClass() )
+	{
+		return true;
+	}
+
+	return false;
+}
+
+static bool IsFieldFromBlackListedClass(UField* Field)
+{
+	return IsClassBlackListed(Field->GetOwnerClass());
+}
+
+static bool HasFunctionBinder(UFunction* Function, UFunction* BindableSignature)
+{
+	if ( Function->NumParms == 1 && BindableSignature->NumParms == 1 )
+	{
+		if ( UProperty* FunctionReturn = Function->GetReturnProperty() )
+		{
+			if ( UProperty* DelegateReturn = BindableSignature->GetReturnProperty() )
+			{
+				// Find the binder that can handle the delegate return type.
+				TSubclassOf<UPropertyBinding> Binder = UWidget::FindBinderClassForDestination(DelegateReturn);
+				if ( Binder != nullptr )
+				{
+					// Ensure that the binder also can handle binding from the property we care about.
+					if ( Binder->GetDefaultObject<UPropertyBinding>()->IsSupportedSource(FunctionReturn) )
+					{
+						return true;
+					}
+				}
+			}
+		}
+	}
+
+	return false;
+}
+
+template <typename Predicate>
+void SPropertyBinding::ForEachBindableFunction(UClass* FromClass, Predicate Pred) const
 {
 	const UWidgetGraphSchema* Schema = GetDefault<UWidgetGraphSchema>();
 	const FSlateFontInfo DetailFontInfo = IDetailLayoutBuilder::GetDetailFont();
 
-	BlueprintFunctionCache.Reset();
-	BlueprintPropertyCache.Reset();
+	UBlueprintGeneratedClass* SkeletonClass = Cast<UBlueprintGeneratedClass>(Blueprint->GeneratedClass);
 
-	// Get the current skeleton class, think header for the blueprint.
-	UBlueprintGeneratedClass* SkeletonClass = Cast<UBlueprintGeneratedClass>(Blueprint->SkeletonGeneratedClass);
-
-	// Grab functions implemented by the blueprint
-	for ( int32 i = 0; i < Blueprint->FunctionGraphs.Num(); i++ )
+	// Walk up class hierarchy for native functions and properties
+	for ( TFieldIterator<UFunction> FuncIt(FromClass, EFieldIteratorFlags::IncludeSuper); FuncIt; ++FuncIt )
 	{
-		UEdGraph* Graph = Blueprint->FunctionGraphs[i];
-		check(Graph);
+		UFunction* Function = *FuncIt;
 
-		// Ignore the construction script, we just want user created functions.
-		if ( Graph->GetFName() == Schema->FN_UserConstructionScript )
+		// Stop processing functions after reaching a base class that it doesn't make sense to go beyond.
+		if ( IsFieldFromBlackListedClass(Function) )
 		{
-			continue;
+			break;
 		}
 
-		FName FunctionFName = Graph->GetFName();
+		// Only allow binding pure functions if we're limited to pure function bindings.
+		if ( GeneratePureBindings && !Function->HasAnyFunctionFlags(FUNC_Const | FUNC_BlueprintPure) )
+		{
+			continue;
+		}
 
-		UFunction* Function = SkeletonClass->FindFunctionByName(FunctionFName, EIncludeSuperFlag::IncludeSuper);
-		if ( Function == nullptr )
+		// Only bind to functions that are callable from blueprints
+		if ( !UEdGraphSchema_K2::CanUserKismetCallFunction(Function) )
 		{
 			continue;
 		}
 
 		// We ignore CPF_ReturnParm because all that matters for binding to script functions is that the number of out parameters match.
-		if ( Function->IsSignatureCompatibleWith(DelegateSignature, UFunction::GetDefaultIgnoredSignatureCompatibilityFlags() | CPF_ReturnParm) )
+		if ( Function->IsSignatureCompatibleWith(BindableSignature, UFunction::GetDefaultIgnoredSignatureCompatibilityFlags() | CPF_ReturnParm) ||
+			 HasFunctionBinder(Function, BindableSignature) )
 		{
-			// Only allow binding pure functions if we're limited to pure function bindings.
-			if ( bIsPure && !Function->HasAllFunctionFlags(FUNC_BlueprintPure) )
-			{
-				continue;
-			}
-			
-			FGraphDisplayInfo DisplayInfo;
-			Graph->GetSchema()->GetGraphDisplayInformation(*Graph, DisplayInfo);
+			TSharedPtr<FFunctionInfo> Info = MakeShareable(new FFunctionInfo());
+			Info->DisplayName = FText::FromName(Function->GetFName());
+			Info->Tooltip = Function->GetMetaData("Tooltip");
+			Info->FuncName = Function->GetFName();
+			Info->Function = Function;
 
-			TSharedPtr<FunctionInfo> NewFuncAction = MakeShareable(new FunctionInfo());
-			NewFuncAction->DisplayName = DisplayInfo.PlainName;
-			NewFuncAction->Tooltip = DisplayInfo.Tooltip;
-			NewFuncAction->FuncName = FunctionFName;
-			NewFuncAction->EdGraph = Graph;
-
-			BlueprintFunctionCache.Add(NewFuncAction);
+			Pred(Info);
 		}
 	}
+}
 
-	// Walk up class hierarchy for native functions and properties
-	for ( UClass* Class = SkeletonClass; Class != UUserWidget::StaticClass(); Class = Class->GetSuperClass() )
-	{
-		// Add matching native functions
-		for ( TFieldIterator<UFunction> FuncIt(Class, EFieldIteratorFlags::ExcludeSuper); FuncIt; ++FuncIt )
-		{
-			UFunction* Func = *FuncIt;
-
-			if ( !Func->HasAnyFunctionFlags(UF::BlueprintCallable | UF::BlueprintPure) )
-			{
-				continue;
-			}
-
-			FName FunctionFName = Func->GetFName();
+template <typename Predicate>
+void SPropertyBinding::ForEachBindableProperty(UStruct* InStruct, Predicate Pred) const
+{
+	UBlueprintGeneratedClass* SkeletonClass = Cast<UBlueprintGeneratedClass>(Blueprint->GeneratedClass);
 
-			UFunction* Function = Class->FindFunctionByName(FunctionFName, EIncludeSuperFlag::IncludeSuper);
-			if ( Function == nullptr )
-			{
-				continue;
-			}
+	for ( TFieldIterator<UProperty> PropIt(InStruct, EFieldIteratorFlags::IncludeSuper); PropIt; ++PropIt )
+	{
+		UProperty* Property = *PropIt;
 
-			// We ignore CPF_ReturnParm because all that matters for binding to script functions is that the number of out parameters match.
-			if ( Function->IsSignatureCompatibleWith(DelegateSignature, UFunction::GetDefaultIgnoredSignatureCompatibilityFlags() | CPF_ReturnParm) )
-			{
-				TSharedPtr<FunctionInfo> NewFuncAction = MakeShareable(new FunctionInfo());
-				NewFuncAction->DisplayName = FText::FromName(Func->GetFName());
-				NewFuncAction->Tooltip = Func->GetMetaData("Tooltip");
-				NewFuncAction->FuncName = FunctionFName;
-				NewFuncAction->EdGraph = nullptr;
+		// Stop processing properties after reaching the stoppped base class
+		if ( IsFieldFromBlackListedClass(Property) )
+		{
+			break;
+		}
 
-				BlueprintFunctionCache.Add(NewFuncAction);
-			}
+		if ( !UEdGraphSchema_K2::CanUserKismetAccessVariable(Property, SkeletonClass, UEdGraphSchema_K2::CannotBeDelegate) )
+		{
+			continue;
 		}
 
-		// Grab functions implemented by the blueprint
-		for ( TFieldIterator<UProperty> PropIt(Class, EFieldIteratorFlags::ExcludeSuper); PropIt; ++PropIt )
+		// Also ignore advanced properties
+		if ( Property->HasAnyPropertyFlags(CPF_AdvancedDisplay | CPF_EditorOnly) )
 		{
-			UProperty* Prop = *PropIt;
+			continue;
+		}
 
-			// Add matching properties.
-			if ( UProperty* ReturnProperty = DelegateSignature->GetReturnProperty() )
+		// Add matching properties, ensure they return the same type as the property.
+		if ( UProperty* ReturnProperty = BindableSignature->GetReturnProperty() )
+		{
+			// Find the binder that can handle the delegate return type.
+			TSubclassOf<UPropertyBinding> Binder = UWidget::FindBinderClassForDestination(ReturnProperty);
+			if ( Binder != nullptr )
 			{
-				if ( ReturnProperty->SameType(Prop) )
+				// Ensure that the binder also can handle binding from the property we care about.
+				if ( Binder->GetDefaultObject<UPropertyBinding>()->IsSupportedSource(Property) )
 				{
-					if ( Prop->HasAnyPropertyFlags(UP::BlueprintReadWrite | UP::BlueprintReadOnly) )
-					{
-						BlueprintPropertyCache.Add(Prop);
-					}
+					Pred(Property);
 				}
 			}
 		}
 	}
 }
 
-TSharedRef<SWidget> SPropertyBinding::OnGenerateDelegateMenu(UWidget* Widget, TSharedRef<IPropertyHandle> PropertyHandle, UFunction* DelegateSignature, bool bIsPure)
+TSharedRef<SWidget> SPropertyBinding::OnGenerateDelegateMenu(UWidget* Widget, TSharedRef<IPropertyHandle> PropertyHandle)
 {
-	RefreshBlueprintMemberCache(DelegateSignature, bIsPure);
-
 	const bool bInShouldCloseWindowAfterMenuSelection = true;
 	FMenuBuilder MenuBuilder(bInShouldCloseWindowAfterMenuSelection, nullptr);
 
-	static FName PropertyIcon(TEXT("Kismet.Tabs.Variables"));
-	static FName FunctionIcon(TEXT("GraphEditor.Function_16x"));
-
-	//FLinearColor ColorOut;
-	//const UClass* VarClass = PropertyHandle->GetPropertyClass();
-	//BrushOut = FBlueprintEditor::GetVarIconAndColor(VarClass, PropertyHandle->GetProperty()->GetFName(), ColorOut);
-
-	if ( CanRemoveBinding(PropertyHandle) )
+	MenuBuilder.BeginSection("BindingActions");
 	{
-		MenuBuilder.BeginSection("RemoveBinding");
+		if ( CanRemoveBinding(PropertyHandle) )
 		{
 			MenuBuilder.AddMenuEntry(
 				LOCTEXT("RemoveBinding", "Remove Binding"),
 				LOCTEXT("RemoveBindingToolTip", "Removes the current binding"),
-				FSlateIcon(),
+				FSlateIcon(FEditorStyle::GetStyleSetName(), "Cross"),
 				FUIAction(FExecuteAction::CreateSP(this, &SPropertyBinding::HandleRemoveBinding, PropertyHandle))
 				);
 		}
-		MenuBuilder.EndSection(); //RemoveBinding
-	}
 
-	MenuBuilder.BeginSection("CreateBinding");
-	{
 		MenuBuilder.AddMenuEntry(
 			LOCTEXT("CreateBinding", "Create Binding"),
-			LOCTEXT("CreateBindingToolTip", "Creates a new function for this property to be bound to"),
-			FSlateIcon(),
-			FUIAction(FExecuteAction::CreateSP(this, &SPropertyBinding::HandleCreateAndAddBinding, Widget, PropertyHandle, DelegateSignature, bIsPure))
+			LOCTEXT("CreateBindingToolTip", "Creates a new function on the widget blueprint that will return the binding data for this property."),
+			FSlateIcon(FEditorStyle::GetStyleSetName(), "Plus"),
+			FUIAction(FExecuteAction::CreateSP(this, &SPropertyBinding::HandleCreateAndAddBinding, Widget, PropertyHandle))
 			);
 	}
 	MenuBuilder.EndSection(); //CreateBinding
 
-	MenuBuilder.BeginSection("Functions", LOCTEXT("Functions", "Functions"));
+	// Properties
 	{
-		for ( TSharedPtr<FunctionInfo>& Function : BlueprintFunctionCache )
-		{
-			MenuBuilder.AddMenuEntry(
-				Function->DisplayName,
-				FText::FromString(Function->Tooltip),
-				FSlateIcon(FEditorStyle::GetStyleSetName(), FunctionIcon),
-				FUIAction(FExecuteAction::CreateSP(this, &SPropertyBinding::HandleAddFunctionBinding, PropertyHandle, Function))
-				);
-		}
-	}
-	MenuBuilder.EndSection(); //Functions
+		// Get the current skeleton class, think header for the blueprint.
+		UBlueprintGeneratedClass* SkeletonClass = Cast<UBlueprintGeneratedClass>(Blueprint->GeneratedClass);
 
-	MenuBuilder.BeginSection("Properties", LOCTEXT("Properties", "Properties"));
-	{
-		for ( UProperty* ExistingProperty : BlueprintPropertyCache )
-		{
-			MenuBuilder.AddMenuEntry(
-				ExistingProperty->GetDisplayNameText(),
-				ExistingProperty->GetToolTipText(),
-				FSlateIcon(FEditorStyle::GetStyleSetName(), PropertyIcon),
-				FUIAction(FExecuteAction::CreateSP(this, &SPropertyBinding::HandleAddPropertyBinding, PropertyHandle, ExistingProperty))
-				);
-		}
+		TArray<UField*> BindingChain;
+		FillPropertyMenu(MenuBuilder, PropertyHandle, SkeletonClass, BindingChain);
 	}
-	MenuBuilder.EndSection(); //Functions
-
 
 	FDisplayMetrics DisplayMetrics;
 	FSlateApplication::Get().GetDisplayMetrics(DisplayMetrics);
@@ -264,6 +266,142 @@ TSharedRef<SWidget> SPropertyBinding::OnGenerateDelegateMenu(UWidget* Widget, TS
 		];
 }
 
+void SPropertyBinding::FillPropertyMenu(FMenuBuilder& MenuBuilder, TSharedRef<IPropertyHandle> PropertyHandle, UStruct* OwnerStruct, TArray<UField*> BindingChain)
+{
+	bool bFoundEntry = false;
+
+	//---------------------------------------
+	// Function Bindings
+
+	if ( UClass* OwnerClass = Cast<UClass>(OwnerStruct) )
+	{
+		static FName FunctionIcon(TEXT("GraphEditor.Function_16x"));
+
+		MenuBuilder.BeginSection("Functions", LOCTEXT("Functions", "Functions"));
+		{
+			ForEachBindableFunction(OwnerClass, [&] (TSharedPtr<FFunctionInfo> Info) {
+				TArray<UField*> NewBindingChain(BindingChain);
+				NewBindingChain.Add(Info->Function);
+
+				bFoundEntry = true;
+
+				MenuBuilder.AddMenuEntry(
+					Info->DisplayName,
+					FText::FromString(Info->Tooltip),
+					FSlateIcon(FEditorStyle::GetStyleSetName(), FunctionIcon),
+					FUIAction(FExecuteAction::CreateSP(this, &SPropertyBinding::HandleAddFunctionBinding, PropertyHandle, Info, NewBindingChain))
+					);
+			});
+		}
+		MenuBuilder.EndSection(); //Functions
+	}
+
+	//---------------------------------------
+	// Property Bindings
+
+	// Get the current skeleton class, think header for the blueprint.
+	UBlueprintGeneratedClass* SkeletonClass = Cast<UBlueprintGeneratedClass>(Blueprint->GeneratedClass);
+
+	// Only show bindable subobjects and variables if we're generating pure bindings.
+	if ( GeneratePureBindings )
+	{
+		UProperty* ReturnProperty = BindableSignature->GetReturnProperty();
+
+		// Find the binder that can handle the delegate return type, don't bother allowing people 
+		// to look for bindings that we don't support
+		if ( ensure(UWidget::FindBinderClassForDestination(ReturnProperty) != nullptr) )
+		{
+			static FName PropertyIcon(TEXT("Kismet.Tabs.Variables"));
+
+			MenuBuilder.BeginSection("Properties", LOCTEXT("Properties", "Properties"));
+			{
+				ForEachBindableProperty(OwnerStruct, [&] (UProperty* Property) {
+					TArray<UField*> NewBindingChain(BindingChain);
+					NewBindingChain.Add(Property);
+
+					bFoundEntry = true;
+
+					MenuBuilder.AddMenuEntry(
+						Property->GetDisplayNameText(),
+						Property->GetToolTipText(),
+						FSlateIcon(FEditorStyle::GetStyleSetName(), PropertyIcon),
+						FUIAction(FExecuteAction::CreateSP(this, &SPropertyBinding::HandleAddPropertyBinding, PropertyHandle, Property, NewBindingChain))
+						);
+				});
+			}
+			MenuBuilder.EndSection(); //Properties
+
+			MenuBuilder.BeginSection("SubObjectProperties", LOCTEXT("SubObjectProperties", "Sub-Object Properties"));
+			{
+				// Add all the properties that are not bindable, but are object or struct members that could contain children
+				// properties that are bindable.
+				for ( TFieldIterator<UProperty> PropIt(OwnerStruct, EFieldIteratorFlags::IncludeSuper); PropIt; ++PropIt )
+				{
+					UProperty* Property = *PropIt;
+
+					// Stop processing properties after reaching the user widget properties.
+					if ( IsFieldFromBlackListedClass(Property) )
+					{
+						break;
+					}
+
+					// If the owner is a class then use the blueprint scheme to determine if it's visible.
+					if ( !UEdGraphSchema_K2::CanUserKismetAccessVariable(Property, SkeletonClass, UEdGraphSchema_K2::CannotBeDelegate) )
+					{
+						continue;
+					}
+
+					if ( Property->HasAllPropertyFlags(CPF_BlueprintVisible) )
+					{
+						UObjectProperty* ObjectProperty = Cast<UObjectProperty>(Property);
+						UStructProperty* StructProperty = Cast<UStructProperty>(Property);
+
+						if ( ObjectProperty || StructProperty )
+						{
+							UStruct* Struct = ObjectProperty ? (UStruct*)ObjectProperty->PropertyClass : (UStruct*)StructProperty->Struct;
+
+							// Ignore any properties that are widgets, we don't want users binding widgets to other widgets.
+							// also ignore any class that is explicitly on the black list.
+							if ( ObjectProperty )
+							{
+								if ( IsClassBlackListed(ObjectProperty->PropertyClass) || ObjectProperty->PropertyClass->IsChildOf(UWidget::StaticClass()) )
+								{
+									continue;
+								}
+							}
+
+							// Stop processing properties after reaching the user widget properties.
+							if ( IsFieldFromBlackListedClass(Property) )
+							{
+								break;
+							}
+
+							TArray<UField*> NewBindingChain(BindingChain);
+							NewBindingChain.Add(Property);
+
+							bFoundEntry = true;
+
+							MenuBuilder.AddSubMenu(
+								Property->GetDisplayNameText(),
+								Property->GetToolTipText(),
+								FNewMenuDelegate::CreateSP(this, &SPropertyBinding::FillPropertyMenu, PropertyHandle, Struct, NewBindingChain)
+								);
+						}
+					}
+				}
+			}
+			MenuBuilder.EndSection(); //SubObjectProperties
+		}
+	}
+
+	if ( bFoundEntry == false && OwnerStruct != SkeletonClass )
+	{
+		MenuBuilder.BeginSection("None", OwnerStruct->GetDisplayNameText());
+		MenuBuilder.AddWidget(SNew(STextBlock).Text(LOCTEXT("None", "None")), FText::GetEmpty());
+		MenuBuilder.EndSection(); //None
+	}
+}
+
 const FSlateBrush* SPropertyBinding::GetCurrentBindingImage(TSharedRef<IPropertyHandle> PropertyHandle) const
 {
 	static FName PropertyIcon(TEXT("Kismet.Tabs.Variables"));
@@ -326,31 +464,38 @@ FText SPropertyBinding::GetCurrentBindingText(TSharedRef<IPropertyHandle> Proper
 		{
 			if ( Binding.ObjectName == OuterObjects[ObjectIndex]->GetName() && Binding.PropertyName == PropertyName )
 			{
-				if ( Binding.Kind == EBindingKind::Function )
+				if ( !Binding.SourcePath.IsEmpty() )
 				{
-					if ( Binding.MemberGuid.IsValid() )
-					{
-						// Graph function, look up by Guid
-						FName FoundName = Blueprint->GetFieldNameFromClassByGuid<UFunction>(Blueprint->GeneratedClass, Binding.MemberGuid);
-						return FText::FromString(FName::NameToDisplayString(FoundName.ToString(), false));
-					}
-					else
-					{
-						// No GUID, native function, return function name.
-						return FText::FromName(Binding.FunctionName);
-					}
+					return Binding.SourcePath.GetDisplayText();
 				}
-				else // Property
+				else
 				{
-					if ( Binding.MemberGuid.IsValid() )
+					if ( Binding.Kind == EBindingKind::Function )
 					{
-						FName FoundName = Blueprint->GetFieldNameFromClassByGuid<UProperty>(Blueprint->GeneratedClass, Binding.MemberGuid);
-						return FText::FromString(FName::NameToDisplayString(FoundName.ToString(), false));
+						if ( Binding.MemberGuid.IsValid() )
+						{
+							// Graph function, look up by Guid
+							FName FoundName = Blueprint->GetFieldNameFromClassByGuid<UFunction>(Blueprint->GeneratedClass, Binding.MemberGuid);
+							return FText::FromString(FName::NameToDisplayString(FoundName.ToString(), false));
+						}
+						else
+						{
+							// No GUID, native function, return function name.
+							return FText::FromName(Binding.FunctionName);
+						}
 					}
-					else
+					else // Property
 					{
-						// No GUID, native property, return source property.
-						return FText::FromName(Binding.SourceProperty);
+						if ( Binding.MemberGuid.IsValid() )
+						{
+							FName FoundName = Blueprint->GetFieldNameFromClassByGuid<UProperty>(Blueprint->GeneratedClass, Binding.MemberGuid);
+							return FText::FromString(FName::NameToDisplayString(FoundName.ToString(), false));
+						}
+						else
+						{
+							// No GUID, native property, return source property.
+							return FText::FromName(Binding.SourceProperty);
+						}
 					}
 				}
 			}
@@ -404,7 +549,13 @@ void SPropertyBinding::HandleRemoveBinding(TSharedRef<IPropertyHandle> PropertyH
 	FBlueprintEditorUtils::MarkBlueprintAsStructurallyModified(Blueprint);
 }
 
-void SPropertyBinding::HandleAddFunctionBinding(TSharedRef<IPropertyHandle> PropertyHandle, TSharedPtr<FunctionInfo> SelectedFunction)
+void SPropertyBinding::HandleAddFunctionBinding(TSharedRef<IPropertyHandle> PropertyHandle, TSharedPtr<FFunctionInfo> SelectedFunction, TArray<UField*> BindingChain)
+{
+	FEditorPropertyPath BindingPath(BindingChain);
+	HandleAddFunctionBinding(PropertyHandle, SelectedFunction, BindingPath);
+}
+
+void SPropertyBinding::HandleAddFunctionBinding(TSharedRef<IPropertyHandle> PropertyHandle, TSharedPtr<FFunctionInfo> SelectedFunction, FEditorPropertyPath& BindingPath)
 {
 	const FScopedTransaction Transaction(LOCTEXT("BindDelegate", "Set Binding"));
 
@@ -418,7 +569,7 @@ void SPropertyBinding::HandleAddFunctionBinding(TSharedRef<IPropertyHandle> Prop
 		Binding.ObjectName = SelectedObject->GetName();
 		Binding.PropertyName = PropertyHandle->GetProperty()->GetFName();
 		Binding.FunctionName = SelectedFunction->FuncName;
-		Binding.MemberGuid = ( SelectedFunction->EdGraph ) ? SelectedFunction->EdGraph->GraphGuid : FGuid();
+		Binding.SourcePath = BindingPath;
 		Binding.Kind = EBindingKind::Function;
 
 		Blueprint->Bindings.Remove(Binding);
@@ -428,12 +579,12 @@ void SPropertyBinding::HandleAddFunctionBinding(TSharedRef<IPropertyHandle> Prop
 	FBlueprintEditorUtils::MarkBlueprintAsStructurallyModified(Blueprint);
 }
 
-void SPropertyBinding::HandleAddPropertyBinding(TSharedRef<IPropertyHandle> PropertyHandle, UProperty* SelectedProperty)
+void SPropertyBinding::HandleAddPropertyBinding(TSharedRef<IPropertyHandle> PropertyHandle, UProperty* SelectedProperty, TArray<UField*> BindingChain)
 {
 	const FScopedTransaction Transaction(LOCTEXT("BindDelegate", "Set Binding"));
 
 	// Get the current skeleton class, think header for the blueprint.
-	UBlueprintGeneratedClass* SkeletonClass = Cast<UBlueprintGeneratedClass>(Blueprint->SkeletonGeneratedClass);
+	UBlueprintGeneratedClass* SkeletonClass = Cast<UBlueprintGeneratedClass>(Blueprint->GeneratedClass);
 
 	Blueprint->Modify();
 
@@ -448,6 +599,7 @@ void SPropertyBinding::HandleAddPropertyBinding(TSharedRef<IPropertyHandle> Prop
 		Binding.ObjectName = SelectedObject->GetName();
 		Binding.PropertyName = PropertyHandle->GetProperty()->GetFName();
 		Binding.SourceProperty = SelectedProperty->GetFName();
+		Binding.SourcePath = FEditorPropertyPath(BindingChain);
 		Binding.MemberGuid = MemberGuid;
 		Binding.Kind = EBindingKind::Property;
 
@@ -458,13 +610,13 @@ void SPropertyBinding::HandleAddPropertyBinding(TSharedRef<IPropertyHandle> Prop
 	FBlueprintEditorUtils::MarkBlueprintAsStructurallyModified(Blueprint);
 }
 
-void SPropertyBinding::HandleCreateAndAddBinding(UWidget* Widget, TSharedRef<IPropertyHandle> PropertyHandle, UFunction* DelegateSignature, bool bIsPure)
+void SPropertyBinding::HandleCreateAndAddBinding(UWidget* Widget, TSharedRef<IPropertyHandle> PropertyHandle)
 {
 	const FScopedTransaction Transaction(LOCTEXT("CreateDelegate", "Create Binding"));
 
 	Blueprint->Modify();
 
-	FString Pre = bIsPure ? FString(TEXT("Get")) : FString(TEXT("On"));
+	FString Pre = GeneratePureBindings ? FString(TEXT("Get")) : FString(TEXT("On"));
 
 	FString WidgetName;
 	if ( Widget && !Widget->IsGeneratedName() )
@@ -485,17 +637,19 @@ void SPropertyBinding::HandleCreateAndAddBinding(UWidget* Widget, TSharedRef<IPr
 		UEdGraphSchema_K2::StaticClass());
 	
 	// Add the binding to the blueprint
-	TSharedPtr<FunctionInfo> SelectedFunction = MakeShareable(new FunctionInfo());
+	TSharedPtr<FFunctionInfo> SelectedFunction = MakeShareable(new FFunctionInfo());
 	SelectedFunction->FuncName = FunctionGraph->GetFName();
-	SelectedFunction->EdGraph = FunctionGraph;
 
-	HandleAddFunctionBinding(PropertyHandle, SelectedFunction);
+	FEditorPropertyPath BindingPath;
+	BindingPath.Segments.Add(FEditorPropertyPathSegment(FunctionGraph));
+
+	HandleAddFunctionBinding(PropertyHandle, SelectedFunction, BindingPath);
 
 	const bool bUserCreated = true;
-	FBlueprintEditorUtils::AddFunctionGraph(Blueprint, FunctionGraph, bUserCreated, DelegateSignature);
+	FBlueprintEditorUtils::AddFunctionGraph(Blueprint, FunctionGraph, bUserCreated, BindableSignature);
 
 	// Only mark bindings as pure that need to be.
-	if ( bIsPure )
+	if ( GeneratePureBindings )
 	{
 		const UEdGraphSchema_K2* Schema_K2 = Cast<UEdGraphSchema_K2>(FunctionGraph->GetSchema());
 		Schema_K2->AddExtraFunctionFlags(FunctionGraph, FUNC_BlueprintPure);
diff --git a/Engine/Source/Editor/UMGEditor/Private/Details/SPropertyBinding.h b/Engine/Source/Editor/UMGEditor/Private/Details/SPropertyBinding.h
--- a/Engine/Source/Editor/UMGEditor/Private/Details/SPropertyBinding.h
+++ b/Engine/Source/Editor/UMGEditor/Private/Details/SPropertyBinding.h
@@ -18,18 +18,17 @@ public:
 	void Construct(const FArguments& InArgs, TSharedRef<FWidgetBlueprintEditor> InEditor, UDelegateProperty* DelegateProperty, TSharedRef<IPropertyHandle> Property);
 
 protected:
-	struct FunctionInfo
+	struct FFunctionInfo
 	{
 		FText DisplayName;
 		FString Tooltip;
 
 		FName FuncName;
-		UEdGraph* EdGraph;
+		UFunction* Function;
 	};
 
-	void RefreshBlueprintMemberCache(const UFunction* DelegateSignature, bool bIsPure);
-
-	TSharedRef<SWidget> OnGenerateDelegateMenu(UWidget* Widget, TSharedRef<IPropertyHandle> PropertyHandle, UFunction* DelegateSignature, bool bIsPure);
+	TSharedRef<SWidget> OnGenerateDelegateMenu(UWidget* Widget, TSharedRef<IPropertyHandle> PropertyHandle);
+	void FillPropertyMenu(FMenuBuilder& MenuBuilder, TSharedRef<IPropertyHandle> PropertyHandle, UStruct* OwnerStruct, TArray<UField*> BindingChain);
 
 	const FSlateBrush* GetCurrentBindingImage(TSharedRef<IPropertyHandle> PropertyHandle) const;
 	FText GetCurrentBindingText(TSharedRef<IPropertyHandle> PropertyHandle) const;
@@ -37,10 +36,11 @@ protected:
 	bool CanRemoveBinding(TSharedRef<IPropertyHandle> PropertyHandle);
 	void HandleRemoveBinding(TSharedRef<IPropertyHandle> PropertyHandle);
 
-	void HandleAddFunctionBinding(TSharedRef<IPropertyHandle> PropertyHandle, TSharedPtr<FunctionInfo> SelectedFunction);
-	void HandleAddPropertyBinding(TSharedRef<IPropertyHandle> PropertyHandle, UProperty* SelectedProperty);
+	void HandleAddFunctionBinding(TSharedRef<IPropertyHandle> PropertyHandle, TSharedPtr<FFunctionInfo> SelectedFunction, TArray<UField*> BindingChain);
+	void HandleAddFunctionBinding(TSharedRef<IPropertyHandle> PropertyHandle, TSharedPtr<FFunctionInfo> SelectedFunction, FEditorPropertyPath& BindingPath);
+	void HandleAddPropertyBinding(TSharedRef<IPropertyHandle> PropertyHandle, UProperty* SelectedProperty, TArray<UField*> BindingChain);
 
-	void HandleCreateAndAddBinding(UWidget* Widget, TSharedRef<IPropertyHandle> PropertyHandle, UFunction* DelegateSignature, bool bIsPure);
+	void HandleCreateAndAddBinding(UWidget* Widget, TSharedRef<IPropertyHandle> PropertyHandle);
 	void GotoFunction(UEdGraph* FunctionGraph);
 
 	EVisibility GetGotoBindingVisibility(TSharedRef<IPropertyHandle> PropertyHandle) const;
@@ -51,10 +51,16 @@ protected:
 
 private:
 
+	template <typename Predicate>
+	void ForEachBindableProperty(UStruct* InStruct, Predicate Pred) const;
+
+	template <typename Predicate>
+	void ForEachBindableFunction(UClass* FromClass, Predicate Pred) const;
+
 	TWeakPtr<FWidgetBlueprintEditor> Editor;
 
 	UWidgetBlueprint* Blueprint;
 
-	TArray< TSharedPtr<FunctionInfo> > BlueprintFunctionCache;
-	TArray< UProperty* > BlueprintPropertyCache;
+	bool GeneratePureBindings;
+	UFunction* BindableSignature;
 };
diff --git a/Engine/Source/Editor/UMGEditor/Private/WidgetBlueprint.cpp b/Engine/Source/Editor/UMGEditor/Private/WidgetBlueprint.cpp
--- a/Engine/Source/Editor/UMGEditor/Private/WidgetBlueprint.cpp
+++ b/Engine/Source/Editor/UMGEditor/Private/WidgetBlueprint.cpp
@@ -1,15 +1,349 @@
 // Copyright 1998-2014 Epic Games, Inc. All Rights Reserved.
 
 #include "UMGEditorPrivatePCH.h"
-#include "Blueprint/WidgetTree.h"
+
 #include "Runtime/MovieSceneCore/Classes/MovieScene.h"
-#include "PropertyTag.h"
+#include "Editor/UnrealEd/Public/Kismet2/StructureEditorUtils.h"
+
+#include "Blueprint/WidgetTree.h"
 #include "Blueprint/WidgetBlueprintGeneratedClass.h"
+#include "PropertyTag.h"
 #include "WidgetBlueprint.h"
 #include "WidgetBlueprintCompiler.h"
+#include "PropertyBinding.h"
 
 #define LOCTEXT_NAMESPACE "UMG"
 
+FEditorPropertyPathSegment::FEditorPropertyPathSegment()
+	: Struct(nullptr)
+	, MemberName(NAME_None)
+	, MemberGuid()
+	, IsProperty(true)
+{
+}
+
+FEditorPropertyPathSegment::FEditorPropertyPathSegment(const UProperty* InProperty)
+{
+	IsProperty = true;
+	MemberName = InProperty->GetFName();
+	if ( InProperty->GetOwnerClass() )
+	{
+		Struct = InProperty->GetOwnerClass();
+		UBlueprint::GetGuidFromClassByFieldName<UProperty>(InProperty->GetOwnerClass(), InProperty->GetFName(), MemberGuid);
+	}
+	else if ( InProperty->GetOwnerStruct() )
+	{
+		Struct = InProperty->GetOwnerStruct();
+		MemberGuid = FStructureEditorUtils::GetGuidForProperty(InProperty);
+	}
+	else
+	{
+		// Should not be possible to hit.
+		check(false);
+	}
+}
+
+FEditorPropertyPathSegment::FEditorPropertyPathSegment(const UFunction* InFunction)
+{
+	IsProperty = false;
+	MemberName = InFunction->GetFName();
+	if ( InFunction->GetOwnerClass() )
+	{
+		Struct = InFunction->GetOwnerClass();
+		UBlueprint::GetGuidFromClassByFieldName<UFunction>(InFunction->GetOwnerClass(), InFunction->GetFName(), MemberGuid);
+	}
+	else
+	{
+		// Should not be possible to hit.
+		check(false);
+	}
+}
+
+FEditorPropertyPathSegment::FEditorPropertyPathSegment(const UEdGraph* InFunctionGraph)
+{
+	IsProperty = false;
+	MemberName = InFunctionGraph->GetFName();
+	UBlueprint* Blueprint = CastChecked<UBlueprint>(InFunctionGraph->GetOuter());
+	Struct = Blueprint->GeneratedClass;
+	check(Struct);
+	MemberGuid = InFunctionGraph->GraphGuid;
+}
+
+bool FEditorPropertyPathSegment::ValidateMember(UDelegateProperty* DelegateProperty, FText& OutError) const
+{
+	if ( DelegateProperty->SignatureFunction->NumParms == 1 )
+	{
+		if ( UProperty* ReturnProperty = DelegateProperty->SignatureFunction->GetReturnProperty() )
+		{
+			// TODO I don't like having the path segment system needing to have knowledge of the binding layer.
+			// think about divorcing the two.
+
+			// Find the binder that can handle the delegate return type.
+			TSubclassOf<UPropertyBinding> Binder = UWidget::FindBinderClassForDestination(ReturnProperty);
+			if ( Binder == nullptr )
+			{
+				OutError = FText::Format(LOCTEXT("Binding_Binder_NotFound", "Member:{0}: No binding exists for {1}."),
+					GetMemberDisplayText(),
+					ReturnProperty->GetClass()->GetDisplayNameText());
+				return false;
+			}
+
+			if ( UField* Field = GetMember() )
+			{
+				if ( UProperty* Property = Cast<UProperty>(Field) )
+				{
+					// Ensure that the binder also can handle binding from the property we care about.
+					if ( Binder->GetDefaultObject<UPropertyBinding>()->IsSupportedSource(Property) )
+					{
+						return true;
+					}
+					else
+					{
+						OutError = FText::Format(LOCTEXT("Binding_UnsupportedType_Property", "Member:{0} Unable to bind {1}, unsupported type."),
+							GetMemberDisplayText(),
+							Property->GetClass()->GetDisplayNameText());
+						return false;
+					}
+				}
+				else if ( UFunction* Function = Cast<UFunction>(Field) )
+				{
+					if ( Function->NumParms == 1 )
+					{
+						if ( Function->HasAnyFunctionFlags(FUNC_Const | FUNC_BlueprintPure) )
+						{
+							if ( UProperty* MemberReturn = Function->GetReturnProperty() )
+							{
+								// Ensure that the binder also can handle binding from the property we care about.
+								if ( Binder->GetDefaultObject<UPropertyBinding>()->IsSupportedSource(MemberReturn) )
+								{
+									return true;
+								}
+								else
+								{
+									OutError = FText::Format(LOCTEXT("Binding_UnsupportedType_Function", "Member:{0} Unable to bind {1}, unsupported type."),
+										GetMemberDisplayText(),
+										MemberReturn->GetClass()->GetDisplayNameText());
+									return false;
+								}
+							}
+							else
+							{
+								OutError = FText::Format(LOCTEXT("Binding_NoReturn", "Member:{0} Has no return value, unable to bind."),
+									GetMemberDisplayText());
+								return false;
+							}
+						}
+						else
+						{
+							OutError = FText::Format(LOCTEXT("Binding_Pure", "Member:{0} Unable to bind, the function is not marked as pure."),
+								GetMemberDisplayText());
+							return false;
+						}
+					}
+					else
+					{
+						OutError = FText::Format(LOCTEXT("Binding_NumArgs", "Member:{0} Has the wrong number of arguments, it needs to return 1 value and take no parameters."),
+							GetMemberDisplayText());
+						return false;
+					}
+				}
+			}
+		}
+	}
+
+	return false;
+}
+
+UField* FEditorPropertyPathSegment::GetMember() const
+{
+	FName FieldName = GetMemberName();
+	if ( FieldName != NAME_None )
+	{
+		UField* Field = FindField<UField>(Struct, FieldName);
+		//if ( Field == nullptr )
+		//{
+		//	if ( UClass* Class = Cast<UClass>(Struct) )
+		//	{
+		//		if ( UBlueprint* Blueprint = Cast<UBlueprint>(Class->ClassGeneratedBy) )
+		//		{
+		//			if ( UClass* SkeletonClass = Blueprint->SkeletonGeneratedClass )
+		//			{
+		//				Field = FindField<UField>(SkeletonClass, FieldName);
+		//			}
+		//		}
+		//	}
+		//}
+
+		return Field;
+	}
+
+	return nullptr;
+}
+
+FName FEditorPropertyPathSegment::GetMemberName() const
+{
+	if ( MemberGuid.IsValid() )
+	{
+		FName NameFromGuid = NAME_None;
+
+		if ( UClass* Class = Cast<UClass>(Struct) )
+		{
+			if ( UBlueprint* Blueprint = Cast<UBlueprint>(Class->ClassGeneratedBy) )
+			{
+				if ( IsProperty )
+				{
+					NameFromGuid = UBlueprint::GetFieldNameFromClassByGuid<UProperty>(Class, MemberGuid);
+				}
+				else
+				{
+					NameFromGuid = UBlueprint::GetFieldNameFromClassByGuid<UFunction>(Class, MemberGuid);
+				}
+			}
+		}
+		else if ( UUserDefinedStruct* UserStruct = Cast<UUserDefinedStruct>(Struct) )
+		{
+			if ( UProperty* Property = FStructureEditorUtils::GetPropertyByGuid(UserStruct, MemberGuid) )
+			{
+				NameFromGuid = Property->GetFName();
+			}
+		}
+
+		//check(NameFromGuid != NAME_None);
+		return NameFromGuid;
+	}
+	
+	//check(MemberName != NAME_None);
+	return MemberName;
+}
+
+FText FEditorPropertyPathSegment::GetMemberDisplayText() const
+{
+	if ( MemberGuid.IsValid() )
+	{
+		if ( UClass* Class = Cast<UClass>(Struct) )
+		{
+			if ( UBlueprint* Blueprint = Cast<UBlueprint>(Class->ClassGeneratedBy) )
+			{
+				if ( IsProperty )
+				{
+					return FText::FromName(UBlueprint::GetFieldNameFromClassByGuid<UProperty>(Class, MemberGuid));
+				}
+				else
+				{
+					return FText::FromName(UBlueprint::GetFieldNameFromClassByGuid<UFunction>(Class, MemberGuid));
+				}
+			}
+		}
+		else if ( UUserDefinedStruct* UserStruct = Cast<UUserDefinedStruct>(Struct) )
+		{
+			if ( UProperty* Property = FStructureEditorUtils::GetPropertyByGuid(UserStruct, MemberGuid) )
+			{
+				return Property->GetDisplayNameText();
+			}
+		}
+	}
+
+	return FText::FromName(MemberName);
+}
+
+FEditorPropertyPath::FEditorPropertyPath()
+{
+}
+
+FEditorPropertyPath::FEditorPropertyPath(const TArray<UField*>& BindingChain)
+{
+	for ( const UField* Field : BindingChain )
+	{
+		if ( const UProperty* Property = Cast<UProperty>(Field) )
+		{
+			Segments.Add(FEditorPropertyPathSegment(Property));
+		}
+		else if ( const UFunction* Function = Cast<UFunction>(Field) )
+		{
+			Segments.Add(FEditorPropertyPathSegment(Function));
+		}
+		else
+		{
+			// Should never happen
+			check(false);
+		}
+	}
+}
+
+bool FEditorPropertyPath::Validate(UDelegateProperty* Destination, FText& OutError) const
+{
+	if ( IsEmpty() )
+	{
+		OutError = LOCTEXT("Binding_Empty", "The binding is empty.");
+		return false;
+	}
+
+	for ( int32 SegmentIndex = 0; SegmentIndex < Segments.Num(); SegmentIndex++ )
+	{
+		const FEditorPropertyPathSegment& Segment = Segments[SegmentIndex];
+		if ( Segment.GetStruct() == nullptr )
+		{
+			OutError = FText::Format(LOCTEXT("Binding_StructNotFound", "Binding:{0}, Unable to locate owner class."),
+				GetDisplayText());
+
+			return false;
+		}
+
+		if ( Segment.GetMember() == nullptr )
+		{
+			OutError = FText::Format(LOCTEXT("Binding_MemberNotFound", "Binding:{0}, Member:{1}, was not found."),
+				GetDisplayText(),
+				Segment.GetMemberDisplayText());
+
+			return false;
+		}
+
+		// TODO Add additional checks, are non-last fields objects, structs still?
+	}
+
+	// Validate the last member in the segment
+	const FEditorPropertyPathSegment& LastSegment = Segments[Segments.Num() - 1];
+	return LastSegment.ValidateMember(Destination, OutError);
+}
+
+FText FEditorPropertyPath::GetDisplayText() const
+{
+	FString DisplayText;
+
+	for ( int32 SegmentIndex = 0; SegmentIndex < Segments.Num(); SegmentIndex++ )
+	{
+		const FEditorPropertyPathSegment& Segment = Segments[SegmentIndex];
+		DisplayText.Append(Segment.GetMemberDisplayText().ToString());
+		if ( SegmentIndex < ( Segments.Num() - 1 ) )
+		{
+			DisplayText.Append(TEXT("."));
+		}
+	}
+
+	return FText::FromString(DisplayText);
+}
+
+FDynamicPropertyPath FEditorPropertyPath::ToPropertyPath() const
+{
+	TArray<FString> PropertyChain;
+
+	for ( const FEditorPropertyPathSegment& Segment : Segments )
+	{
+		FName SegmentName = Segment.GetMemberName();
+
+		if ( SegmentName != NAME_None )
+		{
+			PropertyChain.Add(SegmentName.ToString());
+		}
+		else
+		{
+			return FDynamicPropertyPath();
+		}
+	}
+
+	return FDynamicPropertyPath(PropertyChain);
+}
+
 bool FDelegateEditorBinding::IsBindingValid(UClass* BlueprintGeneratedClass, UWidgetBlueprint* Blueprint, FCompilerResultsLog& MessageLog) const
 {
 	FDelegateRuntimeBinding RuntimeBinding = ToRuntimeBinding(Blueprint);
@@ -28,34 +362,50 @@ bool FDelegateEditorBinding::IsBindingValid(UClass* BlueprintGeneratedClass, UWi
 		// Locate the delegate property on the widget that's a delegate for a property we want to bind.
 		if ( DelegateProperty )
 		{
-			// On our incoming blueprint generated class, try and find the function we claim exists that users
-			// are binding their property too.
-			if ( UFunction* Function = BlueprintGeneratedClass->FindFunctionByName(RuntimeBinding.FunctionName, EIncludeSuperFlag::IncludeSuper) )
+			if ( !SourcePath.IsEmpty() )
+			{
+				FText ValidationError;
+				if ( SourcePath.Validate(DelegateProperty, ValidationError) == false )
+				{
+					FText const ErrorFormat = LOCTEXT("BindingError", "Binding: Property '@@' on Widget '@@': %s");
+					MessageLog.Error(*FString::Printf(*ErrorFormat.ToString(), *ValidationError.ToString()), DelegateProperty, TargetWidget);
+
+					return false;
+				}
+
+				return true;
+			}
+			else
 			{
-				// Check the signatures to ensure these functions match.
-				if ( Function->IsSignatureCompatibleWith(DelegateProperty->SignatureFunction, UFunction::GetDefaultIgnoredSignatureCompatibilityFlags() | CPF_ReturnParm) )
+				// On our incoming blueprint generated class, try and find the function we claim exists that users
+				// are binding their property too.
+				if ( UFunction* Function = BlueprintGeneratedClass->FindFunctionByName(RuntimeBinding.FunctionName, EIncludeSuperFlag::IncludeSuper) )
 				{
-					// Only allow binding pure functions to property bindings.
-					if ( bNeedsToBePure && !Function->HasAllFunctionFlags(FUNC_BlueprintPure) )
+					// Check the signatures to ensure these functions match.
+					if ( Function->IsSignatureCompatibleWith(DelegateProperty->SignatureFunction, UFunction::GetDefaultIgnoredSignatureCompatibilityFlags() | CPF_ReturnParm) )
 					{
-						FText const ErrorFormat = LOCTEXT("BindingNotBoundToPure", "Binding: property '@@' on widget '@@' needs to be bound to a pure function, '@@' is not pure.");
-						MessageLog.Error(*ErrorFormat.ToString(), DelegateProperty, TargetWidget, Function);
+						// Only allow binding pure functions to property bindings.
+						if ( bNeedsToBePure && !Function->HasAnyFunctionFlags(FUNC_Const | FUNC_BlueprintPure) )
+						{
+							FText const ErrorFormat = LOCTEXT("BindingNotBoundToPure", "Binding: property '@@' on widget '@@' needs to be bound to a pure function, '@@' is not pure.");
+							MessageLog.Error(*ErrorFormat.ToString(), DelegateProperty, TargetWidget, Function);
 
-						return false;
-					}
+							return false;
+						}
 
-					return true;
+						return true;
+					}
+					else
+					{
+						FText const ErrorFormat = LOCTEXT("BindingFunctionSigDontMatch", "Binding: property '@@' on widget '@@' bound to function '@@', but the sigatnures don't match.  The function must return the same type as the property and have no parameters.");
+						MessageLog.Error(*ErrorFormat.ToString(), DelegateProperty, TargetWidget, Function);
+					}
 				}
 				else
 				{
-					FText const ErrorFormat = LOCTEXT("BindingFunctionSigDontMatch", "Binding: property '@@' on widget '@@' bound to function '@@', but the sigatnures don't match.  The function must return the same type as the property and have no parameters.");
-					MessageLog.Error(*ErrorFormat.ToString(), DelegateProperty, TargetWidget, Function);
+					//TODO Bindable property removed.
 				}
 			}
-			else
-			{
-				//TODO Bindable property removed.
-			}
 		}
 		else
 		{
@@ -84,6 +434,7 @@ FDelegateRuntimeBinding FDelegateEditorBinding::ToRuntimeBinding(UWidgetBlueprin
 		Binding.FunctionName = FunctionName;
 	}
 	Binding.Kind = Kind;
+	Binding.SourcePath = SourcePath.ToPropertyPath();
 
 	return Binding;
 }
diff --git a/Engine/Source/Editor/UMGEditor/Private/WidgetBlueprintCompiler.cpp b/Engine/Source/Editor/UMGEditor/Private/WidgetBlueprintCompiler.cpp
--- a/Engine/Source/Editor/UMGEditor/Private/WidgetBlueprintCompiler.cpp
+++ b/Engine/Source/Editor/UMGEditor/Private/WidgetBlueprintCompiler.cpp
@@ -30,64 +30,67 @@ void FWidgetBlueprintCompiler::CreateFunctionList()
 
 	for ( FDelegateEditorBinding& EditorBinding : WidgetBlueprint()->Bindings )
 	{
-		const FName PropertyName = EditorBinding.SourceProperty;
-
-		UProperty* Property = FindField<UProperty>(Blueprint->SkeletonGeneratedClass, PropertyName);
-		if ( Property )
+		if ( EditorBinding.SourcePath.IsEmpty() )
 		{
-			// Create the function graph.
-			FString FunctionName = FString(TEXT("__Get")) + PropertyName.ToString();
-			UEdGraph* FunctionGraph = FBlueprintEditorUtils::CreateNewGraph(Blueprint, FBlueprintEditorUtils::FindUniqueKismetName(Blueprint, FunctionName), UEdGraph::StaticClass(), UEdGraphSchema_K2::StaticClass());
+			const FName PropertyName = EditorBinding.SourceProperty;
+
+			UProperty* Property = FindField<UProperty>(Blueprint->SkeletonGeneratedClass, PropertyName);
+			if ( Property )
+			{
+				// Create the function graph.
+				FString FunctionName = FString(TEXT("__Get")) + PropertyName.ToString();
+				UEdGraph* FunctionGraph = FBlueprintEditorUtils::CreateNewGraph(Blueprint, FBlueprintEditorUtils::FindUniqueKismetName(Blueprint, FunctionName), UEdGraph::StaticClass(), UEdGraphSchema_K2::StaticClass());
 
-			// Update the function binding to match the generated graph name
-			EditorBinding.FunctionName = FunctionGraph->GetFName();
+				// Update the function binding to match the generated graph name
+				EditorBinding.FunctionName = FunctionGraph->GetFName();
 
-			const UEdGraphSchema_K2* K2Schema = Cast<const UEdGraphSchema_K2>(FunctionGraph->GetSchema());
+				const UEdGraphSchema_K2* K2Schema = Cast<const UEdGraphSchema_K2>(FunctionGraph->GetSchema());
 
-			Schema->CreateDefaultNodesForGraph(*FunctionGraph);
+				Schema->CreateDefaultNodesForGraph(*FunctionGraph);
 
-			K2Schema->MarkFunctionEntryAsEditable(FunctionGraph, true);
+				K2Schema->MarkFunctionEntryAsEditable(FunctionGraph, true);
 
-			// Create a function entry node
-			FGraphNodeCreator<UK2Node_FunctionEntry> FunctionEntryCreator(*FunctionGraph);
-			UK2Node_FunctionEntry* EntryNode = FunctionEntryCreator.CreateNode();
-			EntryNode->SignatureClass = NULL;
-			EntryNode->SignatureName = FunctionGraph->GetFName();
-			FunctionEntryCreator.Finalize();
+				// Create a function entry node
+				FGraphNodeCreator<UK2Node_FunctionEntry> FunctionEntryCreator(*FunctionGraph);
+				UK2Node_FunctionEntry* EntryNode = FunctionEntryCreator.CreateNode();
+				EntryNode->SignatureClass = NULL;
+				EntryNode->SignatureName = FunctionGraph->GetFName();
+				FunctionEntryCreator.Finalize();
 
-			FGraphNodeCreator<UK2Node_FunctionResult> FunctionReturnCreator(*FunctionGraph);
-			UK2Node_FunctionResult* ReturnNode = FunctionReturnCreator.CreateNode();
-			ReturnNode->SignatureClass = NULL;
-			ReturnNode->SignatureName = FunctionGraph->GetFName();
-			ReturnNode->NodePosX = EntryNode->NodePosX + EntryNode->NodeWidth + 256;
-			ReturnNode->NodePosY = EntryNode->NodePosY;
-			FunctionReturnCreator.Finalize();
+				FGraphNodeCreator<UK2Node_FunctionResult> FunctionReturnCreator(*FunctionGraph);
+				UK2Node_FunctionResult* ReturnNode = FunctionReturnCreator.CreateNode();
+				ReturnNode->SignatureClass = NULL;
+				ReturnNode->SignatureName = FunctionGraph->GetFName();
+				ReturnNode->NodePosX = EntryNode->NodePosX + EntryNode->NodeWidth + 256;
+				ReturnNode->NodePosY = EntryNode->NodePosY;
+				FunctionReturnCreator.Finalize();
 
-			FEdGraphPinType PinType;
-			K2Schema->ConvertPropertyToPinType(Property, /*out*/ PinType);
+				FEdGraphPinType PinType;
+				K2Schema->ConvertPropertyToPinType(Property, /*out*/ PinType);
 
-			UEdGraphPin* ReturnPin = ReturnNode->CreateUserDefinedPin(TEXT("ReturnValue"), PinType);
+				UEdGraphPin* ReturnPin = ReturnNode->CreateUserDefinedPin(TEXT("ReturnValue"), PinType);
 
-			// Auto-connect the pins for entry and exit, so that by default the signature is properly generated
-			UEdGraphPin* EntryNodeExec = K2Schema->FindExecutionPin(*EntryNode, EGPD_Output);
-			UEdGraphPin* ResultNodeExec = K2Schema->FindExecutionPin(*ReturnNode, EGPD_Input);
-			EntryNodeExec->MakeLinkTo(ResultNodeExec);
+				// Auto-connect the pins for entry and exit, so that by default the signature is properly generated
+				UEdGraphPin* EntryNodeExec = K2Schema->FindExecutionPin(*EntryNode, EGPD_Output);
+				UEdGraphPin* ResultNodeExec = K2Schema->FindExecutionPin(*ReturnNode, EGPD_Input);
+				EntryNodeExec->MakeLinkTo(ResultNodeExec);
 
-			FGraphNodeCreator<UK2Node_VariableGet> MemberGetCreator(*FunctionGraph);
-			UK2Node_VariableGet* VarNode = MemberGetCreator.CreateNode();
-			VarNode->VariableReference.SetSelfMember(PropertyName);
-			MemberGetCreator.Finalize();
+				FGraphNodeCreator<UK2Node_VariableGet> MemberGetCreator(*FunctionGraph);
+				UK2Node_VariableGet* VarNode = MemberGetCreator.CreateNode();
+				VarNode->VariableReference.SetSelfMember(PropertyName);
+				MemberGetCreator.Finalize();
 
-			ReturnPin->MakeLinkTo(VarNode->GetValuePin());
+				ReturnPin->MakeLinkTo(VarNode->GetValuePin());
 
-			// We need to flag the entry node to make sure that the compiled function is callable from Kismet2
-			int32 ExtraFunctionFlags = ( FUNC_BlueprintCallable | FUNC_BlueprintEvent | FUNC_Public | FUNC_BlueprintPure );
-			K2Schema->AddExtraFunctionFlags(FunctionGraph, ExtraFunctionFlags);
+				// We need to flag the entry node to make sure that the compiled function is callable from Kismet2
+				int32 ExtraFunctionFlags = ( FUNC_Private | FUNC_Const );
+				K2Schema->AddExtraFunctionFlags(FunctionGraph, ExtraFunctionFlags);
 
-			//Blueprint->FunctionGraphs.Add(FunctionGraph);
+				//Blueprint->FunctionGraphs.Add(FunctionGraph);
 
-			ProcessOneFunctionGraph(FunctionGraph);
-			//FEdGraphUtilities::MergeChildrenGraphsIn(Ubergraph, FunctionGraph, /*bRequireSchemaMatch=*/ true);
+				ProcessOneFunctionGraph(FunctionGraph, true);
+				//FEdGraphUtilities::MergeChildrenGraphsIn(Ubergraph, FunctionGraph, /*bRequireSchemaMatch=*/ true);
+			}
 		}
 	}
 }
@@ -256,13 +259,20 @@ void FWidgetBlueprintCompiler::FinishCompilingClass(UClass* Class)
 
 	BPGClass->Bindings.Reset();
 
-	// Convert all editor time property bindings into a list of bindings
-	// that will be applied at runtime.  Ensure all bindings are still valid.
-	for ( const FDelegateEditorBinding& EditorBinding : Blueprint->Bindings )
+	// Only check bindings on a full compile.  Also don't check them if we're regenerating on load,
+	// that has a nasty tendency to fail because the other dependent classes that may also be blueprints
+	// might not be loaded yet.
+	const bool bIsLoading = Blueprint->bIsRegeneratingOnLoad;
+	if ( bIsFullCompile )
 	{
-		if ( EditorBinding.IsBindingValid(Class, Blueprint, MessageLog) )
+		// Convert all editor time property bindings into a list of bindings
+		// that will be applied at runtime.  Ensure all bindings are still valid.
+		for ( const FDelegateEditorBinding& EditorBinding : Blueprint->Bindings )
 		{
-			BPGClass->Bindings.Add(EditorBinding.ToRuntimeBinding(Blueprint));
+			if ( bIsLoading || EditorBinding.IsBindingValid(Class, Blueprint, MessageLog) )
+			{
+				BPGClass->Bindings.Add(EditorBinding.ToRuntimeBinding(Blueprint));
+			}
 		}
 	}
 
@@ -290,7 +300,7 @@ void FWidgetBlueprintCompiler::EnsureProperGeneratedClass(UClass*& TargetUClass)
 	if ( TargetUClass && !( (UObject*)TargetUClass )->IsA(UWidgetBlueprintGeneratedClass::StaticClass()) )
 	{
 		FKismetCompilerUtilities::ConsignToOblivion(TargetUClass, Blueprint->bIsRegeneratingOnLoad);
-		TargetUClass = NULL;
+		TargetUClass = nullptr;
 	}
 }
 
@@ -298,7 +308,7 @@ void FWidgetBlueprintCompiler::SpawnNewClass(const FString& NewClassName)
 {
 	NewWidgetBlueprintClass = FindObject<UWidgetBlueprintGeneratedClass>(Blueprint->GetOutermost(), *NewClassName);
 
-	if ( NewWidgetBlueprintClass == NULL )
+	if ( NewWidgetBlueprintClass == nullptr )
 	{
 		NewWidgetBlueprintClass = ConstructObject<UWidgetBlueprintGeneratedClass>(UWidgetBlueprintGeneratedClass::StaticClass(), Blueprint->GetOutermost(), FName(*NewClassName), RF_Public | RF_Transactional);
 	}
diff --git a/Engine/Source/Editor/UMGEditor/Public/WidgetBlueprint.h b/Engine/Source/Editor/UMGEditor/Public/WidgetBlueprint.h
--- a/Engine/Source/Editor/UMGEditor/Public/WidgetBlueprint.h
+++ b/Engine/Source/Editor/UMGEditor/Public/WidgetBlueprint.h
@@ -5,11 +5,80 @@
 #include "Engine/Blueprint.h"
 #include "UserWidget.h"
 #include "WidgetAnimation.h"
+#include "DynamicPropertyPath.h"
 
 #include "WidgetBlueprint.generated.h"
 
 class UMovieScene;
 
+/** */
+USTRUCT()
+struct UMGEDITOR_API FEditorPropertyPathSegment
+{
+	GENERATED_USTRUCT_BODY()
+
+public:
+	FEditorPropertyPathSegment();
+	FEditorPropertyPathSegment(const UProperty* InProperty);
+	FEditorPropertyPathSegment(const UFunction* InFunction);
+	FEditorPropertyPathSegment(const UEdGraph* InFunctionGraph);
+
+	UStruct* GetStruct() const { return Struct; }
+	UField* GetMember() const;
+
+	bool ValidateMember(UDelegateProperty* DelegateProperty, FText& OutError) const;
+
+	FName GetMemberName() const;
+	FText GetMemberDisplayText() const;
+
+private:
+
+	/** The owner of the path segment (ie. What class or structure was this property from) */
+	UPROPERTY()
+	UStruct* Struct;
+
+	/** The member name in the structure this segment represents. */
+	UPROPERTY()
+	FName MemberName;
+
+	/**
+	 * The member guid in this structure this segment represents.  If this is valid it should 
+	 * be used instead of Name to get the true name.
+	 */
+	UPROPERTY()
+	FGuid MemberGuid;
+
+	/** true if property, false if function */
+	UPROPERTY()
+	bool IsProperty;
+};
+
+
+/**  */
+USTRUCT()
+struct UMGEDITOR_API FEditorPropertyPath
+{
+	GENERATED_USTRUCT_BODY()
+
+public:
+
+	FEditorPropertyPath();
+	FEditorPropertyPath(const TArray<UField*>& BindingChain);
+
+	bool IsEmpty() const { return Segments.Num() == 0; }
+	bool Validate(UDelegateProperty* Destination, FText& OutError) const;
+
+	FText GetDisplayText() const;
+
+	FDynamicPropertyPath ToPropertyPath() const;
+
+	/** The path of properties. */
+	UPROPERTY()
+	TArray<FEditorPropertyPathSegment> Segments;
+};
+
+
+/** */
 USTRUCT()
 struct UMGEDITOR_API FDelegateEditorBinding
 {
@@ -23,10 +92,6 @@ struct UMGEDITOR_API FDelegateEditorBinding
 	UPROPERTY()
 	FName PropertyName;
 
-	/** The optional source object we are bindings to, this could be a potential model object on the widget. */
-	UPROPERTY()
-	FName DataSource;
-
 	/** The function that was generated to return the SourceProperty */
 	UPROPERTY()
 	FName FunctionName;
@@ -35,12 +100,16 @@ struct UMGEDITOR_API FDelegateEditorBinding
 	UPROPERTY()
 	FName SourceProperty;
 
+	/**  */
+	UPROPERTY()
+	FEditorPropertyPath SourcePath;
+
 	/** If it's an actual Function Graph in the blueprint that we're bound to, there's a GUID we can use to lookup that function, to deal with renames better.  This is that GUID. */
 	UPROPERTY()
 	FGuid MemberGuid;
 
 	UPROPERTY()
-	TEnumAsByte<EBindingKind::Type> Kind;
+	EBindingKind Kind;
 
 	bool operator==( const FDelegateEditorBinding& Other ) const
 	{
diff --git a/Engine/Source/Runtime/UMG/Private/Binding/BoolBinding.cpp b/Engine/Source/Runtime/UMG/Private/Binding/BoolBinding.cpp
new file mode 100644
--- /dev/null
+++ b/Engine/Source/Runtime/UMG/Private/Binding/BoolBinding.cpp
@@ -0,0 +1,36 @@
+// Copyright 1998-2014 Epic Games, Inc. All Rights Reserved.
+
+#include "UMGPrivatePCH.h"
+#include "BoolBinding.h"
+
+#define LOCTEXT_NAMESPACE "UMG"
+
+UBoolBinding::UBoolBinding()
+{
+}
+
+bool UBoolBinding::IsSupportedSource(UProperty* Property) const
+{
+	return IsConcreteTypeCompatibleWithReflectedType<bool>(Property);
+}
+
+bool UBoolBinding::IsSupportedDestination(UProperty* Property) const
+{
+	return IsConcreteTypeCompatibleWithReflectedType<bool>(Property);
+}
+
+bool UBoolBinding::GetValue() const
+{
+	if ( UObject* Source = SourceObject.Get() )
+	{
+		bool Value;
+		if ( SourcePath.GetValue<bool>(Source, Value) )
+		{
+			return Value;
+		}
+	}
+
+	return false;
+}
+
+#undef LOCTEXT_NAMESPACE
diff --git a/Engine/Source/Runtime/UMG/Private/Binding/BrushBinding.cpp b/Engine/Source/Runtime/UMG/Private/Binding/BrushBinding.cpp
new file mode 100644
--- /dev/null
+++ b/Engine/Source/Runtime/UMG/Private/Binding/BrushBinding.cpp
@@ -0,0 +1,61 @@
+// Copyright 1998-2014 Epic Games, Inc. All Rights Reserved.
+
+#include "UMGPrivatePCH.h"
+#include "BrushBinding.h"
+
+#define LOCTEXT_NAMESPACE "UMG"
+
+UBrushBinding::UBrushBinding()
+{
+}
+
+bool UBrushBinding::IsSupportedDestination(UProperty* Property) const
+{
+	return IsConcreteTypeCompatibleWithReflectedType<FSlateBrush>(Property);
+}
+
+bool UBrushBinding::IsSupportedSource(UProperty* Property) const
+{
+	if ( IsConcreteTypeCompatibleWithReflectedType<UObject*>(Property) )
+	{
+		if ( UObjectProperty* ObjectProperty = Cast<UObjectProperty>(Property) )
+		{
+			return ObjectProperty->PropertyClass->IsChildOf(UTexture2D::StaticClass());
+		}
+	}
+
+	return IsConcreteTypeCompatibleWithReflectedType<FSlateBrush>(Property);
+}
+
+FSlateBrush UBrushBinding::GetValue() const
+{
+	if ( UObject* Source = SourceObject.Get() )
+	{
+		if ( bConversion.Get(EConversion::None) == EConversion::None )
+		{
+			FSlateBrush Value;
+			if ( SourcePath.GetValue<FSlateBrush>(Source, Value) )
+			{
+				bConversion = EConversion::None;
+				return Value;
+			}
+		}
+
+		if ( bConversion.Get(EConversion::Texture) == EConversion::Texture )
+		{
+			UObject* Value;
+			if ( SourcePath.GetValue<UObject*>(Source, Value) )
+			{
+				if ( UTexture2D* Texture = Cast<UTexture2D>(Value) )
+				{
+					bConversion = EConversion::Texture;
+					return GetDefault<UWidgetBlueprintLibrary>()->MakeBrushFromTexture(Texture);
+				}
+			}
+		}
+	}
+
+	return FSlateNoResource();
+}
+
+#undef LOCTEXT_NAMESPACE
diff --git a/Engine/Source/Runtime/UMG/Private/Binding/CheckedStateBinding.cpp b/Engine/Source/Runtime/UMG/Private/Binding/CheckedStateBinding.cpp
new file mode 100644
--- /dev/null
+++ b/Engine/Source/Runtime/UMG/Private/Binding/CheckedStateBinding.cpp
@@ -0,0 +1,60 @@
+// Copyright 1998-2014 Epic Games, Inc. All Rights Reserved.
+
+#include "UMGPrivatePCH.h"
+#include "CheckedStateBinding.h"
+
+#define LOCTEXT_NAMESPACE "UMG"
+
+UCheckedStateBinding::UCheckedStateBinding()
+{
+}
+
+bool UCheckedStateBinding::IsSupportedSource(UProperty* Property) const
+{
+	return IsSupportedDestination(Property) || IsConcreteTypeCompatibleWithReflectedType<bool>(Property);
+}
+
+bool UCheckedStateBinding::IsSupportedDestination(UProperty* Property) const
+{
+	static const FName CheckBoxStateEnum(TEXT("ESlateCheckBoxState"));
+
+	if ( UByteProperty* ByteProperty = Cast<UByteProperty>(Property) )
+	{
+		if ( ByteProperty->IsEnum() )
+		{
+			return ByteProperty->Enum->GetFName() == CheckBoxStateEnum;
+		}
+	}
+
+	return false;
+}
+
+ESlateCheckBoxState::Type UCheckedStateBinding::GetValue() const
+{
+	if ( UObject* Source = SourceObject.Get() )
+	{
+		if ( bConversion.Get(EConversion::None) == EConversion::None )
+		{
+			uint8 Value;
+			if ( SourcePath.GetValue<uint8>(Source, Value) )
+			{
+				bConversion = EConversion::None;
+				return static_cast<ESlateCheckBoxState::Type>(Value);
+			}
+		}
+
+		if ( bConversion.Get(EConversion::Bool) == EConversion::Bool )
+		{
+			bool Value;
+			if ( SourcePath.GetValue<bool>(Source, Value) )
+			{
+				bConversion = EConversion::Bool;
+				return Value ? ESlateCheckBoxState::Checked : ESlateCheckBoxState::Unchecked;
+			}
+		}
+	}
+
+	return ESlateCheckBoxState::Unchecked;
+}
+
+#undef LOCTEXT_NAMESPACE
diff --git a/Engine/Source/Runtime/UMG/Private/Binding/ColorBinding.cpp b/Engine/Source/Runtime/UMG/Private/Binding/ColorBinding.cpp
new file mode 100644
--- /dev/null
+++ b/Engine/Source/Runtime/UMG/Private/Binding/ColorBinding.cpp
@@ -0,0 +1,94 @@
+// Copyright 1998-2014 Epic Games, Inc. All Rights Reserved.
+
+#include "UMGPrivatePCH.h"
+#include "ColorBinding.h"
+
+#define LOCTEXT_NAMESPACE "UMG"
+
+UColorBinding::UColorBinding()
+{
+}
+
+bool UColorBinding::IsSupportedDestination(UProperty* Property) const
+{
+	return IsSupportedSource(Property);
+}
+
+bool UColorBinding::IsSupportedSource(UProperty* Property) const
+{
+	return
+		IsConcreteTypeCompatibleWithReflectedType<FSlateColor>(Property) ||
+		IsConcreteTypeCompatibleWithReflectedType<FLinearColor>(Property);
+}
+
+void UColorBinding::Bind(UProperty* Property, FScriptDelegate* Delegate)
+{
+	if ( IsConcreteTypeCompatibleWithReflectedType<FSlateColor>(Property) )
+	{
+		static const FName BinderFunction(TEXT("GetSlateValue"));
+		Delegate->BindUFunction(this, BinderFunction);
+	}
+	else if(IsConcreteTypeCompatibleWithReflectedType<FLinearColor>(Property))
+	{
+		static const FName BinderFunction(TEXT("GetLinearValue"));
+		Delegate->BindUFunction(this, BinderFunction);
+	}
+}
+
+FSlateColor UColorBinding::GetSlateValue() const
+{
+	if ( UObject* Source = SourceObject.Get() )
+	{
+		if ( !bNeedsConversion.Get(false) )
+		{
+			FSlateColor SlateColor;
+			if ( SourcePath.GetValue<FSlateColor>(Source, SlateColor) )
+			{
+				bNeedsConversion = false;
+				return SlateColor;
+			}
+		}
+
+		if ( bNeedsConversion.Get(true) )
+		{
+			FLinearColor LinearValue;
+			if ( SourcePath.GetValue<FLinearColor>(Source, LinearValue) )
+			{
+				bNeedsConversion = true;
+				return FSlateColor(LinearValue);
+			}
+		}
+	}
+
+	return FSlateColor();
+}
+
+FLinearColor UColorBinding::GetLinearValue() const
+{
+	if ( UObject* Source = SourceObject.Get() )
+	{
+		if ( !bNeedsConversion.Get(false) )
+		{
+			FLinearColor LinearValue;
+			if ( SourcePath.GetValue<FLinearColor>(Source, LinearValue) )
+			{
+				bNeedsConversion = false;
+				return LinearValue;
+			}
+		}
+
+		if ( bNeedsConversion.Get(true) )
+		{
+			FSlateColor SlateColor;
+			if ( SourcePath.GetValue<FSlateColor>(Source, SlateColor) )
+			{
+				bNeedsConversion = true;
+				return SlateColor.GetSpecifiedColor();
+			}
+		}
+	}
+
+	return FLinearColor();
+}
+
+#undef LOCTEXT_NAMESPACE
diff --git a/Engine/Source/Runtime/UMG/Private/Binding/DynamicPropertyPath.cpp b/Engine/Source/Runtime/UMG/Private/Binding/DynamicPropertyPath.cpp
new file mode 100644
--- /dev/null
+++ b/Engine/Source/Runtime/UMG/Private/Binding/DynamicPropertyPath.cpp
@@ -0,0 +1,79 @@
+﻿// Copyright 1998-2014 Epic Games, Inc. All Rights Reserved.
+
+#include "UMGPrivatePCH.h"
+#include "DynamicPropertyPath.h"
+
+#define LOCTEXT_NAMESPACE "UMG"
+
+FPropertyPathSegment::FPropertyPathSegment()
+	: Name(NAME_None)
+	, Struct(nullptr)
+	, Field(nullptr)
+	, ArrayIndex(INDEX_NONE)
+{
+
+}
+
+FPropertyPathSegment::FPropertyPathSegment(FString SegmentName)
+	: Struct(nullptr)
+	, Field(nullptr)
+	, ArrayIndex(INDEX_NONE)
+{
+	// Parse the property name and (optional) array index
+	int32 ArrayPos = SegmentName.Find(TEXT("["));
+	if ( ArrayPos != INDEX_NONE )
+	{
+		FString IndexToken = SegmentName.RightChop(ArrayPos + 1).LeftChop(1);
+		ArrayIndex = FCString::Atoi(*IndexToken);
+
+		SegmentName = SegmentName.Left(ArrayPos);
+	}
+
+	Name = FName(*SegmentName);
+}
+
+UField* FPropertyPathSegment::Resolve(UStruct* InStruct) const
+{
+	if ( InStruct )
+	{
+		// only perform the find field work if the structure where this property would resolve to
+		// has changed.  If it hasn't changed, the just return the UProperty we found last time.
+		if ( InStruct != Struct )
+		{
+			Struct = InStruct;
+			Field = FindField<UField>(InStruct, Name);
+		}
+
+		return Field;
+	}
+
+	return nullptr;
+}
+
+FDynamicPropertyPath::FDynamicPropertyPath()
+{
+}
+
+FDynamicPropertyPath::FDynamicPropertyPath(FString Path)
+{
+	FString PropertyString;
+	while ( Path.Split(TEXT("."), &PropertyString, &Path) )
+	{
+		if ( !PropertyString.IsEmpty() )
+		{
+			Segments.Add(FPropertyPathSegment(PropertyString));
+		}
+	}
+
+	Segments.Add(FPropertyPathSegment(Path));
+}
+
+FDynamicPropertyPath::FDynamicPropertyPath(const TArray<FString>& PropertyChain)
+{
+	for ( const FString& Segment : PropertyChain )
+	{
+		Segments.Add(FPropertyPathSegment(Segment));
+	}
+}
+
+#undef LOCTEXT_NAMESPACE
diff --git a/Engine/Source/Runtime/UMG/Private/Binding/FloatBinding.cpp b/Engine/Source/Runtime/UMG/Private/Binding/FloatBinding.cpp
new file mode 100644
--- /dev/null
+++ b/Engine/Source/Runtime/UMG/Private/Binding/FloatBinding.cpp
@@ -0,0 +1,36 @@
+// Copyright 1998-2014 Epic Games, Inc. All Rights Reserved.
+
+#include "UMGPrivatePCH.h"
+#include "FloatBinding.h"
+
+#define LOCTEXT_NAMESPACE "UMG"
+
+UFloatBinding::UFloatBinding()
+{
+}
+
+bool UFloatBinding::IsSupportedDestination(UProperty* Property) const
+{
+	return IsSupportedSource(Property);
+}
+
+bool UFloatBinding::IsSupportedSource(UProperty* Property) const
+{
+	return IsConcreteTypeCompatibleWithReflectedType<float>(Property);
+}
+
+float UFloatBinding::GetValue() const
+{
+	if ( UObject* Source = SourceObject.Get() )
+	{
+		float Value;
+		if ( SourcePath.GetValue<float>(Source, Value) )
+		{
+			return Value;
+		}
+	}
+
+	return 0;
+}
+
+#undef LOCTEXT_NAMESPACE
diff --git a/Engine/Source/Runtime/UMG/Private/Binding/Int32Binding.cpp b/Engine/Source/Runtime/UMG/Private/Binding/Int32Binding.cpp
new file mode 100644
--- /dev/null
+++ b/Engine/Source/Runtime/UMG/Private/Binding/Int32Binding.cpp
@@ -0,0 +1,36 @@
+// Copyright 1998-2014 Epic Games, Inc. All Rights Reserved.
+
+#include "UMGPrivatePCH.h"
+#include "Int32Binding.h"
+
+#define LOCTEXT_NAMESPACE "UMG"
+
+UInt32Binding::UInt32Binding()
+{
+}
+
+bool UInt32Binding::IsSupportedDestination(UProperty* Property) const
+{
+	return IsSupportedSource(Property);
+}
+
+bool UInt32Binding::IsSupportedSource(UProperty* Property) const
+{
+	return IsConcreteTypeCompatibleWithReflectedType<int32>(Property);
+}
+
+int32 UInt32Binding::GetValue() const
+{
+	if ( UObject* Source = SourceObject.Get() )
+	{
+		int32 Value;
+		if ( SourcePath.GetValue<int32>(Source, Value) )
+		{
+			return Value;
+		}
+	}
+
+	return false;
+}
+
+#undef LOCTEXT_NAMESPACE
diff --git a/Engine/Source/Runtime/UMG/Private/Binding/MouseCursorBinding.cpp b/Engine/Source/Runtime/UMG/Private/Binding/MouseCursorBinding.cpp
new file mode 100644
--- /dev/null
+++ b/Engine/Source/Runtime/UMG/Private/Binding/MouseCursorBinding.cpp
@@ -0,0 +1,46 @@
+// Copyright 1998-2014 Epic Games, Inc. All Rights Reserved.
+
+#include "UMGPrivatePCH.h"
+#include "MouseCursorBinding.h"
+
+#define LOCTEXT_NAMESPACE "UMG"
+
+UMouseCursorBinding::UMouseCursorBinding()
+{
+}
+
+bool UMouseCursorBinding::IsSupportedSource(UProperty* Property) const
+{
+	return IsSupportedDestination(Property);
+}
+
+bool UMouseCursorBinding::IsSupportedDestination(UProperty* Property) const
+{
+	static const FName MouseCursorEnum(TEXT("EMouseCursor"));
+	
+	if ( UByteProperty* ByteProperty = Cast<UByteProperty>(Property) )
+	{
+		if ( ByteProperty->IsEnum() )
+		{
+			return ByteProperty->Enum->GetFName() == MouseCursorEnum;
+		}
+	}
+
+	return false;
+}
+
+EMouseCursor::Type UMouseCursorBinding::GetValue() const
+{
+	if ( UObject* Source = SourceObject.Get() )
+	{
+		uint8 Value;
+		if ( SourcePath.GetValue<uint8>(Source, Value) )
+		{
+			return static_cast<EMouseCursor::Type>( Value );
+		}
+	}
+
+	return EMouseCursor::Default;
+}
+
+#undef LOCTEXT_NAMESPACE
diff --git a/Engine/Source/Runtime/UMG/Private/Binding/PropertyBinding.cpp b/Engine/Source/Runtime/UMG/Private/Binding/PropertyBinding.cpp
new file mode 100644
--- /dev/null
+++ b/Engine/Source/Runtime/UMG/Private/Binding/PropertyBinding.cpp
@@ -0,0 +1,28 @@
+// Copyright 1998-2014 Epic Games, Inc. All Rights Reserved.
+
+#include "UMGPrivatePCH.h"
+#include "PropertyBinding.h"
+
+#define LOCTEXT_NAMESPACE "UMG"
+
+UPropertyBinding::UPropertyBinding()
+{
+}
+
+bool UPropertyBinding::IsSupportedSource(UProperty* Property) const
+{
+	return false;
+}
+
+bool UPropertyBinding::IsSupportedDestination(UProperty* Property) const
+{
+	return false;
+}
+
+void UPropertyBinding::Bind(UProperty* Property, FScriptDelegate* Delegate)
+{
+	static const FName BinderFunction(TEXT("GetValue"));
+	Delegate->BindUFunction(this, BinderFunction);
+}
+
+#undef LOCTEXT_NAMESPACE
diff --git a/Engine/Source/Runtime/UMG/Private/Binding/TextBinding.cpp b/Engine/Source/Runtime/UMG/Private/Binding/TextBinding.cpp
new file mode 100644
--- /dev/null
+++ b/Engine/Source/Runtime/UMG/Private/Binding/TextBinding.cpp
@@ -0,0 +1,94 @@
+// Copyright 1998-2014 Epic Games, Inc. All Rights Reserved.
+
+#include "UMGPrivatePCH.h"
+#include "TextBinding.h"
+
+#define LOCTEXT_NAMESPACE "UMG"
+
+UTextBinding::UTextBinding()
+{
+}
+
+bool UTextBinding::IsSupportedDestination(UProperty* Property) const
+{
+	return IsConcreteTypeCompatibleWithReflectedType<FText>(Property);
+}
+
+bool UTextBinding::IsSupportedSource(UProperty* Property) const
+{
+	return
+		IsConcreteTypeCompatibleWithReflectedType<FText>(Property) ||
+		IsConcreteTypeCompatibleWithReflectedType<FString>(Property);
+}
+
+void UTextBinding::Bind(UProperty* Property, FScriptDelegate* Delegate)
+{
+	if ( IsConcreteTypeCompatibleWithReflectedType<FText>(Property) )
+	{
+		static const FName BinderFunction(TEXT("GetTextValue"));
+		Delegate->BindUFunction(this, BinderFunction);
+	}
+	else if ( IsConcreteTypeCompatibleWithReflectedType<FString>(Property) )
+	{
+		static const FName BinderFunction(TEXT("GetStringValue"));
+		Delegate->BindUFunction(this, BinderFunction);
+	}
+}
+
+FText UTextBinding::GetTextValue() const
+{
+	if ( UObject* Source = SourceObject.Get() )
+	{
+		if ( !bNeedsConversion.Get(false) )
+		{
+			FText TextValue;
+			if ( SourcePath.GetValue<FText>(Source, TextValue) )
+			{
+				bNeedsConversion = false;
+				return TextValue;
+			}
+		}
+
+		if ( bNeedsConversion.Get(true) )
+		{
+			FString StringValue;
+			if ( SourcePath.GetValue<FString>(Source, StringValue) )
+			{
+				bNeedsConversion = true;
+				return FText::FromString(StringValue);
+			}
+		}
+	}
+
+	return FText::GetEmpty();
+}
+
+FString UTextBinding::GetStringValue() const
+{
+	if ( UObject* Source = SourceObject.Get() )
+	{
+		if ( !bNeedsConversion.Get(false) )
+		{
+			FString StringValue;
+			if ( SourcePath.GetValue<FString>(Source, StringValue) )
+			{
+				bNeedsConversion = false;
+				return StringValue;
+			}
+		}
+
+		if ( bNeedsConversion.Get(true) )
+		{
+			FText TextValue;
+			if ( SourcePath.GetValue<FText>(Source, TextValue) )
+			{
+				bNeedsConversion = true;
+				return TextValue.ToString();
+			}
+		}
+	}
+
+	return FString();
+}
+
+#undef LOCTEXT_NAMESPACE
diff --git a/Engine/Source/Runtime/UMG/Private/Binding/VisibilityBinding.cpp b/Engine/Source/Runtime/UMG/Private/Binding/VisibilityBinding.cpp
new file mode 100644
--- /dev/null
+++ b/Engine/Source/Runtime/UMG/Private/Binding/VisibilityBinding.cpp
@@ -0,0 +1,46 @@
+// Copyright 1998-2014 Epic Games, Inc. All Rights Reserved.
+
+#include "UMGPrivatePCH.h"
+#include "VisibilityBinding.h"
+
+#define LOCTEXT_NAMESPACE "UMG"
+
+UVisibilityBinding::UVisibilityBinding()
+{
+}
+
+bool UVisibilityBinding::IsSupportedSource(UProperty* Property) const
+{
+	return IsSupportedDestination(Property);
+}
+
+bool UVisibilityBinding::IsSupportedDestination(UProperty* Property) const
+{
+	static const FName VisibilityEnum(TEXT("ESlateVisibility"));
+
+	if ( UByteProperty* ByteProperty = Cast<UByteProperty>(Property) )
+	{
+		if ( ByteProperty->IsEnum() )
+		{
+			return ByteProperty->Enum->GetFName() == VisibilityEnum;
+		}
+	}
+
+	return false;
+}
+
+ESlateVisibility UVisibilityBinding::GetValue() const
+{
+	if ( UObject* Source = SourceObject.Get() )
+	{
+		uint8 Value;
+		if ( SourcePath.GetValue<uint8>(Source, Value) )
+		{
+			return static_cast<ESlateVisibility>( Value );
+		}
+	}
+
+	return ESlateVisibility::Visible;
+}
+
+#undef LOCTEXT_NAMESPACE
diff --git a/Engine/Source/Runtime/UMG/Private/Components/TextBlock.cpp b/Engine/Source/Runtime/UMG/Private/Components/TextBlock.cpp
--- a/Engine/Source/Runtime/UMG/Private/Components/TextBlock.cpp
+++ b/Engine/Source/Runtime/UMG/Private/Components/TextBlock.cpp
@@ -64,6 +64,34 @@ TSharedRef<SWidget> UTextBlock::RebuildWidget()
 	return MyTextBlock.ToSharedRef();
 }
 
+void UTextBlock::OnBindingChanged(const FName& Property)
+{
+	Super::OnBindingChanged(Property);
+
+	if ( MyTextBlock.IsValid() )
+	{
+		static const FName TextProperty(TEXT("TextDelegate"));
+		static const FName ColorAndOpacityProperty(TEXT("ColorAndOpacityDelegate"));
+		static const FName ShadowColorAndOpacityProperty(TEXT("ShadowColorAndOpacityDelegate"));
+
+		if ( Property == TextProperty )
+		{
+			TAttribute<FText> TextBinding = OPTIONAL_BINDING(FText, Text);
+			MyTextBlock->SetText(TextBinding);
+		}
+		else if ( Property == ColorAndOpacityProperty )
+		{
+			TAttribute<FSlateColor> ColorAndOpacityBinding = OPTIONAL_BINDING(FSlateColor, ColorAndOpacity);
+			MyTextBlock->SetColorAndOpacity(ColorAndOpacityBinding);
+		}
+		else if ( Property == ShadowColorAndOpacityProperty )
+		{
+			TAttribute<FLinearColor> ShadowColorAndOpacityBinding = OPTIONAL_BINDING(FLinearColor, ShadowColorAndOpacity);
+			MyTextBlock->SetShadowColorAndOpacity(ShadowColorAndOpacityBinding);
+		}
+	}
+}
+
 void UTextBlock::SynchronizeProperties()
 {
 	Super::SynchronizeProperties();
diff --git a/Engine/Source/Runtime/UMG/Private/Components/Widget.cpp b/Engine/Source/Runtime/UMG/Private/Components/Widget.cpp
--- a/Engine/Source/Runtime/UMG/Private/Components/Widget.cpp
+++ b/Engine/Source/Runtime/UMG/Private/Components/Widget.cpp
@@ -4,9 +4,18 @@
 
 #include "ReflectionMetadata.h"
 
+#include "TextBinding.h"
+#include "FloatBinding.h"
+#include "BoolBinding.h"
+#include "BrushBinding.h"
+#include "ColorBinding.h"
+#include "Int32Binding.h"
+
 /////////////////////////////////////////////////////
 // UWidget
 
+TArray<TSubclassOf<UPropertyBinding>> UWidget::BinderClasses;
+
 UWidget::UWidget(const FObjectInitializer& ObjectInitializer)
 	: Super(ObjectInitializer)
 {
@@ -615,3 +624,96 @@ UWidget* UWidget::FindChildContainingDescendant(UWidget* Root, UWidget* Descenda
 
 	return nullptr;
 }
+
+//bool UWidget::BindProperty(const FName& DestinationProperty, UObject* SourceObject, const FName& SourceProperty)
+//{
+//	UDelegateProperty* DelegateProperty = FindField<UDelegateProperty>(GetClass(), FName(*( DestinationProperty.ToString() + TEXT("Delegate") )));
+//
+//	if ( DelegateProperty )
+//	{
+//		FDynamicPropertyPath BindingPath(SourceProperty.ToString());
+//		return AddBinding(DelegateProperty, SourceObject, BindingPath);
+//	}
+//
+//	return false;
+//}
+
+TSubclassOf<UPropertyBinding> UWidget::FindBinderClassForDestination(UProperty* Property)
+{
+	if ( BinderClasses.Num() == 0 )
+	{
+		for ( TObjectIterator<UClass> ClassIt; ClassIt; ++ClassIt )
+		{
+			if ( ClassIt->IsChildOf(UPropertyBinding::StaticClass()) )
+			{
+				BinderClasses.Add(*ClassIt);
+			}
+		}
+	}
+
+	for ( int32 ClassIndex = 0; ClassIndex < BinderClasses.Num(); ClassIndex++ )
+	{
+		if ( GetDefault<UPropertyBinding>(BinderClasses[ClassIndex])->IsSupportedDestination(Property))
+		{
+			return BinderClasses[ClassIndex];
+		}
+	}
+
+	return nullptr;
+}
+
+static UPropertyBinding* GenerateBinder(UDelegateProperty* DelegateProperty, UObject* Container, UObject* SourceObject, const FDynamicPropertyPath& BindingPath)
+{
+	FScriptDelegate* ScriptDelegate = DelegateProperty->GetPropertyValuePtr_InContainer(Container);
+	if ( ScriptDelegate )
+	{
+		if ( UProperty* ReturnProperty = DelegateProperty->SignatureFunction->GetReturnProperty() )
+		{
+			TSubclassOf<UPropertyBinding> BinderClass = UWidget::FindBinderClassForDestination(ReturnProperty);
+			if ( BinderClass != nullptr )
+			{
+				UPropertyBinding* Binder = ConstructObject<UPropertyBinding>(BinderClass, Container);
+				Binder->SourceObject = SourceObject;
+				Binder->SourcePath = BindingPath;
+				Binder->Bind(ReturnProperty, ScriptDelegate);
+
+				return Binder;
+			}
+		}
+	}
+
+	return nullptr;
+}
+
+bool UWidget::AddBinding(UDelegateProperty* DelegateProperty, UObject* SourceObject, const FDynamicPropertyPath& BindingPath)
+{
+	if ( UPropertyBinding* Binder = GenerateBinder(DelegateProperty, this, SourceObject, BindingPath) )
+	{
+		// Remove any existing binding object for this property.
+		for ( int32 BindingIndex = 0; BindingIndex < NativeBindings.Num(); BindingIndex++ )
+		{
+			if ( NativeBindings[BindingIndex]->DestinationProperty == DelegateProperty->GetFName() )
+			{
+				NativeBindings.RemoveAt(BindingIndex);
+				break;
+			}
+		}
+
+		NativeBindings.Add(Binder);
+
+		// Only notify the bindings have changed if we've already create the underlying slate widget.
+		if ( MyWidget.IsValid() )
+		{
+			OnBindingChanged(DelegateProperty->GetFName());
+		}
+
+		return true;
+	}
+
+	return false;
+}
+
+void UWidget::OnBindingChanged(const FName& Property)
+{
+
+}
diff --git a/Engine/Source/Runtime/UMG/Private/UserWidget.cpp b/Engine/Source/Runtime/UMG/Private/UserWidget.cpp
--- a/Engine/Source/Runtime/UMG/Private/UserWidget.cpp
+++ b/Engine/Source/Runtime/UMG/Private/UserWidget.cpp
@@ -97,6 +97,8 @@ void UUserWidget::SynchronizeProperties()
 {
 	Super::SynchronizeProperties();
 
+	// We get the GCWidget directly because MyWidget could be the fullscreen host widget if we've been added
+	// to the viewport.
 	TSharedPtr<SObjectWidget> SafeGCWidget = MyGCWidget.Pin();
 	if ( SafeGCWidget.IsValid() )
 	{
@@ -747,11 +749,6 @@ void UUserWidget::PreSave()
 	}
 }
 
-void UUserWidget::PostLoad()
-{
-	Super::PostLoad();
-}
-
 #if WITH_EDITOR
 
 const FSlateBrush* UUserWidget::GetEditorIcon()
diff --git a/Engine/Source/Runtime/UMG/Private/WidgetBlueprintGeneratedClass.cpp b/Engine/Source/Runtime/UMG/Private/WidgetBlueprintGeneratedClass.cpp
--- a/Engine/Source/Runtime/UMG/Private/WidgetBlueprintGeneratedClass.cpp
+++ b/Engine/Source/Runtime/UMG/Private/WidgetBlueprintGeneratedClass.cpp
@@ -85,10 +85,20 @@ void UWidgetBlueprintGeneratedClass::InitializeWidget(UUserWidget* UserWidget) c
 
 					if ( DelegateProperty )
 					{
-						FScriptDelegate* ScriptDelegate = DelegateProperty->GetPropertyValuePtr_InContainer(Widget);
-						if ( ScriptDelegate )
+						bool bSourcePathBound = false;
+
+						if ( Binding.SourcePath.IsValid() )
+						{
+							bSourcePathBound = Widget->AddBinding(DelegateProperty, UserWidget, Binding.SourcePath);
+						}
+
+						if ( bSourcePathBound == false )
 						{
-							ScriptDelegate->BindUFunction(UserWidget, Binding.FunctionName);
+							FScriptDelegate* ScriptDelegate = DelegateProperty->GetPropertyValuePtr_InContainer(Widget);
+							if ( ScriptDelegate )
+							{
+								ScriptDelegate->BindUFunction(UserWidget, Binding.FunctionName);
+							}
 						}
 					}
 				}
diff --git a/Engine/Source/Runtime/UMG/Public/Binding/BoolBinding.h b/Engine/Source/Runtime/UMG/Public/Binding/BoolBinding.h
new file mode 100644
--- /dev/null
+++ b/Engine/Source/Runtime/UMG/Public/Binding/BoolBinding.h
@@ -0,0 +1,21 @@
+﻿// Copyright 1998-2014 Epic Games, Inc. All Rights Reserved.
+
+#pragma once
+
+#include "BoolBinding.generated.h"
+
+UCLASS()
+class UMG_API UBoolBinding : public UPropertyBinding
+{
+	GENERATED_BODY()
+
+public:
+
+	UBoolBinding();
+
+	virtual bool IsSupportedSource(UProperty* Property) const override;
+	virtual bool IsSupportedDestination(UProperty* Property) const override;
+
+	UFUNCTION()
+	bool GetValue() const;
+};
diff --git a/Engine/Source/Runtime/UMG/Public/Binding/BrushBinding.h b/Engine/Source/Runtime/UMG/Public/Binding/BrushBinding.h
new file mode 100644
--- /dev/null
+++ b/Engine/Source/Runtime/UMG/Public/Binding/BrushBinding.h
@@ -0,0 +1,31 @@
+﻿// Copyright 1998-2014 Epic Games, Inc. All Rights Reserved.
+
+#pragma once
+
+#include "BrushBinding.generated.h"
+
+UCLASS()
+class UMG_API UBrushBinding : public UPropertyBinding
+{
+	GENERATED_BODY()
+
+public:
+
+	UBrushBinding();
+
+	virtual bool IsSupportedSource(UProperty* Property) const override;
+	virtual bool IsSupportedDestination(UProperty* Property) const override;
+
+	UFUNCTION()
+	FSlateBrush GetValue() const;
+
+private:
+	enum class EConversion : uint8
+	{
+		None,
+		Texture,
+		//Material,
+	};
+
+	mutable TOptional<EConversion> bConversion;
+};
diff --git a/Engine/Source/Runtime/UMG/Public/Binding/CheckedStateBinding.h b/Engine/Source/Runtime/UMG/Public/Binding/CheckedStateBinding.h
new file mode 100644
--- /dev/null
+++ b/Engine/Source/Runtime/UMG/Public/Binding/CheckedStateBinding.h
@@ -0,0 +1,30 @@
+﻿// Copyright 1998-2014 Epic Games, Inc. All Rights Reserved.
+
+#pragma once
+
+#include "CheckedStateBinding.generated.h"
+
+UCLASS()
+class UMG_API UCheckedStateBinding : public UPropertyBinding
+{
+	GENERATED_BODY()
+
+public:
+
+	UCheckedStateBinding();
+
+	virtual bool IsSupportedSource(UProperty* Property) const override;
+	virtual bool IsSupportedDestination(UProperty* Property) const override;
+
+	UFUNCTION()
+	ESlateCheckBoxState::Type GetValue() const;
+
+private:
+	enum class EConversion : uint8
+	{
+		None,
+		Bool
+	};
+
+	mutable TOptional<EConversion> bConversion;
+};
diff --git a/Engine/Source/Runtime/UMG/Public/Binding/ColorBinding.h b/Engine/Source/Runtime/UMG/Public/Binding/ColorBinding.h
new file mode 100644
--- /dev/null
+++ b/Engine/Source/Runtime/UMG/Public/Binding/ColorBinding.h
@@ -0,0 +1,29 @@
+﻿// Copyright 1998-2014 Epic Games, Inc. All Rights Reserved.
+
+#pragma once
+
+#include "ColorBinding.generated.h"
+
+UCLASS()
+class UMG_API UColorBinding : public UPropertyBinding
+{
+	GENERATED_BODY()
+
+public:
+
+	UColorBinding();
+
+	virtual bool IsSupportedSource(UProperty* Property) const override;
+	virtual bool IsSupportedDestination(UProperty* Property) const override;
+
+	virtual void Bind(UProperty* Property, FScriptDelegate* Delegate) override;
+
+	UFUNCTION()
+	FSlateColor GetSlateValue() const;
+
+	UFUNCTION()
+	FLinearColor GetLinearValue() const;
+
+private:
+	mutable TOptional<bool> bNeedsConversion;
+};
diff --git a/Engine/Source/Runtime/UMG/Public/Binding/DynamicPropertyPath.h b/Engine/Source/Runtime/UMG/Public/Binding/DynamicPropertyPath.h
new file mode 100644
--- /dev/null
+++ b/Engine/Source/Runtime/UMG/Public/Binding/DynamicPropertyPath.h
@@ -0,0 +1,338 @@
+﻿// Copyright 1998-2014 Epic Games, Inc. All Rights Reserved.
+
+#pragma once
+
+#include "DynamicPropertyPath.generated.h"
+
+/** */
+template<typename T>
+inline bool IsConcreteTypeCompatibleWithReflectedType(UProperty* Property)
+{
+	if ( UStructProperty* StructProperty = Cast<UStructProperty>(Property) )
+	{
+		return StructProperty->Struct == T::StaticStruct();
+	}
+
+	return false;
+}
+
+template<>
+inline bool IsConcreteTypeCompatibleWithReflectedType<bool>(UProperty* Property)
+{
+	return Property->GetClass() == UBoolProperty::StaticClass();
+}
+
+template<>
+inline bool IsConcreteTypeCompatibleWithReflectedType<int8>(UProperty* Property)
+{
+	return Property->GetClass() == UInt8Property::StaticClass();
+}
+
+template<>
+inline bool IsConcreteTypeCompatibleWithReflectedType<uint8>(UProperty* Property)
+{
+	return Property->GetClass() == UByteProperty::StaticClass();
+}
+
+template<>
+inline bool IsConcreteTypeCompatibleWithReflectedType<int16>(UProperty* Property)
+{
+	return Property->GetClass() == UInt16Property::StaticClass();
+}
+
+template<>
+inline bool IsConcreteTypeCompatibleWithReflectedType<uint16>(UProperty* Property)
+{
+	return Property->GetClass() == UUInt16Property::StaticClass();
+}
+
+template<>
+inline bool IsConcreteTypeCompatibleWithReflectedType<int32>(UProperty* Property)
+{
+	return Property->GetClass() == UIntProperty::StaticClass();
+}
+
+template<>
+inline bool IsConcreteTypeCompatibleWithReflectedType<uint32>(UProperty* Property)
+{
+	return Property->GetClass() == UUInt32Property::StaticClass();
+}
+
+template<>
+inline bool IsConcreteTypeCompatibleWithReflectedType<int64>(UProperty* Property)
+{
+	return Property->GetClass() == UInt64Property::StaticClass();
+}
+
+template<>
+inline bool IsConcreteTypeCompatibleWithReflectedType<uint64>(UProperty* Property)
+{
+	return Property->GetClass() == UUInt64Property::StaticClass();
+}
+
+template<>
+inline bool IsConcreteTypeCompatibleWithReflectedType<float>(UProperty* Property)
+{
+	return Property->GetClass() == UFloatProperty::StaticClass();
+}
+
+template<>
+inline bool IsConcreteTypeCompatibleWithReflectedType<double>(UProperty* Property)
+{
+	return Property->GetClass() == UDoubleProperty::StaticClass();
+}
+
+template<>
+inline bool IsConcreteTypeCompatibleWithReflectedType<FText>(UProperty* Property)
+{
+	return Property->GetClass() == UTextProperty::StaticClass();
+}
+
+template<>
+inline bool IsConcreteTypeCompatibleWithReflectedType<FString>(UProperty* Property)
+{
+	return Property->GetClass() == UStrProperty::StaticClass();
+}
+
+template<>
+inline bool IsConcreteTypeCompatibleWithReflectedType<FLinearColor>(UProperty* Property)
+{
+	static const UScriptStruct* LinearColorStruct = FindObjectChecked<UScriptStruct>(UObject::StaticClass(), TEXT("LinearColor"));
+	
+	if ( UStructProperty* StructProperty = Cast<UStructProperty>(Property) )
+	{
+		return StructProperty->Struct == LinearColorStruct;
+	}
+
+	return false;
+}
+
+template<>
+inline bool IsConcreteTypeCompatibleWithReflectedType<UObject*>(UProperty* Property)
+{
+	if ( UObjectProperty* ObjectProperty = Cast<UObjectProperty>(Property) )
+	{
+		return true;
+		//return ObjectProperty->PropertyClass->IsChildOf(T::StaticClass());
+	}
+
+	return false;
+}
+
+/** A struct used for caching part of a property path.  Don't use this class directly. */
+USTRUCT()
+struct UMG_API FPropertyPathSegment
+{
+	GENERATED_USTRUCT_BODY()
+
+public:
+	/** Implementation detail, don't use this constructor. */
+	FPropertyPathSegment();
+
+	/** Initializes the segment for particular name. */
+	FPropertyPathSegment(FString SegmentName);
+
+	/**
+	 * Resolves the name on the given Struct.  Caches the resulting property so that future calls can be processed quickly.
+	 * @param InStruct the ScriptStruct or Class to look for the property on.
+	 */
+	UField* Resolve(UStruct* InStruct) const;
+
+public:
+
+	/** The sub-component of the property path, a single value between .'s of the path */
+	UPROPERTY()
+	FName Name;
+
+	/** The optional array index. */
+	UPROPERTY()
+	int32 ArrayIndex;
+
+private:
+
+	/** The cached Class or ScriptStruct that was used last to resolve Name to a property. */
+	UPROPERTY(Transient)
+	mutable UStruct* Struct;
+
+	/**
+	 * The cached property on the Struct that this Name resolved to on it last time Resolve was called, if 
+	 * the Struct doesn't change, this value is returned to avoid performing another Field lookup.
+	 */
+	UPROPERTY(Transient)
+	mutable UField* Field;
+};
+
+/** */
+USTRUCT()
+struct UMG_API FDynamicPropertyPath
+{
+	GENERATED_USTRUCT_BODY()
+
+public:
+
+	/** */
+	FDynamicPropertyPath();
+
+	/** */
+	FDynamicPropertyPath(FString Path);
+
+	/** */
+	FDynamicPropertyPath(const TArray<FString>& PropertyChain);
+
+	/** */
+	bool IsValid() const { return Segments.Num() > 0; }
+
+	/** */
+	template<typename T>
+	bool GetValue(UObject* InContainer, T& OutValue) const
+	{
+		UProperty* OutProperty;
+		return GetValue<T>(InContainer, OutValue, OutProperty);
+	}
+
+	/** */
+	template<typename T>
+	bool GetValue(UObject* InContainer, T& OutValue, UProperty*& OutProperty) const
+	{
+		if ( InContainer && IsValid() )
+		{
+			return GetValueRecursive(InContainer->GetClass(), InContainer, INDEX_NONE, 0, OutValue, OutProperty);
+		}
+
+		return false;
+	}
+
+private:
+
+	/** */
+	template<typename T>
+	bool GetValueRecursive(UStruct* InStruct, void* InContainer, int32 ArrayIndex, int32 SegmentIndex, T& OutValue, UProperty*& OutProperty) const
+	{
+		const FPropertyPathSegment& Segment = Segments[SegmentIndex];
+
+		// Obtain the property info from the given structure definition
+		if ( UField* Field = Segment.Resolve(InStruct) )
+		{
+			if ( UProperty* Property = Cast<UProperty>(Field) )
+			{
+				if ( SegmentIndex < ( Segments.Num() - 1 ) )
+				{
+					// Check first to see if this is a simple object (i.e. not an array of objects)
+					if ( UObjectProperty* ObjectProperty = Cast<UObjectProperty>(Property) )
+					{
+						// If it's an object we need to get the value of the property in the container first before we 
+						// can continue, if the object is null we safely stop processing the chain of properties.
+						if ( UObject* CurrentObject = ObjectProperty->GetPropertyValue_InContainer(InContainer) )
+						{
+							return GetValueRecursive(CurrentObject->GetClass(), CurrentObject, ArrayIndex, SegmentIndex + 1, OutValue, OutProperty);
+						}
+					}
+					// Check to see if this is a simple structure (i.e. not an array of structures)
+					else if ( UStructProperty* StructProp = Cast<UStructProperty>(Property) )
+					{
+						// Recursively call back into this function with the structure property and container value
+						return GetValueRecursive(StructProp->Struct, StructProp->ContainerPtrToValuePtr<void>(InContainer), ArrayIndex, SegmentIndex + 1, OutValue, OutProperty);
+					}
+					else if ( UArrayProperty* ArrayProp = Cast<UArrayProperty>(Property) )
+					{
+						// It is an array, now check to see if this is an array of structures
+						if ( UStructProperty* StructProp = Cast<UStructProperty>(ArrayProp->Inner) )
+						{
+							FScriptArrayHelper_InContainer ArrayHelper(ArrayProp, InContainer);
+							if ( ArrayHelper.IsValidIndex(Segment.ArrayIndex) )
+							{
+								// Recursively call back into this function with the array element and container value
+								return GetValueRecursive(StructProp->Struct, ArrayHelper.GetRawPtr(Segment.ArrayIndex), ArrayIndex, SegmentIndex + 1, OutValue, OutProperty);
+							}
+						}
+						// if it's not an array of structs, maybe it's an array of classes
+						else if ( UObjectProperty* ObjectProperty = Cast<UObjectProperty>(ArrayProp->Inner) )
+						{
+							//TODO Add support for arrays of objects.
+						}
+					}
+				}
+				else
+				{
+					// We're on the final property in the path, it may be an array property, so check that first.
+					if ( UArrayProperty* ArrayProp = Cast<UArrayProperty>(Property) )
+					{
+						// if it's an array property, we need to see if we parsed an index as part of the segment
+						// as a user may have baked the index directly into the property path.
+						int32 Index = ArrayIndex != INDEX_NONE ? ArrayIndex : Segment.ArrayIndex;
+						ensure(Index != INDEX_NONE);
+
+						// Property is an array property, so make sure we have a valid index specified
+						FScriptArrayHelper_InContainer ArrayHelper(ArrayProp, InContainer);
+						if ( ArrayHelper.IsValidIndex(Index) )
+						{
+							// Verify that the cpp type matches a known property type.
+							if ( IsConcreteTypeCompatibleWithReflectedType<T>(ArrayProp->Inner) )
+							{
+								// Ensure that the element sizes are the same, prevents the user from doing something terribly wrong.
+								if ( ArrayProp->Inner->ElementSize == sizeof(T) )
+								{
+									OutValue = *static_cast<T*>( static_cast<void*>( ArrayHelper.GetRawPtr(Index) ) );
+									OutProperty = ArrayProp->Inner;
+								}
+							}
+						}
+					}
+					else
+					{
+						// Verify that the cpp type matches a known property type.
+						if ( IsConcreteTypeCompatibleWithReflectedType<T>(Property) )
+						{
+							// Ensure that the element sizes are the same, prevents the user from doing something terribly wrong.
+							if ( Property->ElementSize == sizeof(T) )
+							{
+								// Property is a vector property, so access directly
+								if ( T* ValuePtr = Property->ContainerPtrToValuePtr<T>(InContainer) )
+								{
+									OutValue = *ValuePtr;
+									OutProperty = Property;
+									return true;
+								}
+							}
+						}
+					}
+				}
+			}
+			else
+			{
+				// Only allow functions as the final link in the chain.
+				if ( SegmentIndex == ( Segments.Num() - 1 ) )
+				{
+					UFunction* Function = CastChecked<UFunction>(Field);
+					UObject* ContainerObject = static_cast<UObject*>( InContainer );
+
+					// We only support calling functions that return a single value and take no parameters.
+					if ( Function->NumParms == 1 )
+					{
+						// Verify there's a return property.
+						if ( UProperty* ReturnProperty = Function->GetReturnProperty() )
+						{
+							// Verify that the cpp type matches a known property type.
+							if ( IsConcreteTypeCompatibleWithReflectedType<T>(ReturnProperty) )
+							{
+								// Ensure that the element sizes are the same, prevents the user from doing something terribly wrong.
+								if ( ReturnProperty->ElementSize == sizeof(T) )
+								{
+									ContainerObject->ProcessEvent(Function, &OutValue);
+									return true;
+								}
+							}
+						}
+					}
+				}
+			}
+		}
+
+		return false;
+	}
+
+private:
+
+	UPROPERTY()
+	TArray<FPropertyPathSegment> Segments;
+};
diff --git a/Engine/Source/Runtime/UMG/Public/Binding/FloatBinding.h b/Engine/Source/Runtime/UMG/Public/Binding/FloatBinding.h
new file mode 100644
--- /dev/null
+++ b/Engine/Source/Runtime/UMG/Public/Binding/FloatBinding.h
@@ -0,0 +1,21 @@
+﻿// Copyright 1998-2014 Epic Games, Inc. All Rights Reserved.
+
+#pragma once
+
+#include "FloatBinding.generated.h"
+
+UCLASS()
+class UMG_API UFloatBinding : public UPropertyBinding
+{
+	GENERATED_BODY()
+
+public:
+
+	UFloatBinding();
+
+	virtual bool IsSupportedSource(UProperty* Property) const override;
+	virtual bool IsSupportedDestination(UProperty* Property) const override;
+
+	UFUNCTION()
+	float GetValue() const;
+};
diff --git a/Engine/Source/Runtime/UMG/Public/Binding/Int32Binding.h b/Engine/Source/Runtime/UMG/Public/Binding/Int32Binding.h
new file mode 100644
--- /dev/null
+++ b/Engine/Source/Runtime/UMG/Public/Binding/Int32Binding.h
@@ -0,0 +1,21 @@
+﻿// Copyright 1998-2014 Epic Games, Inc. All Rights Reserved.
+
+#pragma once
+
+#include "Int32Binding.generated.h"
+
+UCLASS()
+class UMG_API UInt32Binding : public UPropertyBinding
+{
+	GENERATED_BODY()
+
+public:
+
+	UInt32Binding();
+
+	virtual bool IsSupportedSource(UProperty* Property) const override;
+	virtual bool IsSupportedDestination(UProperty* Property) const override;
+
+	UFUNCTION()
+	int32 GetValue() const;
+};
diff --git a/Engine/Source/Runtime/UMG/Public/Binding/MouseCursorBinding.h b/Engine/Source/Runtime/UMG/Public/Binding/MouseCursorBinding.h
new file mode 100644
--- /dev/null
+++ b/Engine/Source/Runtime/UMG/Public/Binding/MouseCursorBinding.h
@@ -0,0 +1,21 @@
+﻿// Copyright 1998-2014 Epic Games, Inc. All Rights Reserved.
+
+#pragma once
+
+#include "MouseCursorBinding.generated.h"
+
+UCLASS()
+class UMG_API UMouseCursorBinding : public UPropertyBinding
+{
+	GENERATED_BODY()
+
+public:
+
+	UMouseCursorBinding();
+
+	virtual bool IsSupportedSource(UProperty* Property) const override;
+	virtual bool IsSupportedDestination(UProperty* Property) const override;
+
+	UFUNCTION()
+	EMouseCursor::Type GetValue() const;
+};
diff --git a/Engine/Source/Runtime/UMG/Public/Binding/PropertyBinding.h b/Engine/Source/Runtime/UMG/Public/Binding/PropertyBinding.h
new file mode 100644
--- /dev/null
+++ b/Engine/Source/Runtime/UMG/Public/Binding/PropertyBinding.h
@@ -0,0 +1,34 @@
+﻿// Copyright 1998-2014 Epic Games, Inc. All Rights Reserved.
+
+#pragma once
+
+#include "DynamicPropertyPath.h"
+
+#include "PropertyBinding.generated.h"
+
+UCLASS()
+class UMG_API UPropertyBinding : public UObject
+{
+	GENERATED_BODY()
+
+public:
+	UPropertyBinding();
+
+	virtual bool IsSupportedSource(UProperty* Property) const;
+	virtual bool IsSupportedDestination(UProperty* Property) const;
+
+	virtual void Bind(UProperty* Property, FScriptDelegate* Delegate);
+
+public:
+	/** The source object to use as the initial container to resolve the Source Property Path on. */
+	UPROPERTY()
+	TWeakObjectPtr<UObject> SourceObject;
+
+	/** The property path to trace to resolve this binding on the Source Object */
+	UPROPERTY()
+	FDynamicPropertyPath SourcePath;
+
+	/** Used to determine if a binding already exists on the object and if this binding can be safely removed. */
+	UPROPERTY()
+	FName DestinationProperty;
+};
diff --git a/Engine/Source/Runtime/UMG/Public/Binding/TextBinding.h b/Engine/Source/Runtime/UMG/Public/Binding/TextBinding.h
new file mode 100644
--- /dev/null
+++ b/Engine/Source/Runtime/UMG/Public/Binding/TextBinding.h
@@ -0,0 +1,29 @@
+﻿// Copyright 1998-2014 Epic Games, Inc. All Rights Reserved.
+
+#pragma once
+
+#include "TextBinding.generated.h"
+
+UCLASS()
+class UMG_API UTextBinding : public UPropertyBinding
+{
+	GENERATED_BODY()
+
+public:
+
+	UTextBinding();
+
+	virtual bool IsSupportedSource(UProperty* Property) const override;
+	virtual bool IsSupportedDestination(UProperty* Property) const override;
+
+	virtual void Bind(UProperty* Property, FScriptDelegate* Delegate) override;
+
+	UFUNCTION()
+	FText GetTextValue() const;
+
+	UFUNCTION()
+	FString GetStringValue() const;
+
+private:
+	mutable TOptional<bool> bNeedsConversion;
+};
diff --git a/Engine/Source/Runtime/UMG/Public/Binding/VisibilityBinding.h b/Engine/Source/Runtime/UMG/Public/Binding/VisibilityBinding.h
new file mode 100644
--- /dev/null
+++ b/Engine/Source/Runtime/UMG/Public/Binding/VisibilityBinding.h
@@ -0,0 +1,23 @@
+﻿// Copyright 1998-2014 Epic Games, Inc. All Rights Reserved.
+
+#pragma once
+
+#include "SlateWrapperTypes.h"
+
+#include "VisibilityBinding.generated.h"
+
+UCLASS()
+class UMG_API UVisibilityBinding : public UPropertyBinding
+{
+	GENERATED_BODY()
+
+public:
+
+	UVisibilityBinding();
+
+	virtual bool IsSupportedSource(UProperty* Property) const override;
+	virtual bool IsSupportedDestination(UProperty* Property) const override;
+
+	UFUNCTION()
+	ESlateVisibility GetValue() const;
+};
diff --git a/Engine/Source/Runtime/UMG/Public/Blueprint/UserWidget.h b/Engine/Source/Runtime/UMG/Public/Blueprint/UserWidget.h
--- a/Engine/Source/Runtime/UMG/Public/Blueprint/UserWidget.h
+++ b/Engine/Source/Runtime/UMG/Public/Blueprint/UserWidget.h
@@ -562,11 +562,11 @@ public:
 	void OnAnimationFinishedPlaying(UUMGSequencePlayer& Player );
 
 	/**
-	* Plays a sound through the UI
-	*
-	* @param The sound to play
-	*/
-	UFUNCTION(BlueprintCallable, BlueprintCosmetic, Category = "Sound")
+	 * Plays a sound through the UI
+	 *
+	 * @param The sound to play
+	 */
+	UFUNCTION(BlueprintCallable, BlueprintCosmetic, Category="Sound")
 	void PlaySound(class USoundBase* SoundToPlay);
 
 	/** @returns The UObject wrapper for a given SWidget */
@@ -589,7 +589,6 @@ public:
 
 	// Begin UObject interface
 	virtual void PreSave() override;
-	virtual void PostLoad() override;
 	// End UObject interface
 
 #if WITH_EDITOR
diff --git a/Engine/Source/Runtime/UMG/Public/Blueprint/WidgetBlueprintGeneratedClass.h b/Engine/Source/Runtime/UMG/Public/Blueprint/WidgetBlueprintGeneratedClass.h
--- a/Engine/Source/Runtime/UMG/Public/Blueprint/WidgetBlueprintGeneratedClass.h
+++ b/Engine/Source/Runtime/UMG/Public/Blueprint/WidgetBlueprintGeneratedClass.h
@@ -1,7 +1,10 @@
 // Copyright 1998-2014 Epic Games, Inc. All Rights Reserved.
 
 #pragma once
+
 #include "Engine/BlueprintGeneratedClass.h"
+#include "DynamicPropertyPath.h"
+
 #include "WidgetBlueprintGeneratedClass.generated.h"
 
 class UMovieScene;
@@ -10,14 +13,11 @@ class UUserWidget;
 class UWidgetAnimation;
 
 UENUM()
-namespace EBindingKind
+enum class EBindingKind : uint8
 {
-	enum Type
-	{
-		Function,
-		Property,
-	};
-}
+	Function,
+	Property
+};
 
 USTRUCT()
 struct FDelegateRuntimeBinding
@@ -36,9 +36,13 @@ struct FDelegateRuntimeBinding
 	UPROPERTY()
 	FName FunctionName;
 
+	/**  */
+	UPROPERTY()
+	FDynamicPropertyPath SourcePath;
+
 	/** The kind of binding we're performing, are we binding to a property or a function. */
 	UPROPERTY()
-	TEnumAsByte<EBindingKind::Type> Kind;
+	EBindingKind Kind;
 };
 
 
diff --git a/Engine/Source/Runtime/UMG/Public/Components/TextBlock.h b/Engine/Source/Runtime/UMG/Public/Components/TextBlock.h
--- a/Engine/Source/Runtime/UMG/Public/Components/TextBlock.h
+++ b/Engine/Source/Runtime/UMG/Public/Components/TextBlock.h
@@ -145,6 +145,7 @@ public:
 protected:
 	// UWidget interface
 	virtual TSharedRef<SWidget> RebuildWidget() override;
+	virtual void OnBindingChanged(const FName& Property) override;
 	// End of UWidget interface
 
 protected:
diff --git a/Engine/Source/Runtime/UMG/Public/Components/Widget.h b/Engine/Source/Runtime/UMG/Public/Components/Widget.h
--- a/Engine/Source/Runtime/UMG/Public/Components/Widget.h
+++ b/Engine/Source/Runtime/UMG/Public/Components/Widget.h
@@ -4,6 +4,7 @@
 #include "Components/Visual.h"
 #include "SlateWrapperTypes.h"
 #include "WidgetTransform.h"
+#include "DynamicPropertyPath.h"
 
 #include "Widget.generated.h"
 
@@ -35,15 +36,13 @@ public:
 	DECLARE_DYNAMIC_DELEGATE_RetVal(float, FGetFloat);
 	DECLARE_DYNAMIC_DELEGATE_RetVal(int32, FGetInt32);
 	DECLARE_DYNAMIC_DELEGATE_RetVal(FText, FGetText);
-	DECLARE_DYNAMIC_DELEGATE_RetVal(FVector2D, FGetVector2D);
-	DECLARE_DYNAMIC_DELEGATE_RetVal(FVector, FGetVector);
-	//DECLARE_DYNAMIC_DELEGATE_RetVal(FVector4, FGetVector4);
-	DECLARE_DYNAMIC_DELEGATE_RetVal(FMargin, FGetMargin);
+	//DECLARE_DYNAMIC_DELEGATE_RetVal(FVector2D, FGetVector2D);
+	//DECLARE_DYNAMIC_DELEGATE_RetVal(FVector, FGetVector);
+	//DECLARE_DYNAMIC_DELEGATE_RetVal(FMargin, FGetMargin);
 	DECLARE_DYNAMIC_DELEGATE_RetVal(FSlateColor, FGetSlateColor);
 	DECLARE_DYNAMIC_DELEGATE_RetVal(FLinearColor, FGetLinearColor);
 	DECLARE_DYNAMIC_DELEGATE_RetVal(ESlateVisibility, FGetSlateVisibility);
 	DECLARE_DYNAMIC_DELEGATE_RetVal(EMouseCursor::Type, FGetMouseCursor);
-	DECLARE_DYNAMIC_DELEGATE_RetVal(USlateBrushAsset*, FGetSlateBrushAsset);
 	DECLARE_DYNAMIC_DELEGATE_RetVal(FSlateBrush, FGetSlateBrush);
 	DECLARE_DYNAMIC_DELEGATE_RetVal(ESlateCheckBoxState::Type, FGetCheckBoxState);
 	DECLARE_DYNAMIC_DELEGATE_RetVal(UWidget*, FGetContent);
@@ -272,7 +271,20 @@ public:
 	 */
 	bool IsChildOf(UWidget* PossibleParent);
 
+	///**
+	// * Allows binding to properties of other objects at runtime.
+	// */
+	//UFUNCTION(BlueprintCallable, Category="Widget")
+	//bool BindProperty(const FName& DestinationProperty, UObject* SourceObject, const FName& SourceProperty);
+
+	/**  */
+	bool AddBinding(UDelegateProperty* DelegateProperty, UObject* SourceObject, const FDynamicPropertyPath& BindingPath);
+
+	static TSubclassOf<class UPropertyBinding> FindBinderClassForDestination(UProperty* Property);
+
+	// Begin UObject
 	virtual void PostLoad() override;
+	// End UObject
 	
 #if WITH_EDITOR
 	/** Is the label generated or provided by the user? */
@@ -332,6 +344,9 @@ public:
 	static void GatherAllChildren(UWidget* Root, TSet<UWidget*>& Children);
 	static UWidget* FindChildContainingDescendant(UWidget* Root, UWidget* Descendant);
 
+protected:
+	virtual void OnBindingChanged(const FName& Property);
+
 protected:
 	/** Function implemented by all subclasses of UWidget is called when the underlying SWidget needs to be constructed. */
 	virtual TSharedRef<SWidget> RebuildWidget();
@@ -369,4 +384,10 @@ protected:
 	/** Is this widget being displayed on a designer surface */
 	UPROPERTY(Transient)
 	bool bDesignTime;
+
+	/** Native property bindings. */
+	UPROPERTY(Transient)
+	TArray<class UPropertyBinding*> NativeBindings;
+
+	static TArray<TSubclassOf<UPropertyBinding>> BinderClasses;
 };
diff --git a/Engine/Source/Runtime/UMG/Public/UMG.h b/Engine/Source/Runtime/UMG/Public/UMG.h
--- a/Engine/Source/Runtime/UMG/Public/UMG.h
+++ b/Engine/Source/Runtime/UMG/Public/UMG.h
@@ -21,9 +21,14 @@
 /* Public Includes
  *****************************************************************************/
 
-// Components
+// Slate Types
 #include "SlateWrapperTypes.h"
 
+// Binding
+#include "DynamicPropertyPath.h"
+#include "PropertyBinding.h"
+
+// Components
 #include "DragDropOperation.h"
 
 #include "Visual.h"
