Build - More non-unity fixes.

diff --git a/Engine/Source/Editor/UMGEditor/Private/Details/SPropertyBinding.cpp b/Engine/Source/Editor/UMGEditor/Private/Details/SPropertyBinding.cpp
--- a/Engine/Source/Editor/UMGEditor/Private/Details/SPropertyBinding.cpp
+++ b/Engine/Source/Editor/UMGEditor/Private/Details/SPropertyBinding.cpp
@@ -8,6 +8,8 @@
 #include "WidgetGraphSchema.h"
 #include "BlueprintEditorUtils.h"
 #include "ScopedTransaction.h"
+#include "WidgetComponent.h"
+#include "WidgetAnimation.h"
 #include "PropertyBinding.h"
 
 
