UMG - Now rebasing the first segment after duplicating a widget blueprint to ensure all the 'self' class references are updated.

diff --git a/Engine/Source/Editor/UMGEditor/Private/WidgetBlueprint.cpp b/Engine/Source/Editor/UMGEditor/Private/WidgetBlueprint.cpp
--- a/Engine/Source/Editor/UMGEditor/Private/WidgetBlueprint.cpp
+++ b/Engine/Source/Editor/UMGEditor/Private/WidgetBlueprint.cpp
@@ -69,6 +69,11 @@ FEditorPropertyPathSegment::FEditorPropertyPathSegment(const UEdGraph* InFunctio
 	MemberGuid = InFunctionGraph->GraphGuid;
 }
 
+void FEditorPropertyPathSegment::Rebase(UBlueprint* SegmentBase)
+{
+	Struct = SegmentBase->GeneratedClass;
+}
+
 bool FEditorPropertyPathSegment::ValidateMember(UDelegateProperty* DelegateProperty, FText& OutError) const
 {
 	if ( DelegateProperty->SignatureFunction->NumParms == 1 )
@@ -270,6 +275,17 @@ FEditorPropertyPath::FEditorPropertyPath(const TArray<UField*>& BindingChain)
 	}
 }
 
+bool FEditorPropertyPath::Rebase(UBlueprint* SegmentBase)
+{
+	if ( !IsEmpty() )
+	{
+		Segments[0].Rebase(SegmentBase);
+		return true;
+	}
+
+	return false;
+}
+
 bool FEditorPropertyPath::Validate(UDelegateProperty* Destination, FText& OutError) const
 {
 	if ( IsEmpty() )
@@ -507,6 +523,22 @@ void UWidgetBlueprint::PostLoad()
 	}
 }
 
+void UWidgetBlueprint::PostDuplicate(bool bDuplicateForPIE)
+{
+	Super::PostDuplicate(bDuplicateForPIE);
+
+	if ( !bDuplicatingReadOnly )
+	{
+		// We need to update all the bindings and change each bindings first segment in the path
+		// to be the new class this blueprint generates, as all bindings must first originate on 
+		// the widget blueprint, the first segment is always a reference to 'self'.
+		for ( FDelegateEditorBinding& Binding : Bindings )
+		{
+			Binding.SourcePath.Rebase(this);
+		}
+	}
+}
+
 UClass* UWidgetBlueprint::GetBlueprintClass() const
 {
 	return UWidgetBlueprintGeneratedClass::StaticClass();
diff --git a/Engine/Source/Editor/UMGEditor/Public/WidgetBlueprint.h b/Engine/Source/Editor/UMGEditor/Public/WidgetBlueprint.h
--- a/Engine/Source/Editor/UMGEditor/Public/WidgetBlueprint.h
+++ b/Engine/Source/Editor/UMGEditor/Public/WidgetBlueprint.h
@@ -26,6 +26,7 @@ public:
 	UStruct* GetStruct() const { return Struct; }
 	UField* GetMember() const;
 
+	void Rebase(UBlueprint* SegmentBase);
 	bool ValidateMember(UDelegateProperty* DelegateProperty, FText& OutError) const;
 
 	FName GetMemberName() const;
@@ -62,16 +63,29 @@ struct UMGEDITOR_API FEditorPropertyPath
 
 public:
 
+	/**  */
 	FEditorPropertyPath();
+
+	/**  */
 	FEditorPropertyPath(const TArray<UField*>& BindingChain);
 
+	/**  */
+	bool Rebase(UBlueprint* SegmentBase);
+
+	/**  */
 	bool IsEmpty() const { return Segments.Num() == 0; }
+
+	/**  */
 	bool Validate(UDelegateProperty* Destination, FText& OutError) const;
 
+	/**  */
 	FText GetDisplayText() const;
 
+	/**  */
 	FDynamicPropertyPath ToPropertyPath() const;
 
+public:
+
 	/** The path of properties. */
 	UPROPERTY()
 	TArray<FEditorPropertyPathSegment> Segments;
@@ -184,6 +198,7 @@ public:
 
 	/** UObject interface */
 	virtual void PostLoad() override;
+	virtual void PostDuplicate(bool bDuplicateForPIE) override;
 	
 	// UBlueprint interface
 	virtual UClass* GetBlueprintClass() const override;
