UMG - Fixing an issue with the new binding system, bindings made to some events were incorrectly classified and showed up as a compile error.  Those bindings should work correctly going forward if you set them now.  Also fixing some visualization issues of the bind dropdown, it could some times hide the Goto function button.  Also the goto function button wasn't always working for blueprint functions.  That has been fixed.

diff --git a/Engine/Source/Editor/UMGEditor/Private/Details/SPropertyBinding.cpp b/Engine/Source/Editor/UMGEditor/Private/Details/SPropertyBinding.cpp
--- a/Engine/Source/Editor/UMGEditor/Private/Details/SPropertyBinding.cpp
+++ b/Engine/Source/Editor/UMGEditor/Private/Details/SPropertyBinding.cpp
@@ -37,7 +37,7 @@ void SPropertyBinding::Construct(const FArguments& InArgs, TSharedRef<FWidgetBlu
 		SNew(SHorizontalBox)
 
 		+ SHorizontalBox::Slot()
-		.AutoWidth()
+		.FillWidth(1.0f)
 		[
 			SNew(SComboButton)
 			.OnGetMenuContent(this, &SPropertyBinding::OnGenerateDelegateMenu, Widget, Property)
@@ -571,7 +571,21 @@ void SPropertyBinding::HandleAddFunctionBinding(TSharedRef<IPropertyHandle> Prop
 		Binding.ObjectName = SelectedObject->GetName();
 		Binding.PropertyName = PropertyHandle->GetProperty()->GetFName();
 		Binding.FunctionName = SelectedFunction->FuncName;
-		Binding.SourcePath = BindingPath;
+
+		// On event function bindings, we don't allow nested bindings, so don't fill out
+		// the source path, it will be interpreted as needing a dynamic binder.
+		if ( GeneratePureBindings )
+		{
+			Binding.SourcePath = BindingPath;
+		}
+		else
+		{
+			UBlueprint::GetGuidFromClassByFieldName<UFunction>(
+				SelectedFunction->Function->GetOwnerClass(),
+				SelectedFunction->Function->GetFName(),
+				Binding.MemberGuid);
+		}
+
 		Binding.Kind = EBindingKind::Function;
 
 		Blueprint->Bindings.Remove(Binding);
diff --git a/Engine/Source/Editor/UMGEditor/Private/WidgetBlueprint.cpp b/Engine/Source/Editor/UMGEditor/Private/WidgetBlueprint.cpp
--- a/Engine/Source/Editor/UMGEditor/Private/WidgetBlueprint.cpp
+++ b/Engine/Source/Editor/UMGEditor/Private/WidgetBlueprint.cpp
@@ -156,6 +156,8 @@ bool FEditorPropertyPathSegment::ValidateMember(UDelegateProperty* DelegatePrope
 		}
 	}
 
+	OutError = LOCTEXT("Binding_UnknownError", "Unknown Error");
+
 	return false;
 }
 
