UMG - Fixing a crash bug when creating bindings.  The binding segment system can now handle function binding that doesn't explictly have a dynamic path binder that supports it as long as the function signature being bound is an exact match.

diff --git a/Engine/Source/Editor/UMGEditor/Private/Customizations/DetailCustomizations.cpp b/Engine/Source/Editor/UMGEditor/Private/Customizations/DetailCustomizations.cpp
--- a/Engine/Source/Editor/UMGEditor/Private/Customizations/DetailCustomizations.cpp
+++ b/Engine/Source/Editor/UMGEditor/Private/Customizations/DetailCustomizations.cpp
@@ -100,6 +100,7 @@ void FBlueprintWidgetCustomization::CreateEventCustomization( IDetailLayoutBuild
 	IDetailCategoryBuilder& PropertyCategory = DetailLayout.EditCategory(FObjectEditorUtils::GetCategoryFName(Property), FText::GetEmpty(), ECategoryPriority::Uncommon);
 
 	IDetailPropertyRow& PropertyRow = PropertyCategory.AddProperty(DelegatePropertyHandle);
+	PropertyRow.OverrideResetToDefault(true, FSimpleDelegate::CreateSP(this, &FBlueprintWidgetCustomization::ResetToDefault_RemoveBinding, DelegatePropertyHandle));
 
 	FString LabelStr = Property->GetName();
 	LabelStr.RemoveFromEnd(TEXT("Event"));
@@ -130,13 +131,33 @@ void FBlueprintWidgetCustomization::CreateEventCustomization( IDetailLayoutBuild
 		]
 		.ValueContent()
 		.MinDesiredWidth(200)
-		.MaxDesiredWidth(200)
+		.MaxDesiredWidth(250)
 		[
 			SNew(SPropertyBinding, Editor.Pin().ToSharedRef(), Property, DelegatePropertyHandle)
 			.GeneratePureBindings(false)
 		];
 }
 
+void FBlueprintWidgetCustomization::ResetToDefault_RemoveBinding(TSharedRef<IPropertyHandle> PropertyHandle)
+{
+	const FScopedTransaction Transaction(LOCTEXT("BindDelegate", "Remove Binding"));
+
+	Blueprint->Modify();
+
+	TArray<UObject*> OuterObjects;
+	PropertyHandle->GetOuterObjects(OuterObjects);
+	for ( UObject* SelectedObject : OuterObjects )
+	{
+		FDelegateEditorBinding Binding;
+		Binding.ObjectName = SelectedObject->GetName();
+		Binding.PropertyName = PropertyHandle->GetProperty()->GetFName();
+
+		Blueprint->Bindings.Remove(Binding);
+	}
+
+	FBlueprintEditorUtils::MarkBlueprintAsStructurallyModified(Blueprint);
+}
+
 void FBlueprintWidgetCustomization::CreateMulticastEventCustomization(IDetailLayoutBuilder& DetailLayout, FName ThisComponentName, UClass* PropertyClass, UMulticastDelegateProperty* DelegateProperty)
 {
 	const FString AddString = FString(TEXT("Add "));
@@ -216,7 +237,7 @@ void FBlueprintWidgetCustomization::CreateMulticastEventCustomization(IDetailLay
 	]
 	.ValueContent()
 	.MinDesiredWidth(200)
-	.MaxDesiredWidth(200)
+	.MaxDesiredWidth(250)
 	[
 		SNew(SGraphSchemaActionButton, Editor.Pin(), ClickAction)
 	];
diff --git a/Engine/Source/Editor/UMGEditor/Private/Customizations/DetailCustomizations.h b/Engine/Source/Editor/UMGEditor/Private/Customizations/DetailCustomizations.h
--- a/Engine/Source/Editor/UMGEditor/Private/Customizations/DetailCustomizations.h
+++ b/Engine/Source/Editor/UMGEditor/Private/Customizations/DetailCustomizations.h
@@ -35,6 +35,8 @@ private:
 	void CreateEventCustomization( IDetailLayoutBuilder& DetailLayout, UDelegateProperty* Property, UWidget* Widget );
 
 	void CreateMulticastEventCustomization(IDetailLayoutBuilder& DetailLayout, FName ThisComponentName, UClass* PropertyClass, UMulticastDelegateProperty* Property);
+
+	void ResetToDefault_RemoveBinding(TSharedRef<IPropertyHandle> PropertyHandle);
 private:
 
 	TWeakPtr<FWidgetBlueprintEditor> Editor;
diff --git a/Engine/Source/Editor/UMGEditor/Private/Details/SPropertyBinding.cpp b/Engine/Source/Editor/UMGEditor/Private/Details/SPropertyBinding.cpp
--- a/Engine/Source/Editor/UMGEditor/Private/Details/SPropertyBinding.cpp
+++ b/Engine/Source/Editor/UMGEditor/Private/Details/SPropertyBinding.cpp
@@ -572,13 +572,9 @@ void SPropertyBinding::HandleAddFunctionBinding(TSharedRef<IPropertyHandle> Prop
 		Binding.PropertyName = PropertyHandle->GetProperty()->GetFName();
 		Binding.FunctionName = SelectedFunction->FuncName;
 
-		// On event function bindings, we don't allow nested bindings, so don't fill out
-		// the source path, it will be interpreted as needing a dynamic binder.
-		if ( GeneratePureBindings )
-		{
-			Binding.SourcePath = BindingPath;
-		}
-		else
+		Binding.SourcePath = BindingPath;
+
+		if ( SelectedFunction->Function )
 		{
 			UBlueprint::GetGuidFromClassByFieldName<UFunction>(
 				SelectedFunction->Function->GetOwnerClass(),
diff --git a/Engine/Source/Editor/UMGEditor/Private/Details/SPropertyBinding.h b/Engine/Source/Editor/UMGEditor/Private/Details/SPropertyBinding.h
--- a/Engine/Source/Editor/UMGEditor/Private/Details/SPropertyBinding.h
+++ b/Engine/Source/Editor/UMGEditor/Private/Details/SPropertyBinding.h
@@ -20,6 +20,11 @@ public:
 protected:
 	struct FFunctionInfo
 	{
+		FFunctionInfo()
+			: Function(nullptr)
+		{
+		}
+
 		FText DisplayName;
 		FString Tooltip;
 
diff --git a/Engine/Source/Editor/UMGEditor/Private/WidgetBlueprint.cpp b/Engine/Source/Editor/UMGEditor/Private/WidgetBlueprint.cpp
--- a/Engine/Source/Editor/UMGEditor/Private/WidgetBlueprint.cpp
+++ b/Engine/Source/Editor/UMGEditor/Private/WidgetBlueprint.cpp
@@ -76,6 +76,19 @@ void FEditorPropertyPathSegment::Rebase(UBlueprint* SegmentBase)
 
 bool FEditorPropertyPathSegment::ValidateMember(UDelegateProperty* DelegateProperty, FText& OutError) const
 {
+	// We may be binding to a function that doesn't have a explicit binder system that can handle it.  In that case
+	// check to see if the function signatures are compatible, if it is, even if we don't have a binder we can just
+	// directly bind the function to the delegate.
+	if ( UFunction* Function = Cast<UFunction>(GetMember()) )
+	{
+		// Check the signatures to ensure these functions match.
+		if ( Function->IsSignatureCompatibleWith(DelegateProperty->SignatureFunction, UFunction::GetDefaultIgnoredSignatureCompatibilityFlags() | CPF_ReturnParm) )
+		{
+			return true;
+		}
+	}
+
+	// Next check to see if we have a binder suitable for handling this case.
 	if ( DelegateProperty->SignatureFunction->NumParms == 1 )
 	{
 		if ( UProperty* ReturnProperty = DelegateProperty->SignatureFunction->GetReturnProperty() )
