UMG - Now painting helper lines when moving the anchor gizmo, along with percentage text.

diff --git a/Engine/Source/Editor/UMGEditor/Private/Designer/SDesignerView.cpp b/Engine/Source/Editor/UMGEditor/Private/Designer/SDesignerView.cpp
--- a/Engine/Source/Editor/UMGEditor/Private/Designer/SDesignerView.cpp
+++ b/Engine/Source/Editor/UMGEditor/Private/Designer/SDesignerView.cpp
@@ -703,8 +703,7 @@ bool SDesignerView::GetWidgetParentGeometry(const FWidgetReference& Widget, FGeo
 	{
 		if ( UPanelWidget* Parent = WidgetPreview->GetParent() )
 		{
-			FWidgetReference ParentReference = BlueprintEditor.Pin()->GetReferenceFromPreview(Parent);
-			return GetWidgetGeometry(ParentReference, Geometry);
+			return GetWidgetGeometry(Parent, Geometry);
 		}
 	}
 
@@ -714,23 +713,45 @@ bool SDesignerView::GetWidgetParentGeometry(const FWidgetReference& Widget, FGeo
 
 bool SDesignerView::GetWidgetGeometry(const FWidgetReference& Widget, FGeometry& Geometry) const
 {
-	if ( UWidget* WidgetPreview = Widget.GetPreview() )
+	if ( const UWidget* PreviewWidget = Widget.GetPreview() )
+	{
+		return GetWidgetGeometry(PreviewWidget, Geometry);
+	}
+
+	return false;
+}
+
+bool SDesignerView::GetWidgetGeometry(const UWidget* PreviewWidget, FGeometry& Geometry) const
+{
+	TSharedPtr<SWidget> CachedPreviewWidget = PreviewWidget->GetCachedWidget();
+	if ( CachedPreviewWidget.IsValid() )
 	{
-		TSharedPtr<SWidget> CachedPreviewWidget = WidgetPreview->GetCachedWidget();
-		if ( CachedPreviewWidget.IsValid() )
+		const FArrangedWidget* ArrangedWidget = CachedWidgetGeometry.Find(CachedPreviewWidget.ToSharedRef());
+		if ( ArrangedWidget )
 		{
-			const FArrangedWidget* ArrangedWidget = CachedWidgetGeometry.Find(CachedPreviewWidget.ToSharedRef());
-			if ( ArrangedWidget )
-			{
-				Geometry = ArrangedWidget->Geometry;
-				return true;
-			}
+			Geometry = ArrangedWidget->Geometry;
+			return true;
 		}
 	}
 
 	return false;
 }
 
+FGeometry SDesignerView::MakeGeometryWindowLocal(const FGeometry& WidgetGeometry) const
+{
+	FGeometry NewGeometry = WidgetGeometry;
+
+	TSharedPtr<SWindow> WidgetWindow = FSlateApplication::Get().FindWidgetWindow(SharedThis(this));
+	if ( WidgetWindow.IsValid() )
+	{
+		TSharedRef<SWindow> CurrentWindowRef = WidgetWindow.ToSharedRef();
+
+		NewGeometry.AppendTransform(FSlateLayoutTransform(Inverse(CurrentWindowRef->GetPositionInScreen())));
+	}
+
+	return NewGeometry;
+}
+
 void SDesignerView::ClearExtensionWidgets()
 {
 	ExtensionWidgetCanvas->ClearChildren();
diff --git a/Engine/Source/Editor/UMGEditor/Private/Designer/SDesignerView.h b/Engine/Source/Editor/UMGEditor/Private/Designer/SDesignerView.h
--- a/Engine/Source/Editor/UMGEditor/Private/Designer/SDesignerView.h
+++ b/Engine/Source/Editor/UMGEditor/Private/Designer/SDesignerView.h
@@ -63,7 +63,9 @@ public:
 	virtual ETransformMode::Type GetTransformMode() const override;
 	virtual FGeometry GetDesignerGeometry() const override;
 	virtual bool GetWidgetGeometry(const FWidgetReference& Widget, FGeometry& Geometry) const override;
+	virtual bool GetWidgetGeometry(const UWidget* PreviewWidget, FGeometry& Geometry) const override;
 	virtual bool GetWidgetParentGeometry(const FWidgetReference& Widget, FGeometry& Geometry) const override;
+	virtual FGeometry MakeGeometryWindowLocal(const FGeometry& WidgetGeometry) const override;
 	virtual void MarkDesignModifed(bool bRequiresRecompile) override;
 	// End of IUMGDesigner interface
 
diff --git a/Engine/Source/Editor/UMGEditor/Private/Extensions/CanvasSlotExtension.cpp b/Engine/Source/Editor/UMGEditor/Private/Extensions/CanvasSlotExtension.cpp
--- a/Engine/Source/Editor/UMGEditor/Private/Extensions/CanvasSlotExtension.cpp
+++ b/Engine/Source/Editor/UMGEditor/Private/Extensions/CanvasSlotExtension.cpp
@@ -247,7 +247,8 @@ void FCanvasSlotExtension::Tick(const FGeometry& AllottedGeometry, const double
 
 void FCanvasSlotExtension::Paint(const TSet< FWidgetReference >& Selection, const FGeometry& AllottedGeometry, const FSlateRect& MyClippingRect, FSlateWindowElementList& OutDrawElements, int32 LayerId) const
 {
-	PaintCollisionLines(Selection, AllottedGeometry, MyClippingRect, OutDrawElements, LayerId);
+	//PaintCollisionLines(Selection, AllottedGeometry, MyClippingRect, OutDrawElements, LayerId);
+	PaintDragPercentages(Selection, AllottedGeometry, MyClippingRect, OutDrawElements, LayerId);
 }
 
 FReply FCanvasSlotExtension::HandleAnchorBeginDrag(const FGeometry& Geometry, const FPointerEvent& Event, EAnchorWidget::Type AnchorType)
@@ -421,6 +422,152 @@ FReply FCanvasSlotExtension::HandleAnchorDragging(const FGeometry& Geometry, con
 	return FReply::Unhandled();
 }
 
+void FCanvasSlotExtension::PaintDragPercentages(const TSet< FWidgetReference >& Selection, const FGeometry& AllottedGeometry, const FSlateRect& MyClippingRect, FSlateWindowElementList& OutDrawElements, int32 LayerId) const
+{
+	// Just show the percentage lines when we're moving the anchor gizmo
+	if ( !bMovingAnchor )
+	{
+		return;
+	}
+
+	for ( const FWidgetReference& Selection : SelectionCache )
+	{
+		if ( UWidget* PreviewWidget = Selection.GetPreview() )
+		{
+			if ( UCanvasPanel* Canvas = Cast<UCanvasPanel>(PreviewWidget->GetParent()) )
+			{
+				UCanvasPanelSlot* PreviewCanvasSlot = CastChecked<UCanvasPanelSlot>(PreviewWidget->Slot);
+
+				FGeometry GeometryForSlot;
+				if ( Canvas->GetGeometryForSlot(PreviewCanvasSlot, GeometryForSlot) )
+				{
+					FGeometry CanvasGeometry;
+					Designer->GetWidgetGeometry(Canvas, CanvasGeometry);
+					CanvasGeometry = Designer->MakeGeometryWindowLocal(CanvasGeometry);
+
+					FVector2D CanvasSize = CanvasGeometry.Size;
+
+					const FAnchorData LayoutData = PreviewCanvasSlot->LayoutData;
+					const FVector2D AnchorMin = LayoutData.Anchors.Minimum;
+					const FVector2D AnchorMax = LayoutData.Anchors.Maximum;
+
+					auto DrawSegment =[&] (FVector2D Offset, FVector2D Start, FVector2D End, float Value, FVector2D TextTransform) {
+						PaintLineWithText(
+							Start + Offset,
+							End + Offset,
+							FText::FromString(FString::Printf(TEXT("%.1f%%"), Value)),
+							TextTransform,
+							CanvasGeometry, MyClippingRect, OutDrawElements, LayerId);
+					};
+
+					// Horizontal
+					{
+						auto DrawHorizontalSegment =[&] (FVector2D Offset, FVector2D TextTransform) {
+							
+							DrawSegment(Offset, FVector2D(0, 0), FVector2D(AnchorMin.X * CanvasSize.X, 0), AnchorMin.X * 100, TextTransform);
+							DrawSegment(Offset, FVector2D(AnchorMax.X * CanvasSize.X, 0), FVector2D(CanvasSize.X, 0), AnchorMax.X * 100, TextTransform);
+
+							if ( LayoutData.Anchors.IsStretchedHorizontal() )
+							{
+								DrawSegment(Offset, FVector2D(AnchorMin.X * CanvasSize.X, 0), FVector2D(AnchorMax.X * CanvasSize.X, 0), ( AnchorMax.X - AnchorMin.X ) * 100, TextTransform);
+							}
+						};
+
+						DrawHorizontalSegment(FVector2D(0, AnchorMin.Y * CanvasSize.Y), FVector2D(-1, -1));
+
+						if ( LayoutData.Anchors.IsStretchedVertical() )
+						{
+							DrawHorizontalSegment(FVector2D(0, AnchorMax.Y * CanvasSize.Y), FVector2D(-1, 0.25));
+						}
+					}
+
+					// Vertical
+					{
+						auto DrawVerticalSegment =[&] (FVector2D Offset, FVector2D TextTransform) {
+							DrawSegment(Offset, FVector2D(0, 0), FVector2D(0, AnchorMin.Y * CanvasSize.Y), AnchorMin.Y * 100, TextTransform);
+							DrawSegment(Offset, FVector2D(0, AnchorMax.Y * CanvasSize.Y), FVector2D(0, CanvasSize.Y), AnchorMax.Y * 100, TextTransform);
+
+							if ( LayoutData.Anchors.IsStretchedVertical() )
+							{
+								DrawSegment(Offset, FVector2D(0, AnchorMin.Y * CanvasSize.Y), FVector2D(0, AnchorMax.Y * CanvasSize.Y), ( AnchorMax.Y - AnchorMin.Y ) * 100, TextTransform);
+							}
+						};
+
+						DrawVerticalSegment(FVector2D(AnchorMin.X * CanvasSize.X, 0), FVector2D(-1, -1));
+
+						if ( LayoutData.Anchors.IsStretchedHorizontal() )
+						{
+							DrawVerticalSegment(FVector2D(AnchorMax.X * CanvasSize.X, 0), FVector2D(0.25, -1));
+						}
+					}
+				}
+			}
+		}
+	}
+}
+
+void FCanvasSlotExtension::PaintLineWithText(FVector2D Start, FVector2D End, FText Text, FVector2D TextTransform, const FGeometry& AllottedGeometry, const FSlateRect& MyClippingRect, FSlateWindowElementList& OutDrawElements, int32 LayerId) const
+{
+	TArray<FVector2D> LinePoints;
+	LinePoints.AddUninitialized(2);
+
+	LinePoints[0] = Start;
+	LinePoints[1] = End;
+
+	FLinearColor Color(0.5f, 0.75, 1);
+	const bool bAntialias = true;
+
+	FSlateDrawElement::MakeLines(
+		OutDrawElements,
+		LayerId,
+		AllottedGeometry.ToPaintGeometry(),
+		LinePoints,
+		MyClippingRect,
+		ESlateDrawEffect::None,
+		Color,
+		bAntialias);
+
+	FSlateFontInfo AnchorFont(FPaths::EngineContentDir() / TEXT("Slate/Fonts/Roboto-Regular.ttf"), 10 * (1 / Designer->GetPreviewScale()));
+	const FVector2D TextSize = FSlateApplication::Get().GetRenderer()->GetFontMeasureService()->Measure(Text, AnchorFont);
+
+	FVector2D Offset(0, 0);
+
+	if ( Start.Y == End.Y )
+	{
+		// If the lines run horizontally due to the Y's being the same:
+
+		// Position the text centered on the line.
+		Offset.X = TextSize.X / 2.0f;
+
+		// Offset the text vertically by the full size of the text, plus a little padding (2).
+		Offset.Y = TextSize.Y + ( 10 * ( 1 / Designer->GetPreviewScale() ) );
+	}
+	else
+	{
+		// If the lines are running vertically: 
+
+		Offset.X = TextSize.X + ( 10 * ( 1 / Designer->GetPreviewScale() ) );
+		
+		// The Y offset should be half the height of the text.
+		Offset.Y = TextSize.Y / 2.0f;
+	}
+
+	Offset = Offset * TextTransform;
+
+	FVector2D TextPos = ( LinePoints[0] + LinePoints[1] ) / 2.0f + Offset;
+
+	FSlateDrawElement::MakeText(
+		OutDrawElements,
+		LayerId,
+		AllottedGeometry.MakeChild(TextPos, TextPos).ToPaintGeometry(),
+		Text,
+		AnchorFont,
+		MyClippingRect,
+		ESlateDrawEffect::None,
+		FLinearColor::White
+		);
+}
+
 void FCanvasSlotExtension::PaintCollisionLines(const TSet< FWidgetReference >& Selection, const FGeometry& AllottedGeometry, const FSlateRect& MyClippingRect, FSlateWindowElementList& OutDrawElements, int32 LayerId) const
 {
 	for ( const FWidgetReference& WidgetRef : Selection )
diff --git a/Engine/Source/Editor/UMGEditor/Private/Extensions/CanvasSlotExtension.h b/Engine/Source/Editor/UMGEditor/Private/Extensions/CanvasSlotExtension.h
--- a/Engine/Source/Editor/UMGEditor/Private/Extensions/CanvasSlotExtension.h
+++ b/Engine/Source/Editor/UMGEditor/Private/Extensions/CanvasSlotExtension.h
@@ -60,6 +60,8 @@ private:
 	static void GetCollisionSegmentsFromGeometry(FGeometry ArrangedGeometry, TArray<FVector2D>& Segments);
 
 	void PaintCollisionLines(const TSet< FWidgetReference >& Selection, const FGeometry& AllottedGeometry, const FSlateRect& MyClippingRect, FSlateWindowElementList& OutDrawElements, int32 LayerId) const;
+	void PaintDragPercentages(const TSet< FWidgetReference >& Selection, const FGeometry& AllottedGeometry, const FSlateRect& MyClippingRect, FSlateWindowElementList& OutDrawElements, int32 LayerId) const;
+	void PaintLineWithText(FVector2D Start, FVector2D End, FText Text, FVector2D TextTransform, const FGeometry& AllottedGeometry, const FSlateRect& MyClippingRect, FSlateWindowElementList& OutDrawElements, int32 LayerId) const;
 
 	void ProximitySnapValue(float SnapFrequency, float SnapProximity, float& Value);
 
diff --git a/Engine/Source/Editor/UMGEditor/Public/IUMGDesigner.h b/Engine/Source/Editor/UMGEditor/Public/IUMGDesigner.h
--- a/Engine/Source/Editor/UMGEditor/Public/IUMGDesigner.h
+++ b/Engine/Source/Editor/UMGEditor/Public/IUMGDesigner.h
@@ -45,6 +45,16 @@ public:
 	 */
 	virtual bool GetWidgetGeometry(const FWidgetReference& Widget, FGeometry& Geometry) const = 0;
 
+	/**
+	 * Gets the previous frames widget geometry.
+	 */
+	virtual bool GetWidgetGeometry(const UWidget* PreviewWidget, FGeometry& Geometry) const = 0;
+
+	/**
+	 * Takes geometry and adds the inverse of the window transform to get the geometry in the space of the window.
+	 */
+	virtual FGeometry MakeGeometryWindowLocal(const FGeometry& WidgetGeometry) const = 0;
+
 	/**
 	 * Gets the previous frames widget geometry of the parent of the provided widget.
 	 */
