UMG - Now showing the anchor percentages on hover of the anchor gizmo pieces.

diff --git a/Engine/Source/Editor/UMGEditor/Private/Extensions/CanvasSlotExtension.cpp b/Engine/Source/Editor/UMGEditor/Private/Extensions/CanvasSlotExtension.cpp
--- a/Engine/Source/Editor/UMGEditor/Private/Extensions/CanvasSlotExtension.cpp
+++ b/Engine/Source/Editor/UMGEditor/Private/Extensions/CanvasSlotExtension.cpp
@@ -9,6 +9,58 @@
 
 #define LOCTEXT_NAMESPACE "UMG"
 
+/////////////////////////////////////////////////////
+// SEventShim
+
+class SEventShim : public SCompoundWidget
+{
+public:
+	SLATE_BEGIN_ARGS(SEventShim)
+		: _Content()
+		, _OnMouseEnter()
+		, _OnMouseLeave()
+	{}
+		/** Slot for this designers content (optional) */
+		SLATE_DEFAULT_SLOT(FArguments, Content)
+
+		SLATE_EVENT(FSimpleDelegate, OnMouseEnter)
+		SLATE_EVENT(FSimpleDelegate, OnMouseLeave)
+
+	SLATE_END_ARGS()
+
+	void Construct(const FArguments& InArgs)
+	{
+		MouseEnter = InArgs._OnMouseEnter;
+		MouseLeave = InArgs._OnMouseLeave;
+
+		ChildSlot
+		[
+			InArgs._Content.Widget
+		];
+	}
+
+	virtual void OnMouseEnter(const FGeometry& MyGeometry, const FPointerEvent& MouseEvent) override
+	{
+		SCompoundWidget::OnMouseEnter(MyGeometry, MouseEvent);
+
+		MouseEnter.ExecuteIfBound();
+	}
+
+	virtual void OnMouseLeave(const FPointerEvent& MouseEvent) override
+	{
+		SCompoundWidget::OnMouseLeave(MouseEvent);
+
+		MouseLeave.ExecuteIfBound();
+	}
+
+private:
+	FSimpleDelegate MouseEnter;
+	FSimpleDelegate MouseLeave;
+};
+
+/////////////////////////////////////////////////////
+// FCanvasSlotExtension
+
 const float SnapDistance = 7;
 
 static float DistancePointToLine2D(const FVector2D& LinePointA, const FVector2D& LinePointB, const FVector2D& PointC)
@@ -22,6 +74,7 @@ static float DistancePointToLine2D(const FVector2D& LinePointA, const FVector2D&
 
 FCanvasSlotExtension::FCanvasSlotExtension()
 	: bMovingAnchor(false)
+	, bHoveringAnchor(false)
 {
 	ExtensionId = FName(TEXT("CanvasSlot"));
 }
@@ -94,18 +147,33 @@ TSharedRef<SWidget> FCanvasSlotExtension::MakeAnchorWidget(EAnchorWidget::Type A
 		.Visibility(this, &FCanvasSlotExtension::GetAnchorVisibility, AnchorType)
 		.Padding(FMargin(0))
 		[
-			SNew(SBox)
-			.WidthOverride(Width)
-			.HeightOverride(Height)
-			.HAlign(HAlign_Fill)
-			.VAlign(VAlign_Fill)
+			SNew(SEventShim)
+			.OnMouseEnter(this, &FCanvasSlotExtension::OnMouseEnterAnchor)
+			.OnMouseLeave(this, &FCanvasSlotExtension::OnMouseLeaveAnchor)
 			[
-				SNew(SImage)
-				.Image(this, &FCanvasSlotExtension::GetAnchorBrush, AnchorType)
+				SNew(SBox)
+				.WidthOverride(Width)
+				.HeightOverride(Height)
+				.HAlign(HAlign_Fill)
+				.VAlign(VAlign_Fill)
+				[
+					SNew(SImage)
+					.Image(this, &FCanvasSlotExtension::GetAnchorBrush, AnchorType)
+				]
 			]
 		];
 }
 
+void FCanvasSlotExtension::OnMouseEnterAnchor()
+{
+	bHoveringAnchor = true;
+}
+
+void FCanvasSlotExtension::OnMouseLeaveAnchor()
+{
+	bHoveringAnchor = false;
+}
+
 const FSlateBrush* FCanvasSlotExtension::GetAnchorBrush(EAnchorWidget::Type AnchorType) const
 {
 	switch ( AnchorType )
@@ -425,7 +493,7 @@ FReply FCanvasSlotExtension::HandleAnchorDragging(const FGeometry& Geometry, con
 void FCanvasSlotExtension::PaintDragPercentages(const TSet< FWidgetReference >& Selection, const FGeometry& AllottedGeometry, const FSlateRect& MyClippingRect, FSlateWindowElementList& OutDrawElements, int32 LayerId) const
 {
 	// Just show the percentage lines when we're moving the anchor gizmo
-	if ( !bMovingAnchor )
+	if ( !(bMovingAnchor || bHoveringAnchor) )
 	{
 		return;
 	}
diff --git a/Engine/Source/Editor/UMGEditor/Private/Extensions/CanvasSlotExtension.h b/Engine/Source/Editor/UMGEditor/Private/Extensions/CanvasSlotExtension.h
--- a/Engine/Source/Editor/UMGEditor/Private/Extensions/CanvasSlotExtension.h
+++ b/Engine/Source/Editor/UMGEditor/Private/Extensions/CanvasSlotExtension.h
@@ -51,6 +51,9 @@ private:
 
 	TSharedRef<SWidget> MakeAnchorWidget(EAnchorWidget::Type AnchorType, float Width, float Height);
 
+	void OnMouseEnterAnchor();
+	void OnMouseLeaveAnchor();
+
 	const FSlateBrush* GetAnchorBrush(EAnchorWidget::Type AnchorType) const;
 	EVisibility GetAnchorVisibility(EAnchorWidget::Type AnchorType) const;
 	FVector2D GetAnchorAlignment(EAnchorWidget::Type AnchorType) const;
@@ -73,6 +76,9 @@ private:
 	/** */
 	bool bMovingAnchor;
 
+	/** */
+	bool bHoveringAnchor;
+
 	/** */
 	FVector2D MouseDownPosition;
 
