UMG - Crash Fix - Animation renaming now checks for valid variable names, checking collisions with other members of the blueprint.

diff --git a/Engine/Source/Editor/UMGEditor/Private/TabFactory/AnimationTabSummoner.cpp b/Engine/Source/Editor/UMGEditor/Private/TabFactory/AnimationTabSummoner.cpp
--- a/Engine/Source/Editor/UMGEditor/Private/TabFactory/AnimationTabSummoner.cpp
+++ b/Engine/Source/Editor/UMGEditor/Private/TabFactory/AnimationTabSummoner.cpp
@@ -98,11 +98,24 @@ private:
 
 		FString NewName = InText.ToString();
 
-		UWidgetBlueprint* Blueprint = BlueprintEditor.Pin()->GetWidgetBlueprintObj();
-		if( Animation->GetName() != NewName && FindObject<UWidgetAnimation>( Blueprint, *NewName, true ) )
+		if ( Animation->GetName() != NewName )
 		{
-			OutErrorMessage = LOCTEXT("ExistingMovieSceneError", "An animation with this name already exists");
-			return false;
+			UWidgetBlueprint* Blueprint = BlueprintEditor.Pin()->GetWidgetBlueprintObj();
+			if ( FindObject<UWidgetAnimation>(Blueprint, *NewName, true) )
+			{
+				OutErrorMessage = LOCTEXT("ExistingMovieSceneError", "An animation with this name already exists");
+				return false;
+			}
+
+			FKismetNameValidator Validator(Blueprint);
+			EValidatorResult ValidationResult = Validator.IsValid(NewName);
+
+			if ( ValidationResult != EValidatorResult::Ok )
+			{
+				FString ErrorString = FKismetNameValidator::GetErrorString(NewName, ValidationResult);
+				OutErrorMessage = FText::FromString(ErrorString);
+				return false;
+			}
 		}
 
 		return true;
