Fix non-unity build error, missing header include.

diff --git a/Engine/Source/Editor/UMGEditor/Private/TabFactory/AnimationTabSummoner.cpp b/Engine/Source/Editor/UMGEditor/Private/TabFactory/AnimationTabSummoner.cpp
--- a/Engine/Source/Editor/UMGEditor/Private/TabFactory/AnimationTabSummoner.cpp
+++ b/Engine/Source/Editor/UMGEditor/Private/TabFactory/AnimationTabSummoner.cpp
@@ -15,6 +15,7 @@
 #include "WidgetBlueprint.h"
 #include "Kismet2/BlueprintEditorUtils.h"
 #include "Misc/TextFilter.h"
+#include "Editor/UnrealEd/Public/Kismet2/Kismet2NameValidators.h"
 
 #define LOCTEXT_NAMESPACE "UMG"
 
