UMG - Binding system now correctly checks owner struct before owner class now to prevent native structs from returning the owning class that owns the native struct first as the owner of the property.  Enhancing some error logging to provide more info.

diff --git a/Engine/Source/Editor/UMGEditor/Private/WidgetBlueprint.cpp b/Engine/Source/Editor/UMGEditor/Private/WidgetBlueprint.cpp
--- a/Engine/Source/Editor/UMGEditor/Private/WidgetBlueprint.cpp
+++ b/Engine/Source/Editor/UMGEditor/Private/WidgetBlueprint.cpp
@@ -26,16 +26,16 @@ FEditorPropertyPathSegment::FEditorPropertyPathSegment(const UProperty* InProper
 {
 	IsProperty = true;
 	MemberName = InProperty->GetFName();
-	if ( InProperty->GetOwnerClass() )
-	{
-		Struct = InProperty->GetOwnerClass();
-		UBlueprint::GetGuidFromClassByFieldName<UProperty>(InProperty->GetOwnerClass(), InProperty->GetFName(), MemberGuid);
-	}
-	else if ( InProperty->GetOwnerStruct() )
+	if ( InProperty->GetOwnerStruct() )
 	{
 		Struct = InProperty->GetOwnerStruct();
 		MemberGuid = FStructureEditorUtils::GetGuidForProperty(InProperty);
 	}
+	else if ( InProperty->GetOwnerClass() )
+	{
+		Struct = InProperty->GetOwnerClass();
+		UBlueprint::GetGuidFromClassByFieldName<UProperty>(InProperty->GetOwnerClass(), InProperty->GetFName(), MemberGuid);
+	}
 	else
 	{
 		// Should not be possible to hit.
@@ -312,24 +312,26 @@ bool FEditorPropertyPath::Validate(UDelegateProperty* Destination, FText& OutErr
 	for ( int32 SegmentIndex = 0; SegmentIndex < Segments.Num(); SegmentIndex++ )
 	{
 		const FEditorPropertyPathSegment& Segment = Segments[SegmentIndex];
-		if ( Segment.GetStruct() == nullptr )
+		if ( UStruct* OwnerStruct = Segment.GetStruct() )
 		{
-			OutError = FText::Format(LOCTEXT("Binding_StructNotFound", "Binding:{0}, Unable to locate owner class."),
-				GetDisplayText());
+			if ( Segment.GetMember() == nullptr )
+			{
+				OutError = FText::Format(LOCTEXT("Binding_MemberNotFound", "Binding: '{0}' : '{1}' was not found on '{2}'."),
+					GetDisplayText(),
+					Segment.GetMemberDisplayText(),
+					OwnerStruct->GetDisplayNameText());
 
-			return false;
+				return false;
+			}
 		}
-
-		if ( Segment.GetMember() == nullptr )
+		else
 		{
-			OutError = FText::Format(LOCTEXT("Binding_MemberNotFound", "Binding:{0}, Member:{1}, was not found."),
+			OutError = FText::Format(LOCTEXT("Binding_StructNotFound", "Binding: '{0}' : Unable to locate owner class or struct for '{1}'"),
 				GetDisplayText(),
 				Segment.GetMemberDisplayText());
 
 			return false;
 		}
-
-		// TODO Add additional checks, are non-last fields objects, structs still?
 	}
 
 	// Validate the last member in the segment
