UMG - Adding rulers to the designer, the designer now also adds some padding when performing zoom to fit to prevent overlaping with the ruler or the buttons or text that overlays on top of the preview.

diff --git a/Engine/Source/Editor/UMGEditor/Private/Designer/SDesignSurface.cpp b/Engine/Source/Editor/UMGEditor/Private/Designer/SDesignSurface.cpp
--- a/Engine/Source/Editor/UMGEditor/Private/Designer/SDesignSurface.cpp
+++ b/Engine/Source/Editor/UMGEditor/Private/Designer/SDesignSurface.cpp
@@ -135,6 +135,8 @@ void SDesignSurface::Construct(const FArguments& InArgs)
 	ZoomTargetTopLeft = FVector2D::ZeroVector;
 	ZoomTargetBottomRight = FVector2D::ZeroVector;
 
+	ZoomToFitPadding = FVector2D(100, 100);
+
 	ChildSlot
 	[
 		InArgs._Content.Widget
@@ -190,6 +192,8 @@ int32 SDesignSurface::OnPaint(const FPaintArgs& Args, const FGeometry& AllottedG
 
 FReply SDesignSurface::OnMouseButtonDown(const FGeometry& MyGeometry, const FPointerEvent& MouseEvent)
 {
+	SCompoundWidget::OnMouseButtonDown(MyGeometry, MouseEvent);
+
 	if ( MouseEvent.GetEffectingButton() == EKeys::RightMouseButton )
 	{
 		bIsPanning = false;
@@ -200,6 +204,8 @@ FReply SDesignSurface::OnMouseButtonDown(const FGeometry& MyGeometry, const FPoi
 
 FReply SDesignSurface::OnMouseButtonUp(const FGeometry& MyGeometry, const FPointerEvent& MouseEvent)
 {
+	SCompoundWidget::OnMouseButtonUp(MyGeometry, MouseEvent);
+
 	if ( MouseEvent.GetEffectingButton() == EKeys::RightMouseButton )
 	{
 		bIsPanning = false;
@@ -352,7 +358,7 @@ bool SDesignSurface::ZoomToLocation(const FVector2D& CurrentSizeWithoutZoom, con
 	// Find lowest zoom level that will display all nodes
 	for ( int32 Zoom = 0; Zoom < DefaultZoomLevel; ++Zoom )
 	{
-		const FVector2D SizeWithZoom = CurrentSizeWithoutZoom / ZoomLevels->GetZoomAmount(Zoom);
+		const FVector2D SizeWithZoom = (CurrentSizeWithoutZoom - ZoomToFitPadding) / ZoomLevels->GetZoomAmount(Zoom);
 		const FVector2D LeftOverSize = SizeWithZoom - DesiredSize;
 
 		if ( ( DesiredSize.X > SizeWithZoom.X ) || ( DesiredSize.Y > SizeWithZoom.Y ) )
diff --git a/Engine/Source/Editor/UMGEditor/Private/Designer/SDesignSurface.h b/Engine/Source/Editor/UMGEditor/Private/Designer/SDesignSurface.h
--- a/Engine/Source/Editor/UMGEditor/Private/Designer/SDesignSurface.h
+++ b/Engine/Source/Editor/UMGEditor/Private/Designer/SDesignSurface.h
@@ -100,6 +100,7 @@ protected:
 
 	FVector2D ZoomTargetTopLeft;
 	FVector2D ZoomTargetBottomRight;
+	FVector2D ZoomToFitPadding;
 
 	/** Does the user need to press Control in order to over-zoom. */
 	bool bRequireControlToOverZoom;
diff --git a/Engine/Source/Editor/UMGEditor/Private/Designer/SDesignerView.cpp b/Engine/Source/Editor/UMGEditor/Private/Designer/SDesignerView.cpp
--- a/Engine/Source/Editor/UMGEditor/Private/Designer/SDesignerView.cpp
+++ b/Engine/Source/Editor/UMGEditor/Private/Designer/SDesignerView.cpp
@@ -19,6 +19,7 @@
 
 #include "WidgetTemplateDragDropOp.h"
 #include "SZoomPan.h"
+#include "SRuler.h"
 #include "SDisappearingBar.h"
 #include "SDesignerToolBar.h"
 #include "DesignerCommands.h"
@@ -244,29 +245,6 @@ void SDesignerView::Construct(const FArguments& InArgs, TSharedPtr<FWidgetBluepr
 				.OnPaintHandler(this, &SDesignerView::HandleEffectsPainting)
 			]
 
-			// Top-right corner text indicating PIE is active
-			+ SOverlay::Slot()
-			.Padding(0)
-			.VAlign(VAlign_Fill)
-			.HAlign(HAlign_Fill)
-			[
-				SNew(SImage)
-				.Visibility(this, &SDesignerView::PIENotification)
-				.Image(FEditorStyle::GetBrush(TEXT("Graph.PlayInEditor")))
-			]
-
-			// Top-right corner text indicating PIE is active
-			+ SOverlay::Slot()
-			.Padding(20)
-			.VAlign(VAlign_Top)
-			.HAlign(HAlign_Right)
-			[
-				SNew(STextBlock)
-				.Visibility(this, &SDesignerView::PIENotification)
-				.TextStyle( FEditorStyle::Get(), "Graph.SimulatingText" )
-				.Text( LOCTEXT("SIMULATING", "SIMULATING") )
-			]
-
 			// A layer in the overlay where we put all the user intractable widgets, like the reorder widgets.
 			+ SOverlay::Slot()
 			.HAlign(HAlign_Fill)
@@ -276,195 +254,260 @@ void SDesignerView::Construct(const FArguments& InArgs, TSharedPtr<FWidgetBluepr
 				.Visibility(EVisibility::SelfHitTestInvisible)
 			]
 
-			// Top bar with buttons for changing the designer
 			+ SOverlay::Slot()
-			.HAlign(HAlign_Fill)
-			.VAlign(VAlign_Top)
 			[
-				SNew(SHorizontalBox)
+				SNew(SGridPanel)
+				.FillColumn(1, 1.0f)
+				.FillRow(1, 1.0f)
 
-				+ SHorizontalBox::Slot()
-				.AutoWidth()
-				.Padding(6, 2, 0, 0)
+				// Corner
+				+ SGridPanel::Slot(0, 0)
 				[
-					SNew(STextBlock)
-					.TextStyle(FEditorStyle::Get(), "Graph.ZoomText")
-					.Text(this, &SDesignerView::GetZoomText)
-					.ColorAndOpacity(this, &SDesignerView::GetZoomTextColorAndOpacity)
-					.Visibility(EVisibility::SelfHitTestInvisible)
+					SNew(SBorder)
+					.BorderImage(FCoreStyle::Get().GetBrush("GenericWhiteBox"))
+					.BorderBackgroundColor(FLinearColor(FColor(48, 48, 48)))
 				]
 
-				+ SHorizontalBox::Slot()
-				.FillWidth(1.0f)
+				// Top Ruler
+				+ SGridPanel::Slot(1, 0)
 				[
-					SNew(SSpacer)
-					.Size(FVector2D(1, 1))
+					SAssignNew(TopRuler, SRuler)
+					.Orientation(Orient_Horizontal)
+					.Visibility(this, &SDesignerView::GetRulerVisibility)
 				]
 
-				+ SHorizontalBox::Slot()
-				.AutoWidth()
-				.Padding(5.0f)
+				// Side Ruler
+				+ SGridPanel::Slot(0, 1)
 				[
-					SNew(SDesignerToolBar)
-					.CommandList(CommandList)
+					SAssignNew(SideRuler, SRuler)
+					.Orientation(Orient_Vertical)
+					.Visibility(this, &SDesignerView::GetRulerVisibility)
 				]
 
-				+ SHorizontalBox::Slot()
-				.AutoWidth()
-				.Padding(5.0f)
+				// Designer overlay UI, toolbar, status messages, zoom level...etc
+				+ SGridPanel::Slot(1, 1)
 				[
-					SNew(SButton)
-					.ButtonStyle(FEditorStyle::Get(), "ViewportMenu.Button")
-					.ToolTipText(LOCTEXT("ZoomToFit_ToolTip", "Zoom To Fit"))
-					.OnClicked(this, &SDesignerView::HandleZoomToFitClicked)
-					[
-						SNew(SImage)
-						.Image(FEditorStyle::GetBrush("UMGEditor.ZoomToFit"))
-					]
-				]
-				
-				+ SHorizontalBox::Slot()
-				.AutoWidth()
-				.Padding(5.0f)
-				[
-					SNew(SComboButton)
-					.ButtonStyle(FEditorStyle::Get(), "ViewportMenu.Button")
-					.ForegroundColor(FLinearColor::Black)
-					.OnGetMenuContent(this, &SDesignerView::GetAspectMenu)
-					.ContentPadding(2.0f)
-					.ButtonContent()
-					[
-						SNew(STextBlock)
-						.Text(LOCTEXT("PreviewSize", "Preview Size"))
-						.TextStyle(FEditorStyle::Get(), "ViewportMenu.Label")
-					]
+					CreateOverlayUI()
 				]
+			]
+		]
+	);
 
-				+ SHorizontalBox::Slot()
-				.AutoWidth()
-				.Padding(5.0f)
-				.VAlign(VAlign_Center)
-				[
-					SNew(SNumericEntryBox<int32>)
-					.AllowSpin(true)
-					.Delta(1)
-					.MinSliderValue(1)
-					.MinValue(1)
-					.MaxSliderValue(TOptional<int32>(1000))
-					.Value(this, &SDesignerView::GetCustomResolutionWidth)
-					.OnValueChanged(this, &SDesignerView::OnCustomResolutionWidthChanged)
-					.Visibility(this, &SDesignerView::GetCustomResolutionEntryVisibility)
-					.MinDesiredValueWidth(50)
-					.LabelPadding(0)
-					.Label()
-					[
-						SNumericEntryBox<int32>::BuildLabel(LOCTEXT("Width", "Width"), FLinearColor::White, SNumericEntryBox<int32>::RedLabelBackgroundColor)
-					]
-				]
+	BlueprintEditor.Pin()->OnSelectedWidgetsChanged.AddRaw(this, &SDesignerView::OnEditorSelectionChanged);
 
-				+ SHorizontalBox::Slot()
-				.AutoWidth()
-				.Padding(5.0f)
-				.VAlign(VAlign_Center)
-				[
-					SNew(SNumericEntryBox<int32>)
-					.AllowSpin(true)
-					.Delta(1)
-					.MinSliderValue(1)
-					.MaxSliderValue(TOptional<int32>(1000))
-					.MinValue(1)
-					.Value(this, &SDesignerView::GetCustomResolutionHeight)
-					.OnValueChanged(this, &SDesignerView::OnCustomResolutionHeightChanged)
-					.Visibility(this, &SDesignerView::GetCustomResolutionEntryVisibility)
-					.MinDesiredValueWidth(50)
-					.LabelPadding(0)
-					.Label()
-					[
-						SNumericEntryBox<int32>::BuildLabel(LOCTEXT("Height", "Height"), FLinearColor::White, SNumericEntryBox<int32>::GreenLabelBackgroundColor)
-					]
-				]
-			]
+	ZoomToFit(/*bInstantZoom*/ true);
+}
 
-			// Info Bar, displays heads up information about some actions.
-			+ SOverlay::Slot()
-			.HAlign(HAlign_Fill)
-			.VAlign(VAlign_Bottom)
+TSharedRef<SWidget> SDesignerView::CreateOverlayUI()
+{
+	return SNew(SOverlay)
+
+	// Top-right corner text indicating PIE is active
+	+ SOverlay::Slot()
+	.Padding(0)
+	.VAlign(VAlign_Fill)
+	.HAlign(HAlign_Fill)
+	[
+		SNew(SImage)
+		.Visibility(this, &SDesignerView::PIENotification)
+		.Image(FEditorStyle::GetBrush(TEXT("Graph.PlayInEditor")))
+	]
+
+	// Top-right corner text indicating PIE is active
+	+ SOverlay::Slot()
+	.Padding(20)
+	.VAlign(VAlign_Top)
+	.HAlign(HAlign_Right)
+	[
+		SNew(STextBlock)
+		.Visibility(this, &SDesignerView::PIENotification)
+		.TextStyle( FEditorStyle::Get(), "Graph.SimulatingText" )
+		.Text( LOCTEXT("SIMULATING", "SIMULATING") )
+	]
+
+	// Top bar with buttons for changing the designer
+	+ SOverlay::Slot()
+	.HAlign(HAlign_Fill)
+	.VAlign(VAlign_Top)
+	[
+		SNew(SHorizontalBox)
+
+		+ SHorizontalBox::Slot()
+		.AutoWidth()
+		.Padding(6, 2, 0, 0)
+		[
+			SNew(STextBlock)
+			.TextStyle(FEditorStyle::Get(), "Graph.ZoomText")
+			.Text(this, &SDesignerView::GetZoomText)
+			.ColorAndOpacity(this, &SDesignerView::GetZoomTextColorAndOpacity)
+			.Visibility(EVisibility::SelfHitTestInvisible)
+		]
+
+		+ SHorizontalBox::Slot()
+		.FillWidth(1.0f)
+		[
+			SNew(SSpacer)
+			.Size(FVector2D(1, 1))
+		]
+
+		+ SHorizontalBox::Slot()
+		.AutoWidth()
+		.Padding(5.0f)
+		[
+			SNew(SDesignerToolBar)
+			.CommandList(CommandList)
+		]
+
+		+ SHorizontalBox::Slot()
+		.AutoWidth()
+		.Padding(5.0f)
+		[
+			SNew(SButton)
+			.ButtonStyle(FEditorStyle::Get(), "ViewportMenu.Button")
+			.ToolTipText(LOCTEXT("ZoomToFit_ToolTip", "Zoom To Fit"))
+			.OnClicked(this, &SDesignerView::HandleZoomToFitClicked)
 			[
-				SNew(SDisappearingBar)
-				[
-					SNew(SBorder)
-					.BorderImage(FEditorStyle::GetBrush("WhiteBrush"))
-					.BorderBackgroundColor(FLinearColor(0.10, 0.10, 0.10, 0.75))
-					.HAlign(HAlign_Center)
-					.VAlign(VAlign_Center)
-					.Padding(FMargin(0, 5))
-					.Visibility(this, &SDesignerView::GetInfoBarVisibility)
-					[
-						SNew(STextBlock)
-						.TextStyle(FEditorStyle::Get(), "Graph.ZoomText")
-						.Text(this, &SDesignerView::GetInfoBarText)
-					]
-				]
+				SNew(SImage)
+				.Image(FEditorStyle::GetBrush("UMGEditor.ZoomToFit"))
+			]
+		]
+				
+		+ SHorizontalBox::Slot()
+		.AutoWidth()
+		.Padding(5.0f)
+		[
+			SNew(SComboButton)
+			.ButtonStyle(FEditorStyle::Get(), "ViewportMenu.Button")
+			.ForegroundColor(FLinearColor::Black)
+			.OnGetMenuContent(this, &SDesignerView::GetAspectMenu)
+			.ContentPadding(2.0f)
+			.ButtonContent()
+			[
+				SNew(STextBlock)
+				.Text(LOCTEXT("PreviewSize", "Preview Size"))
+				.TextStyle(FEditorStyle::Get(), "ViewportMenu.Label")
 			]
+		]
 
-			// Bottom bar to show current resolution & AR
-			+ SOverlay::Slot()
-			.HAlign(HAlign_Fill)
-			.VAlign(VAlign_Bottom)
+		+ SHorizontalBox::Slot()
+		.AutoWidth()
+		.Padding(5.0f)
+		.VAlign(VAlign_Center)
+		[
+			SNew(SNumericEntryBox<int32>)
+			.AllowSpin(true)
+			.Delta(1)
+			.MinSliderValue(1)
+			.MinValue(1)
+			.MaxSliderValue(TOptional<int32>(1000))
+			.Value(this, &SDesignerView::GetCustomResolutionWidth)
+			.OnValueChanged(this, &SDesignerView::OnCustomResolutionWidthChanged)
+			.Visibility(this, &SDesignerView::GetCustomResolutionEntryVisibility)
+			.MinDesiredValueWidth(50)
+			.LabelPadding(0)
+			.Label()
 			[
-				SNew(SHorizontalBox)
+				SNumericEntryBox<int32>::BuildLabel(LOCTEXT("Width", "Width"), FLinearColor::White, SNumericEntryBox<int32>::RedLabelBackgroundColor)
+			]
+		]
 
-				+ SHorizontalBox::Slot()
-				.AutoWidth()
-				.Padding(6, 0, 0, 2)
-				[
-					SNew(STextBlock)
-					.Visibility(this, &SDesignerView::GetResolutionTextVisibility)
-					.TextStyle(FEditorStyle::Get(), "Graph.ZoomText")
-					.Text(this, &SDesignerView::GetCurrentResolutionText)
-					.ColorAndOpacity(this, &SDesignerView::GetResolutionTextColorAndOpacity)
-				]
+		+ SHorizontalBox::Slot()
+		.AutoWidth()
+		.Padding(5.0f)
+		.VAlign(VAlign_Center)
+		[
+			SNew(SNumericEntryBox<int32>)
+			.AllowSpin(true)
+			.Delta(1)
+			.MinSliderValue(1)
+			.MaxSliderValue(TOptional<int32>(1000))
+			.MinValue(1)
+			.Value(this, &SDesignerView::GetCustomResolutionHeight)
+			.OnValueChanged(this, &SDesignerView::OnCustomResolutionHeightChanged)
+			.Visibility(this, &SDesignerView::GetCustomResolutionEntryVisibility)
+			.MinDesiredValueWidth(50)
+			.LabelPadding(0)
+			.Label()
+			[
+				SNumericEntryBox<int32>::BuildLabel(LOCTEXT("Height", "Height"), FLinearColor::White, SNumericEntryBox<int32>::GreenLabelBackgroundColor)
+			]
+		]
+	]
+
+	// Info Bar, displays heads up information about some actions.
+	+ SOverlay::Slot()
+	.HAlign(HAlign_Fill)
+	.VAlign(VAlign_Bottom)
+	[
+		SNew(SDisappearingBar)
+		[
+			SNew(SBorder)
+			.BorderImage(FEditorStyle::GetBrush("WhiteBrush"))
+			.BorderBackgroundColor(FLinearColor(0.10, 0.10, 0.10, 0.75))
+			.HAlign(HAlign_Center)
+			.VAlign(VAlign_Center)
+			.Padding(FMargin(0, 5))
+			.Visibility(this, &SDesignerView::GetInfoBarVisibility)
+			[
+				SNew(STextBlock)
+				.TextStyle(FEditorStyle::Get(), "Graph.ZoomText")
+				.Text(this, &SDesignerView::GetInfoBarText)
+			]
+		]
+	]
+
+	// Bottom bar to show current resolution & AR
+	+ SOverlay::Slot()
+	.HAlign(HAlign_Fill)
+	.VAlign(VAlign_Bottom)
+	[
+		SNew(SHorizontalBox)
+
+		+ SHorizontalBox::Slot()
+		.AutoWidth()
+		.Padding(6, 0, 0, 2)
+		[
+			SNew(STextBlock)
+			.Visibility(this, &SDesignerView::GetResolutionTextVisibility)
+			.TextStyle(FEditorStyle::Get(), "Graph.ZoomText")
+			.Text(this, &SDesignerView::GetCurrentResolutionText)
+			.ColorAndOpacity(this, &SDesignerView::GetResolutionTextColorAndOpacity)
+		]
 
-				+ SHorizontalBox::Slot()
-				.FillWidth(1.0f)
-				.HAlign(HAlign_Right)
-				.Padding(0, 0, 6, 2)
-				[
-					SNew(SHorizontalBox)
+		+ SHorizontalBox::Slot()
+		.FillWidth(1.0f)
+		.HAlign(HAlign_Right)
+		.Padding(0, 0, 6, 2)
+		[
+			SNew(SHorizontalBox)
 
-					+ SHorizontalBox::Slot()
-					.AutoWidth()
-					[
-						SNew(STextBlock)
-						.TextStyle(FEditorStyle::Get(), "Graph.ZoomText")
-						.Text(this, &SDesignerView::GetCurrentDPIScaleText)
-						.ColorAndOpacity(FLinearColor(1, 1, 1, 0.25f))
-					]
+			+ SHorizontalBox::Slot()
+			.AutoWidth()
+			[
+				SNew(STextBlock)
+				.TextStyle(FEditorStyle::Get(), "Graph.ZoomText")
+				.Text(this, &SDesignerView::GetCurrentDPIScaleText)
+				.ColorAndOpacity(FLinearColor(1, 1, 1, 0.25f))
+			]
 
-					+ SHorizontalBox::Slot()
-					.AutoWidth()
-					.Padding(6, 0, 0, 0)
-					[
-						SNew(SButton)
-						.ButtonStyle(FEditorStyle::Get(), "HoverHintOnly")
-						.ContentPadding(FMargin(3, 1))
-						.OnClicked(this, &SDesignerView::HandleDPISettingsClicked)
-						.ToolTipText(LOCTEXT("DPISettingsTooltip", "Configure the UI Scale Curve to control how the UI is scaled on different resolutions."))
-						.HAlign(HAlign_Center)
-						.VAlign(VAlign_Center)
-						[
-							SNew(SImage)
-							.Image(FEditorStyle::GetBrush("UMGEditor.DPISettings"))
-						]
-					]
+			+ SHorizontalBox::Slot()
+			.AutoWidth()
+			.Padding(6, 0, 0, 0)
+			[
+				SNew(SButton)
+				.ButtonStyle(FEditorStyle::Get(), "HoverHintOnly")
+				.ContentPadding(FMargin(3, 1))
+				.OnClicked(this, &SDesignerView::HandleDPISettingsClicked)
+				.ToolTipText(LOCTEXT("DPISettingsTooltip", "Configure the UI Scale Curve to control how the UI is scaled on different resolutions."))
+				.HAlign(HAlign_Center)
+				.VAlign(VAlign_Center)
+				[
+					SNew(SImage)
+					.Image(FEditorStyle::GetBrush("UMGEditor.DPISettings"))
 				]
 			]
 		]
-	);
-
-	BlueprintEditor.Pin()->OnSelectedWidgetsChanged.AddRaw(this, &SDesignerView::OnEditorSelectionChanged);
-
-	ZoomToFit(/*bInstantZoom*/ true);
+	];
 }
 
 SDesignerView::~SDesignerView()
@@ -1019,6 +1062,15 @@ FReply SDesignerView::OnMouseButtonDown(const FGeometry& MyGeometry, const FPoin
 			DraggingStartPositionScreenSpace = MouseEvent.GetScreenSpacePosition();
 		}
 	}
+	else
+	{
+		// Clear the selection immediately if we didn't click anything.
+		if(MouseEvent.GetEffectingButton() == EKeys::LeftMouseButton)
+		{
+			TSet<FWidgetReference> SelectedTemplates;
+			BlueprintEditor.Pin()->SelectWidgets(SelectedTemplates, false);
+		}
+	}
 
 	// Capture mouse for the drag handle and general mouse actions
 	return FReply::Handled().PreventThrottling().SetUserFocus(AsShared(), EFocusCause::Mouse).CaptureMouse(AsShared());
@@ -1118,14 +1170,14 @@ FReply SDesignerView::OnMouseMove(const FGeometry& MyGeometry, const FPointerEve
 
 void SDesignerView::OnMouseEnter(const FGeometry& MyGeometry, const FPointerEvent& MouseEvent)
 {
-	SCompoundWidget::OnMouseEnter(MyGeometry, MouseEvent);
+	SDesignSurface::OnMouseEnter(MyGeometry, MouseEvent);
 
 	BlueprintEditor.Pin()->ClearHoveredWidget();
 }
 
 void SDesignerView::OnMouseLeave(const FPointerEvent& MouseEvent)
 {
-	SCompoundWidget::OnMouseLeave(MouseEvent);
+	SDesignSurface::OnMouseLeave(MouseEvent);
 
 	BlueprintEditor.Pin()->ClearHoveredWidget();
 }
@@ -1310,11 +1362,35 @@ void SDesignerView::Tick(const FGeometry& AllottedGeometry, const double InCurre
 		Ext->Tick(AllottedGeometry, InCurrentTime, InDeltaTime);
 	}
 
+	// Compute the origin in absolute space.
+	FGeometry RootGeometry = CachedWidgetGeometry.FindChecked(PreviewSurface.ToSharedRef()).Geometry;
+	FVector2D AbsoluteOrigin = MakeGeometryWindowLocal(RootGeometry).LocalToAbsolute(FVector2D::ZeroVector);
+
+	TopRuler->SetRuling(AbsoluteOrigin, 1.0f / GetPreviewScale());
+	SideRuler->SetRuling(AbsoluteOrigin, 1.0f / GetPreviewScale());
+
+	if ( IsHovered() )
+	{
+		// Get cursor in absolute window space.
+		FVector2D CursorPos = FSlateApplication::Get().GetCursorPos();
+		CursorPos = MakeGeometryWindowLocal(RootGeometry).LocalToAbsolute(RootGeometry.AbsoluteToLocal(CursorPos));
+
+		TopRuler->SetCursor(CursorPos);
+		SideRuler->SetCursor(CursorPos);
+	}
+	else
+	{
+		TopRuler->SetCursor(TOptional<FVector2D>());
+		SideRuler->SetCursor(TOptional<FVector2D>());
+	}
+
 	SDesignSurface::Tick(AllottedGeometry, InCurrentTime, InDeltaTime);
 }
 
 FReply SDesignerView::OnDragDetected(const FGeometry& MyGeometry, const FPointerEvent& MouseEvent)
 {
+	SDesignSurface::OnDragDetected(MyGeometry, MouseEvent);
+
 	FWidgetReference SelectedWidget = GetSelectedWidget();
 
 	if ( SelectedWidget.IsValid() )
@@ -1337,7 +1413,7 @@ FReply SDesignerView::OnDragDetected(const FGeometry& MyGeometry, const FPointer
 
 void SDesignerView::OnDragEnter(const FGeometry& MyGeometry, const FDragDropEvent& DragDropEvent)
 {
-	SCompoundWidget::OnDragEnter(MyGeometry, DragDropEvent);
+	SDesignSurface::OnDragEnter(MyGeometry, DragDropEvent);
 
 	BlueprintEditor.Pin()->ClearHoveredWidget();
 
@@ -1346,7 +1422,7 @@ void SDesignerView::OnDragEnter(const FGeometry& MyGeometry, const FDragDropEven
 
 void SDesignerView::OnDragLeave(const FDragDropEvent& DragDropEvent)
 {
-	SCompoundWidget::OnDragLeave(DragDropEvent);
+	SDesignSurface::OnDragLeave(DragDropEvent);
 
 	BlueprintEditor.Pin()->ClearHoveredWidget();
 
@@ -1372,6 +1448,8 @@ void SDesignerView::OnDragLeave(const FDragDropEvent& DragDropEvent)
 
 FReply SDesignerView::OnDragOver(const FGeometry& MyGeometry, const FDragDropEvent& DragDropEvent)
 {
+	SDesignSurface::OnDragOver(MyGeometry, DragDropEvent);
+
 	UWidgetBlueprint* BP = GetBlueprint();
 	
 	if (DropPreviewWidget)
@@ -1708,6 +1786,8 @@ UWidget* SDesignerView::ProcessDropAndAddWidget(const FGeometry& MyGeometry, con
 
 FReply SDesignerView::OnDrop(const FGeometry& MyGeometry, const FDragDropEvent& DragDropEvent)
 {
+	SDesignSurface::OnDrop(MyGeometry, DragDropEvent);
+
 	bMovingExistingWidget = false;
 
 	UWidgetBlueprint* BP = GetBlueprint();
@@ -2005,4 +2085,9 @@ FReply SDesignerView::HandleZoomToFitClicked()
 	return FReply::Handled();
 }
 
+EVisibility SDesignerView::GetRulerVisibility() const
+{
+	return EVisibility::Visible;
+}
+
 #undef LOCTEXT_NAMESPACE
diff --git a/Engine/Source/Editor/UMGEditor/Private/Designer/SDesignerView.h b/Engine/Source/Editor/UMGEditor/Private/Designer/SDesignerView.h
--- a/Engine/Source/Editor/UMGEditor/Private/Designer/SDesignerView.h
+++ b/Engine/Source/Editor/UMGEditor/Private/Designer/SDesignerView.h
@@ -22,6 +22,7 @@ namespace EDesignerMessage
 class FDesignerExtension;
 class UPanelWidget;
 class UUserWidget;
+class SRuler;
 
 /**
  * The designer for widgets.  Allows for laying out widgets in a drag and drop environment.
@@ -36,6 +37,8 @@ public:
 	void Construct(const FArguments& InArgs, TSharedPtr<class FWidgetBlueprintEditor> InBlueprintEditor);
 	virtual ~SDesignerView();
 
+	TSharedRef<SWidget> CreateOverlayUI();
+
 	// SWidget interface
 	virtual FReply OnMouseButtonDown(const FGeometry& MyGeometry, const FPointerEvent& MouseEvent) override;
 	virtual FReply OnMouseButtonUp(const FGeometry& MyGeometry, const FPointerEvent& MouseEvent) override;
@@ -165,6 +168,7 @@ private:
 
 private:
 	FReply HandleZoomToFitClicked();
+	EVisibility GetRulerVisibility() const;
 
 private:
 	static const FString ConfigSectionName;
@@ -245,6 +249,12 @@ private:
 	/**  */
 	FWeakWidgetPath SelectedWidgetPath;
 
+	/** The ruler bar at the top of the designer. */
+	TSharedPtr<SRuler> TopRuler;
+
+	/** The ruler bar on the left side of the designer. */
+	TSharedPtr<SRuler> SideRuler;
+
 	/** */
 	EDesignerMessage::Type DesignerMessage;
 
diff --git a/Engine/Source/Editor/UMGEditor/Private/Designer/SRuler.cpp b/Engine/Source/Editor/UMGEditor/Private/Designer/SRuler.cpp
new file mode 100644
--- /dev/null
+++ b/Engine/Source/Editor/UMGEditor/Private/Designer/SRuler.cpp
@@ -0,0 +1,379 @@
+// Copyright 1998-2015 Epic Games, Inc. All Rights Reserved.
+
+#include "UMGEditorPrivatePCH.h"
+
+#include "SRuler.h"
+
+#define LOCTEXT_NAMESPACE "UMG"
+
+namespace ScrubConstants
+{
+	/** The minimum amount of pixels between each major ticks on the widget */
+	const int32 MinPixelsPerDisplayTick = 5;
+
+	/**The smallest number of units between between major tick marks */
+	const float MinDisplayTickSpacing = 0.1f;
+}
+
+/** Utility struct for converting between scrub range space and local/absolute screen space */
+struct FScrubRangeToScreen
+{
+	float RulerLengthSlateUnits;
+
+	TRange<float> ViewInput;
+	float ViewInputRange;
+	float PixelsPerInput;
+
+	FScrubRangeToScreen(TRange<float> InViewInput, const float InRulerLengthSlateUnits)
+	{
+		RulerLengthSlateUnits = InRulerLengthSlateUnits;
+
+		ViewInput = InViewInput;
+		ViewInputRange = ViewInput.Size<float>();
+		PixelsPerInput = ViewInputRange > 0 ? ( RulerLengthSlateUnits / ViewInputRange ) : 0;
+
+		// Round the float to the nearest four decimals, this keeps it stable as the user moves the ruler around
+		// otherwise small variations will cause the spacing to change wildly as the user pans.
+		PixelsPerInput = FMath::RoundToFloat(PixelsPerInput * 10000.0f) / 10000.0f;
+	}
+
+	/** Local Widget Space -> Curve Input domain. */
+	float LocalXToInput(float ScreenX) const
+	{
+		float LocalX = ScreenX;
+		return ( LocalX / PixelsPerInput ) + ViewInput.GetLowerBoundValue();
+	}
+
+	/** Curve Input domain -> local Widget Space */
+	float InputToLocalX(float Input) const
+	{
+		return ( Input - ViewInput.GetLowerBoundValue() ) * PixelsPerInput;
+	}
+};
+
+struct FDrawTickArgs
+{
+	/** Geometry of the area */
+	FGeometry AllottedGeometry;
+	/** Clipping rect of the area */
+	FSlateRect ClippingRect;
+	/** Color of each tick */
+	FLinearColor TickColor;
+	/** Color of each tick */
+	FLinearColor TextColor;
+	/** Offset in Y where to start the tick */
+	float TickOffset;
+	/** Height in of major ticks */
+	float MajorTickHeight;
+	/** Start layer for elements */
+	int32 StartLayer;
+	/** Draw effects to apply */
+	ESlateDrawEffect::Type DrawEffects;
+	/** Whether or not to only draw major ticks */
+	bool bOnlyDrawMajorTicks;
+};
+
+/**
+ * Gets the the next spacing value in the series 
+ * to determine a good spacing value
+ * E.g, .001,.005,.010,.050,.100,.500,1.000,etc
+ */
+static float GetNextSpacing( uint32 CurrentStep )
+{
+	if(CurrentStep & 0x01) 
+	{
+		// Odd numbers
+		return FMath::Pow( 10.f, 0.5f*((float)(CurrentStep-1)) + 1.f );
+	}
+	else 
+	{
+		// Even numbers
+		return 0.5f * FMath::Pow( 10.f, 0.5f*((float)(CurrentStep)) + 1.f );
+	}
+}
+
+/**
+ * Determines the optimal spacing between tick marks in the slider for a given pixel density
+ * Increments until a minimum amount of slate units specified by MinTick is reached
+ * 
+ * @param InPixelsPerInput	The density of pixels between each input
+ * @param MinTick			The minimum slate units per tick allowed
+ * @param MinTickSpacing	The minimum tick spacing in time units allowed
+ * @return the optimal spacing in time units
+ */
+float DetermineOptimalSpacing( float InPixelsPerInput, uint32 MinTick, float MinTickSpacing )
+{
+	uint32 CurStep = 0;
+
+	// Start with the smallest spacing
+	float Spacing = MinTickSpacing;
+
+	while( Spacing * InPixelsPerInput < MinTick )
+	{
+		Spacing = MinTickSpacing * GetNextSpacing( CurStep );
+		CurStep++;
+	}
+
+	return Spacing;
+}
+
+/////////////////////////////////////////////////////
+// SRuler
+
+void SRuler::Construct(const SRuler::FArguments& InArgs)
+{
+	Orientation = InArgs._Orientation;
+	AbsoluteOrigin = FVector2D(0, 0);
+	SlateToUnitScale = 1;
+
+	MouseButtonDownHandler = InArgs._OnMouseButtonDown;
+}
+
+void SRuler::SetRuling(FVector2D InAbsoluteOrigin, float InSlateToUnitScale)
+{
+	AbsoluteOrigin = InAbsoluteOrigin;
+	SlateToUnitScale = InSlateToUnitScale;
+}
+
+void SRuler::SetCursor(TOptional<FVector2D> InAbsoluteCursor)
+{
+	AbsoluteCursor = InAbsoluteCursor;
+}
+
+int32 SRuler::DrawTicks( FSlateWindowElementList& OutDrawElements, const struct FScrubRangeToScreen& RangeToScreen, FDrawTickArgs& InArgs ) const
+{
+	const float Spacing = DetermineOptimalSpacing( RangeToScreen.PixelsPerInput, ScrubConstants::MinPixelsPerDisplayTick, ScrubConstants::MinDisplayTickSpacing );
+
+	// Sub divisions
+	// @todo Sequencer may need more robust calculation
+	const int32 Divider = 10;
+	// For slightly larger halfway tick mark
+	const int32 HalfDivider = Divider / 2;
+	// Find out where to start from
+	int32 OffsetNum = FMath::FloorToInt(RangeToScreen.ViewInput.GetLowerBoundValue() / Spacing);
+	
+	FSlateFontInfo SmallLayoutFont( FPaths::EngineContentDir() / TEXT("Slate/Fonts/Roboto-Regular.ttf"), 7 );
+
+	TArray<FVector2D> LinePoints;
+	LinePoints.AddUninitialized(2);
+
+	float Number = 0;
+	while ( ( Number = OffsetNum*Spacing ) < RangeToScreen.ViewInput.GetUpperBoundValue() )
+	{
+		// X position local to start of the widget area
+		float XPos = RangeToScreen.InputToLocalX(Number);
+		uint32 AbsOffsetNum = FMath::Abs(OffsetNum);
+
+		if ( AbsOffsetNum % Divider == 0 )
+		{
+			FVector2D Offset = Orientation == Orient_Horizontal ? FVector2D(XPos, InArgs.TickOffset) : FVector2D(InArgs.TickOffset, XPos);
+			FVector2D TickSize = Orientation == Orient_Horizontal ? FVector2D(1.0f, InArgs.MajorTickHeight) : FVector2D(InArgs.MajorTickHeight, 1.0f);
+
+			LinePoints[0] = FVector2D(1.0f, 1.0f);
+			LinePoints[1] = TickSize;
+
+			// lines should not need anti-aliasing
+			const bool bAntiAliasLines = false;
+
+			// Draw each tick mark
+			FSlateDrawElement::MakeLines(
+				OutDrawElements,
+				InArgs.StartLayer,
+				InArgs.AllottedGeometry.ToPaintGeometry( Offset, TickSize ),
+				LinePoints,
+				InArgs.ClippingRect,
+				InArgs.DrawEffects,
+				InArgs.TickColor,
+				false
+				);
+
+			if( !InArgs.bOnlyDrawMajorTicks )
+			{
+				float TextNumber = FMath::Abs(Number);
+
+				FString FrameString = Spacing == ScrubConstants::MinDisplayTickSpacing ? FString::Printf(TEXT("%.1f"), TextNumber) : FString::Printf(TEXT("%.0f"), TextNumber);
+
+				// If the orientation is vertical break the number up over multiple lines.
+				if ( Orientation == Orient_Vertical )
+				{
+					for ( int32 CharIndex = 0; CharIndex < FrameString.Len() - 1; CharIndex += 2 )
+					{
+						FrameString.InsertAt(CharIndex + 1, '\n');
+					}
+				}
+
+				// Space the text between the tick mark but slightly above
+				const TSharedRef< FSlateFontMeasure > FontMeasureService = FSlateApplication::Get().GetRenderer()->GetFontMeasureService();
+				FVector2D TextSize = FontMeasureService->Measure(FrameString, SmallLayoutFont);
+				FVector2D TextOffset = Orientation == Orient_Horizontal ? 
+					FVector2D(XPos - ( TextSize.X*0.5f ), FMath::Abs(InArgs.AllottedGeometry.Size.Y - ( InArgs.MajorTickHeight + TextSize.Y )))
+					:
+					FVector2D(FMath::Abs(InArgs.AllottedGeometry.Size.X - ( InArgs.MajorTickHeight + TextSize.X )), XPos - ( TextSize.Y*0.5f ));
+
+				FSlateDrawElement::MakeText(
+					OutDrawElements,
+					InArgs.StartLayer+1, 
+					InArgs.AllottedGeometry.ToPaintGeometry( TextOffset, TextSize ), 
+					FrameString, 
+					SmallLayoutFont, 
+					InArgs.ClippingRect, 
+					InArgs.DrawEffects,
+					InArgs.TextColor
+				);
+			}
+		}
+		else if( !InArgs.bOnlyDrawMajorTicks )
+		{
+			// Compute the size of each tick mark.  If we are half way between to visible values display a slightly larger tick mark
+			const float MinorTickHeight = AbsOffsetNum % HalfDivider == 0 ? 7.0f : 4.0f;
+
+			FVector2D Offset = Orientation == Orient_Horizontal ? FVector2D(XPos, FMath::Abs(InArgs.AllottedGeometry.Size.Y - MinorTickHeight)) : FVector2D(FMath::Abs(InArgs.AllottedGeometry.Size.X - MinorTickHeight), XPos);
+			FVector2D TickSize = Orientation == Orient_Horizontal ? FVector2D(1, MinorTickHeight) : FVector2D(MinorTickHeight, 1);
+
+			LinePoints[0] = FVector2D(1.0f,1.0f);
+			LinePoints[1] = TickSize;
+
+			const bool bAntiAlias = false;
+			// Draw each sub mark
+			FSlateDrawElement::MakeLines(
+				OutDrawElements,
+				InArgs.StartLayer,
+				InArgs.AllottedGeometry.ToPaintGeometry( Offset, TickSize ),
+				LinePoints,
+				InArgs.ClippingRect,
+				InArgs.DrawEffects,
+				InArgs.TickColor,
+				bAntiAlias
+			);
+		}
+		// Advance to next tick mark
+		++OffsetNum;
+	}
+
+	// Draw line that runs along the bottom of all the ticks.
+	{
+		LinePoints[0] = Orientation == Orient_Horizontal ? FVector2D(0, InArgs.AllottedGeometry.Size.Y) : FVector2D(InArgs.AllottedGeometry.Size.X, 0);
+		LinePoints[1] = Orientation == Orient_Horizontal ? FVector2D(InArgs.AllottedGeometry.Size.X, InArgs.AllottedGeometry.Size.Y) : FVector2D(InArgs.AllottedGeometry.Size.X, InArgs.AllottedGeometry.Size.Y);
+
+		const bool bAntiAlias = false;
+		// Draw each sub mark
+		FSlateDrawElement::MakeLines(
+			OutDrawElements,
+			InArgs.StartLayer,
+			InArgs.AllottedGeometry.ToPaintGeometry(),
+			LinePoints,
+			InArgs.ClippingRect.ExtendBy(FMargin(1, 1, 1, 1)),
+			InArgs.DrawEffects,
+			InArgs.TickColor,
+			bAntiAlias
+		);
+	}
+
+	// Draw the line that shows where the cursor is.
+	if ( AbsoluteCursor.IsSet() )
+	{
+		FVector2D LocalCursor = InArgs.AllottedGeometry.AbsoluteToLocal(AbsoluteCursor.GetValue());
+
+		LinePoints[0] = Orientation == Orient_Horizontal ? FVector2D(LocalCursor.X, 0) : FVector2D(0, LocalCursor.Y);
+		LinePoints[1] = Orientation == Orient_Horizontal ? FVector2D(LocalCursor.X, InArgs.AllottedGeometry.Size.Y) : FVector2D(InArgs.AllottedGeometry.Size.X, LocalCursor.Y);
+
+		const bool bAntiAlias = false;
+		// Draw each sub mark
+		FSlateDrawElement::MakeLines(
+			OutDrawElements,
+			++InArgs.StartLayer,
+			InArgs.AllottedGeometry.ToPaintGeometry(),
+			LinePoints,
+			InArgs.ClippingRect,
+			InArgs.DrawEffects,
+			FLinearColor(FColor(0x19, 0xD1, 0x19)),
+			bAntiAlias
+		);
+	}
+
+	return InArgs.StartLayer;
+}
+
+int32 SRuler::OnPaint(const FPaintArgs& Args, const FGeometry& AllottedGeometry, const FSlateRect& MyClippingRect, FSlateWindowElementList& OutDrawElements, int32 LayerId, const FWidgetStyle& InWidgetStyle, bool bParentEnabled) const
+{
+	const bool bEnabled = bParentEnabled;
+	const ESlateDrawEffect::Type DrawEffects = bEnabled ? ESlateDrawEffect::None : ESlateDrawEffect::DisabledEffect;
+
+	static const FSlateBrush* WhiteBrush = FCoreStyle::Get().GetBrush("GenericWhiteBox");
+
+	// Draw solid background
+	FSlateDrawElement::MakeBox( 
+		OutDrawElements,
+		LayerId,
+		AllottedGeometry.ToPaintGeometry(),
+		WhiteBrush,
+		MyClippingRect, 
+		DrawEffects,
+		FLinearColor(FColor(48,48,48))
+	);
+
+	FVector2D LocalOrigin = AllottedGeometry.AbsoluteToLocal(AbsoluteOrigin);
+
+	float Scale = SlateToUnitScale < ScrubConstants::MinDisplayTickSpacing ? ScrubConstants::MinDisplayTickSpacing : SlateToUnitScale;
+	
+	float Origin = Orientation == Orient_Horizontal ? LocalOrigin.X : LocalOrigin.Y;
+	float Min = 0 - ( Origin * SlateToUnitScale );
+	float Max = Orientation == Orient_Horizontal ? AllottedGeometry.Size.X : AllottedGeometry.Size.Y;
+	Max = (Max - Origin) * SlateToUnitScale;
+
+	TRange<float> LocalViewRange(Min, Max);
+	FScrubRangeToScreen RangeToScreen(LocalViewRange, Orientation == Orient_Horizontal ? AllottedGeometry.Size.X : AllottedGeometry.Size.Y);
+
+	const float MajorTickHeight = 9.0f;
+
+	FDrawTickArgs TickArgs;
+	TickArgs.AllottedGeometry = AllottedGeometry;
+	TickArgs.bOnlyDrawMajorTicks = false;
+	TickArgs.TickColor = FLinearColor(FColor(97, 97, 97));
+	TickArgs.TextColor = FLinearColor(FColor(165, 165, 165));
+	TickArgs.ClippingRect = MyClippingRect;
+	TickArgs.DrawEffects = DrawEffects;
+	TickArgs.StartLayer = LayerId;
+	TickArgs.TickOffset = Orientation == Orient_Horizontal ? FMath::Abs(AllottedGeometry.Size.Y - MajorTickHeight) : FMath::Abs(AllottedGeometry.Size.X - MajorTickHeight);
+	TickArgs.MajorTickHeight = MajorTickHeight;
+
+	LayerId = DrawTicks(OutDrawElements, RangeToScreen, TickArgs);
+
+	return LayerId;
+}
+
+FReply SRuler::OnMouseButtonDown( const FGeometry& MyGeometry, const FPointerEvent& MouseEvent )
+{
+	if ( MouseButtonDownHandler.IsBound() )
+	{
+		// If a handler is assigned, call it.
+		return MouseButtonDownHandler.Execute(MyGeometry, MouseEvent);
+	}
+	else
+	{
+		// otherwise the event is unhandled.
+		return FReply::Unhandled();
+	}
+}
+
+FReply SRuler::OnMouseButtonUp(const FGeometry& MyGeometry, const FPointerEvent& MouseEvent)
+{
+	return FReply::Unhandled();
+}
+
+FReply SRuler::OnMouseMove(const FGeometry& MyGeometry, const FPointerEvent& MouseEvent)
+{
+	return FReply::Unhandled();
+}
+
+FVector2D SRuler::ComputeDesiredSize() const
+{
+	return Orientation == Orient_Horizontal ? FVector2D(100, 18) : FVector2D(18, 100);
+}
+
+FReply SRuler::OnMouseWheel(const FGeometry& MyGeometry, const FPointerEvent& MouseEvent)
+{
+	return FReply::Unhandled();
+}
+
+#undef LOCTEXT_NAMESPACE
diff --git a/Engine/Source/Editor/UMGEditor/Private/Designer/SRuler.h b/Engine/Source/Editor/UMGEditor/Private/Designer/SRuler.h
new file mode 100644
--- /dev/null
+++ b/Engine/Source/Editor/UMGEditor/Private/Designer/SRuler.h
@@ -0,0 +1,73 @@
+// Copyright 1998-2015 Epic Games, Inc. All Rights Reserved.
+
+#pragma once
+
+class SRuler : public SCompoundWidget
+{
+public:
+
+	SLATE_BEGIN_ARGS(SRuler)
+		: _Orientation(Orient_Horizontal)
+		, _OnMouseButtonDown()
+	{}
+
+	SLATE_ARGUMENT(EOrientation, Orientation)
+
+	SLATE_EVENT(FPointerEventHandler, OnMouseButtonDown)
+
+	SLATE_END_ARGS()
+
+
+	/**
+	 * Construct the widget
+	 * 
+	 * @param InArgs   A declaration from which to construct the widget
+	 */
+	void Construct( const FArguments& InArgs );
+
+	/**
+	 * Sets the rules to use when rendering the ruler.
+	 *
+	 * @param AbsoluteOrigin The origin of the ruler in absolute window space.
+	 * @param SlateToUnitScale The scaler applied to slate units to get the same number 
+	 * of units currently being displayed by the rule marks.
+	 */
+	void SetRuling(FVector2D AbsoluteOrigin, float SlateToUnitScale);
+
+	/**
+	 * Sets the cursor position in window absolute space.
+	 *
+	 * @param AbsoluteCursor The position of the cursor in absolute window space.
+	 */
+	void SetCursor(TOptional<FVector2D> AbsoluteCursor);
+
+protected:
+	// SWidget interface
+	virtual FVector2D ComputeDesiredSize() const override;
+	virtual int32 OnPaint( const FPaintArgs& Args, const FGeometry& AllottedGeometry, const FSlateRect& MyClippingRect, FSlateWindowElementList& OutDrawElements, int32 LayerId, const FWidgetStyle& InWidgetStyle, bool bParentEnabled ) const override;
+	virtual FReply OnMouseButtonDown( const FGeometry& MyGeometry, const FPointerEvent& MouseEvent ) override;
+	virtual FReply OnMouseButtonUp( const FGeometry& MyGeometry, const FPointerEvent& MouseEvent ) override;
+	virtual FReply OnMouseMove( const FGeometry& MyGeometry, const FPointerEvent& MouseEvent ) override;
+	virtual FReply OnMouseWheel( const FGeometry& MyGeometry, const FPointerEvent& MouseEvent ) override;
+
+private:
+	/** Draws the ticks on the track */
+	int32 DrawTicks(FSlateWindowElementList& OutDrawElements, const struct FScrubRangeToScreen& RangeToScreen, struct FDrawTickArgs& InArgs) const;
+
+private:
+	/** The orientation of the ruler */
+	EOrientation Orientation;
+
+	/** The absolute origin of the document being measured.  The 0 on the ruler will start here. */
+	FVector2D AbsoluteOrigin;
+
+	/** The absolute position of the cursor so that little lines can be drawn showing the cursors position on the ruler. */
+	TOptional<FVector2D> AbsoluteCursor;
+
+	/** The current conversion from Slate Pixel To Unit Size. */
+	float SlateToUnitScale;
+
+	/** The public event we expose when the mouse button is pressed onto the ruler. */
+	FPointerEventHandler MouseButtonDownHandler;
+};
+
