UE-6596 - Fixing an editor crash and memory corruption issue that occurs when UMG 3D Components have been placed in the level and the widget they spawn is recompiled with certain changes. Restructuring the way external blueprint compilers are registered so that they can properly wait until all compiling has finished to do re-registering of components...etc.

diff --git a/Engine/Plugins/ScriptPlugin/Source/ScriptEditorPlugin/Private/ScriptEditorPlugin.cpp b/Engine/Plugins/ScriptPlugin/Source/ScriptEditorPlugin/Private/ScriptEditorPlugin.cpp
--- a/Engine/Plugins/ScriptPlugin/Source/ScriptEditorPlugin/Private/ScriptEditorPlugin.cpp
+++ b/Engine/Plugins/ScriptPlugin/Source/ScriptEditorPlugin/Private/ScriptEditorPlugin.cpp
@@ -45,16 +45,19 @@ public:
 /**
 * Script blueprint module
 */
-class FScriptEditorPlugin : public IScriptEditorPlugin
+class FScriptEditorPlugin : public IScriptEditorPlugin, IBlueprintCompiler
 {
 	/** IModuleInterface implementation */
 	virtual void StartupModule() override;
 	virtual void ShutdownModule() override;
 
-private:
+public:
+
+	virtual bool CanCompile(const UBlueprint* Blueprint) override;
 
-	/** Callback for the script compiler */
-	FReply CompileScript(UBlueprint* Blueprint, const FKismetCompilerOptions& CompileOptions, FCompilerResultsLog& Results, TArray<UObject*>* ObjLoaded);
+	virtual void PreCompile(UBlueprint* Blueprint) override;
+	virtual void Compile(UBlueprint* Blueprint, const FKismetCompilerOptions& CompileOptions, FCompilerResultsLog& Results, TArray<UObject*>* ObjLoaded) override;
+	virtual void PostCompile(UBlueprint* Blueprint) override;
 };
 
 IMPLEMENT_MODULE(FScriptEditorPlugin, ScriptEditorPlugin)
@@ -66,25 +69,34 @@ void FScriptEditorPlugin::StartupModule()
 
 	// Register blueprint compiler
 	IKismetCompilerInterface& KismetCompilerModule = FModuleManager::LoadModuleChecked<IKismetCompilerInterface>("KismetCompiler");
-	FBlueprintCompileDelegate CompileDelegate;
-	CompileDelegate.BindRaw(this, &FScriptEditorPlugin::CompileScript);
-	KismetCompilerModule.GetCompilers().Add(CompileDelegate);
+	KismetCompilerModule.GetCompilers().Add(this);
 }
 
 void FScriptEditorPlugin::ShutdownModule()
 {
 }
 
-FReply FScriptEditorPlugin::CompileScript(UBlueprint* Blueprint, const FKismetCompilerOptions& CompileOptions, FCompilerResultsLog& Results, TArray<UObject*>* ObjLoaded)
+bool FScriptEditorPlugin::CanCompile(const UBlueprint* Blueprint)
+{
+	return Cast<UScriptBlueprint>(Blueprint) != nullptr;
+}
+
+void FScriptEditorPlugin::PreCompile(UBlueprint* Blueprint)
+{
+
+}
+
+void FScriptEditorPlugin::Compile(UBlueprint* Blueprint, const FKismetCompilerOptions& CompileOptions, FCompilerResultsLog& Results, TArray<UObject*>* ObjLoaded)
 {
-	if (UScriptBlueprint* ScriptBlueprint = Cast<UScriptBlueprint>(Blueprint))
+	if ( UScriptBlueprint* ScriptBlueprint = Cast<UScriptBlueprint>(Blueprint) )
 	{
 		FScriptBlueprintCompiler Compiler(ScriptBlueprint, Results, CompileOptions, ObjLoaded);
 		Compiler.Compile();
 		check(Compiler.NewClass);
-
-		return FReply::Handled();
 	}
+}
 
-	return FReply::Unhandled();
-}
\ No newline at end of file
+void FScriptEditorPlugin::PostCompile(UBlueprint* Blueprint)
+{
+
+}
diff --git a/Engine/Source/Editor/KismetCompiler/Private/KismetCompilerModule.cpp b/Engine/Source/Editor/KismetCompiler/Private/KismetCompilerModule.cpp
--- a/Engine/Source/Editor/KismetCompiler/Private/KismetCompilerModule.cpp
+++ b/Engine/Source/Editor/KismetCompiler/Private/KismetCompilerModule.cpp
@@ -57,20 +57,19 @@ void FKismet2CompilerModule::CompileBlueprintInner(class UBlueprint* Blueprint,
 	else
 	{
 		// Loop through all external compiler delegates attempting to compile the blueprint.
-		FReply Handled = FReply::Unhandled();
-		for ( FBlueprintCompileDelegate& Compiler : Compilers )
+		bool Compiled = false;
+		for ( IBlueprintCompiler* Compiler : Compilers )
 		{
-			Handled = Compiler.Execute(Blueprint, CompileOptions, Results, ObjLoaded);
-
-			// Don't allow any other compiler to handle it if they reported it was handled.
-			if ( Handled.IsEventHandled() )
+			if ( Compiler->CanCompile(Blueprint) )
 			{
+				Compiled = true;
+				Compiler->Compile(Blueprint, CompileOptions, Results, ObjLoaded);
 				break;
 			}
 		}
 
 		// if no one handles it, then use the default blueprint compiler.
-		if ( !Handled.IsEventHandled() )
+		if ( !Compiled )
 		{
 			if ( UAnimBlueprint* AnimBlueprint = Cast<UAnimBlueprint>(Blueprint) )
 			{
@@ -108,6 +107,11 @@ void FKismet2CompilerModule::CompileBlueprint(class UBlueprint* Blueprint, const
 
 	const bool bIsBrandNewBP = (Blueprint->SkeletonGeneratedClass == NULL) && (Blueprint->GeneratedClass == NULL) && (Blueprint->ParentClass != NULL) && !CompileOptions.bIsDuplicationInstigated;
 
+	for ( IBlueprintCompiler* Compiler : Compilers )
+	{
+		Compiler->PreCompile(Blueprint);
+	}
+
 	if ((CompileOptions.CompileType != EKismetCompileType::BytecodeOnly) && (CompileOptions.CompileType != EKismetCompileType::Cpp))
 	{
 		BP_SCOPED_COMPILER_EVENT_NAME(TEXT("Compile Skeleton Class"));
@@ -182,6 +186,11 @@ void FKismet2CompilerModule::CompileBlueprint(class UBlueprint* Blueprint, const
 		}
 	}
 
+	for ( IBlueprintCompiler* Compiler : Compilers )
+	{
+		Compiler->PostCompile(Blueprint);
+	}
+
 	UPackage* Package = Blueprint->GetOutermost();
 	if( Package )
 	{
diff --git a/Engine/Source/Editor/KismetCompiler/Public/KismetCompilerModule.h b/Engine/Source/Editor/KismetCompiler/Public/KismetCompilerModule.h
--- a/Engine/Source/Editor/KismetCompiler/Public/KismetCompilerModule.h
+++ b/Engine/Source/Editor/KismetCompiler/Public/KismetCompilerModule.h
@@ -12,7 +12,16 @@ class UBlueprint;
 //////////////////////////////////////////////////////////////////////////
 // IKismetCompilerInterface
 
-DECLARE_DELEGATE_RetVal_FourParams(FReply, FBlueprintCompileDelegate, UBlueprint* /*Blueprint*/, const FKismetCompilerOptions& /*CompileOptions*/, FCompilerResultsLog& /*Results*/, TArray<UObject*>* /*ObjLoaded*/);
+class IBlueprintCompiler
+{
+public:
+	virtual void PreCompile(UBlueprint* Blueprint) = 0;
+
+	virtual bool CanCompile(const UBlueprint* Blueprint) = 0;
+	virtual void Compile(UBlueprint* Blueprint, const FKismetCompilerOptions& CompileOptions, FCompilerResultsLog& Results, TArray<UObject*>* ObjLoaded) = 0;
+	
+	virtual void PostCompile(UBlueprint* Blueprint) = 0;
+};
 
 class IKismetCompilerInterface : public IModuleInterface
 {
@@ -51,7 +60,7 @@ public:
 	/**
 	 * Gets a list of all compilers for blueprints.  You can register new compilers through this list.
 	 */
-	virtual TArray<FBlueprintCompileDelegate>& GetCompilers() = 0;
+	virtual TArray<IBlueprintCompiler*>& GetCompilers() = 0;
 };
 
 //////////////////////////////////////////////////////////////////////////
@@ -68,11 +77,11 @@ public:
 	virtual void CompileStructure(class UUserDefinedStruct* Struct, FCompilerResultsLog& Results) override;
 	virtual void RecoverCorruptedBlueprint(class UBlueprint* Blueprint) override;
 	virtual void RemoveBlueprintGeneratedClasses(class UBlueprint* Blueprint) override;
-	virtual TArray<FBlueprintCompileDelegate>& GetCompilers() override { return Compilers; }
+	virtual TArray<IBlueprintCompiler*>& GetCompilers() override { return Compilers; }
 	// End implementation
 private:
 	void CompileBlueprintInner(class UBlueprint* Blueprint, const FKismetCompilerOptions& CompileOptions, FCompilerResultsLog& Results, TArray<UObject*>* ObjLoaded);
 
-	TArray<FBlueprintCompileDelegate> Compilers;
+	TArray<IBlueprintCompiler*> Compilers;
 };
 
diff --git a/Engine/Source/Editor/UMGEditor/Private/UMGEditorModule.cpp b/Engine/Source/Editor/UMGEditor/Private/UMGEditorModule.cpp
--- a/Engine/Source/Editor/UMGEditor/Private/UMGEditorModule.cpp
+++ b/Engine/Source/Editor/UMGEditor/Private/UMGEditorModule.cpp
@@ -21,12 +21,13 @@
 
 const FName UMGEditorAppIdentifier = FName(TEXT("UMGEditorApp"));
 
-class FUMGEditorModule : public IUMGEditorModule
+class FUMGEditorModule : public IUMGEditorModule, public IBlueprintCompiler
 {
 public:
 	/** Constructor, set up console commands and variables **/
 	FUMGEditorModule()
 	{
+		ReRegister = nullptr;
 	}
 
 	/** Called right after the module DLL has been loaded and the module object has been created */
@@ -44,9 +45,7 @@ public:
 
 		// Register widget blueprint compiler we do this no matter what.
 		IKismetCompilerInterface& KismetCompilerModule = FModuleManager::LoadModuleChecked<IKismetCompilerInterface>("KismetCompiler");
-		FBlueprintCompileDelegate sd;
-		sd.BindRaw(this, &FUMGEditorModule::CompileWidgetBlueprint);
-		KismetCompilerModule.GetCompilers().Add(sd);
+		KismetCompilerModule.GetCompilers().Add(this);
 
 		// Register asset types
 		IAssetTools& AssetTools = FModuleManager::LoadModuleChecked<FAssetToolsModule>("AssetTools").Get();
@@ -76,27 +75,34 @@ public:
 		CreatedAssetTypeActions.Empty();
 	}
 
-	FReply CompileWidgetBlueprint(UBlueprint* Blueprint, const FKismetCompilerOptions& CompileOptions, FCompilerResultsLog& Results, TArray<UObject*>* ObjLoaded)
+	bool CanCompile(const UBlueprint* Blueprint)
 	{
-		if ( UWidgetBlueprint* WidgetBlueprint = Cast<UWidgetBlueprint>(Blueprint) )
-		{
-			{
-				TComponentReregisterContext<UWidgetComponent> ComponentReregisterContext;
-
-				FWidgetBlueprintCompiler Compiler(WidgetBlueprint, Results, CompileOptions, ObjLoaded);
-				Compiler.Compile();
-				check(Compiler.NewClass);
-			}
+		return Cast<UWidgetBlueprint>(Blueprint) != nullptr;
+	}
 
-			if ( GIsEditor && GEditor )
-			{
-				GEditor->RedrawAllViewports(true);
-			}
+	void PreCompile(UBlueprint* Blueprint)
+	{
+		ReRegister = new TComponentReregisterContext<UWidgetComponent>();
+	}
 
-			return FReply::Handled();
+	void Compile(UBlueprint* Blueprint, const FKismetCompilerOptions& CompileOptions, FCompilerResultsLog& Results, TArray<UObject*>* ObjLoaded)
+	{
+		if ( UWidgetBlueprint* WidgetBlueprint = CastChecked<UWidgetBlueprint>(Blueprint) )
+		{
+			FWidgetBlueprintCompiler Compiler(WidgetBlueprint, Results, CompileOptions, ObjLoaded);
+			Compiler.Compile();
+			check(Compiler.NewClass);
 		}
+	}
 
-		return FReply::Unhandled();
+	void PostCompile(UBlueprint* Blueprint)
+	{
+		delete ReRegister;
+		
+		if ( GIsEditor && GEditor )
+		{
+			GEditor->RedrawAllViewports(true);
+		}
 	}
 
 	/** Gets the extensibility managers for outside entities to extend gui page editor's menus and toolbars */
@@ -116,6 +122,8 @@ private:
 
 	/** All created asset type actions.  Cached here so that we can unregister it during shutdown. */
 	TArray< TSharedPtr<IAssetTypeActions> > CreatedAssetTypeActions;
+
+	TComponentReregisterContext<UWidgetComponent>* ReRegister;
 };
 
 IMPLEMENT_MODULE(FUMGEditorModule, UMGEditor);
