UMG Editor - The disappearing bar used in the editor is now hit test invisible.

diff --git a/Engine/Source/Editor/UMGEditor/Private/Designer/SDisappearingBar.h b/Engine/Source/Editor/UMGEditor/Private/Designer/SDisappearingBar.h
--- a/Engine/Source/Editor/UMGEditor/Private/Designer/SDisappearingBar.h
+++ b/Engine/Source/Editor/UMGEditor/Private/Designer/SDisappearingBar.h
@@ -11,9 +11,14 @@ class SDisappearingBar : public SCompoundWidget
 {
 public:
 
-	SLATE_BEGIN_ARGS(SDisappearingBar) {}
+	SLATE_BEGIN_ARGS(SDisappearingBar)
+	{
+		_Visibility = EVisibility::HitTestInvisible;
+	}
+
 		/** Slot for this designers content (optional) */
 		SLATE_DEFAULT_SLOT(FArguments, Content)
+
 	SLATE_END_ARGS()
 
 	void Construct(const FArguments& InArgs);
