UMG - Fixing lots of small usability problems in the designer.  Copy / Paste of widgets now attempts to recreate slot information of the root items into their new container.  AKA, things copy and pasted in the canvas panel don't show up the same size; now they do.  Lots of weird bugs when resizing, moving, re-anchoring widgets have been fixed.

diff --git a/Engine/Source/Editor/UMGEditor/Private/Designer/SDesignerView.cpp b/Engine/Source/Editor/UMGEditor/Private/Designer/SDesignerView.cpp
--- a/Engine/Source/Editor/UMGEditor/Private/Designer/SDesignerView.cpp
+++ b/Engine/Source/Editor/UMGEditor/Private/Designer/SDesignerView.cpp
@@ -1356,6 +1356,21 @@ void SDesignerView::Tick(const FGeometry& AllottedGeometry, const double InCurre
 	FArrangedWidget WindowWidgetGeometry(PreviewHitTestRoot.ToSharedRef(), AllottedGeometry);
 	PopulateWidgetGeometryCache(WindowWidgetGeometry);
 
+	TArray< TFunction<void()> >& QueuedActions = BlueprintEditor.Pin()->GetQueuedDesignerActions();
+	for ( TFunction<void()>& Action : QueuedActions )
+	{
+		Action();
+	}
+
+	if ( QueuedActions.Num() > 0 )
+	{
+		QueuedActions.Reset();
+
+		CachedWidgetGeometry.Reset();
+		FArrangedWidget WindowWidgetGeometry(PreviewHitTestRoot.ToSharedRef(), AllottedGeometry);
+		PopulateWidgetGeometryCache(WindowWidgetGeometry);
+	}
+
 	// Tick all designer extensions in case they need to update widgets
 	for ( const TSharedRef<FDesignerExtension>& Ext : DesignerExtensions )
 	{
@@ -1568,23 +1583,27 @@ UWidget* SDesignerView::ProcessDropAndAddWidget(const FGeometry& MyGeometry, con
 			FVector2D LocalPosition = WidgetUnderCursorGeometry.AbsoluteToLocal(DragDropEvent.GetScreenSpacePosition());
 			if ( UPanelSlot* Slot = Parent->AddChild(Widget) )
 			{
-				// HACK UMG - This seems like a bad idea to call TakeWidget
-				TSharedPtr<SWidget> SlateWidget = Widget->TakeWidget();
-				SlateWidget->SlatePrepass();
-				const FVector2D& WidgetDesiredSize = SlateWidget->GetDesiredSize();
+				// Special logic for canvas panel slots.
+				if ( UCanvasPanelSlot* CanvasSlot = Cast<UCanvasPanelSlot>(Slot) )
+				{
+					// HACK UMG - This seems like a bad idea to call TakeWidget
+					TSharedPtr<SWidget> SlateWidget = Widget->TakeWidget();
+					SlateWidget->SlatePrepass();
+					const FVector2D& WidgetDesiredSize = SlateWidget->GetDesiredSize();
 
-				static const FVector2D MinimumDefaultSize(100, 40);
-				FVector2D LocalSize = FVector2D(FMath::Max(WidgetDesiredSize.X, MinimumDefaultSize.X), FMath::Max(WidgetDesiredSize.Y, MinimumDefaultSize.Y));
+					static const FVector2D MinimumDefaultSize(100, 40);
+					FVector2D LocalSize = FVector2D(FMath::Max(WidgetDesiredSize.X, MinimumDefaultSize.X), FMath::Max(WidgetDesiredSize.Y, MinimumDefaultSize.Y));
 
-				const UWidgetDesignerSettings* DesignerSettings = GetDefault<UWidgetDesignerSettings>();
-				if ( DesignerSettings->GridSnapEnabled )
-				{
-					LocalPosition.X = ( (int32)LocalPosition.X ) - (( (int32)LocalPosition.X ) % DesignerSettings->GridSnapSize);
-					LocalPosition.Y = ( (int32)LocalPosition.Y ) - (( (int32)LocalPosition.Y ) % DesignerSettings->GridSnapSize);
-				}
+					const UWidgetDesignerSettings* DesignerSettings = GetDefault<UWidgetDesignerSettings>();
+					if ( DesignerSettings->GridSnapEnabled )
+					{
+						LocalPosition.X = ( (int32)LocalPosition.X ) - ( ( (int32)LocalPosition.X ) % DesignerSettings->GridSnapSize );
+						LocalPosition.Y = ( (int32)LocalPosition.Y ) - ( ( (int32)LocalPosition.Y ) % DesignerSettings->GridSnapSize );
+					}
 
-				Slot->SetDesiredPosition(LocalPosition);
-				Slot->SetDesiredSize(LocalSize);
+					CanvasSlot->SetPosition(LocalPosition);
+					CanvasSlot->SetSize(LocalSize);
+				}
 
 				DropPreviewParent = Parent;
 
@@ -1736,7 +1755,7 @@ UWidget* SDesignerView::ProcessDropAndAddWidget(const FGeometry& MyGeometry, con
 							FWidgetBlueprintEditorUtils::ImportPropertiesFromText(Slot, SelectedDragDropOp->ExportedSlotProperties);
 
 							CanvasSlot->SaveBaseLayout();
-							Slot->SetDesiredPosition(NewPosition);
+							CanvasSlot->SetDesiredPosition(NewPosition);
 							CanvasSlot->RebaseLayout();
 
 							FWidgetBlueprintEditorUtils::ExportPropertiesToText(Slot, SelectedDragDropOp->ExportedSlotProperties);
@@ -1749,7 +1768,6 @@ UWidget* SDesignerView::ProcessDropAndAddWidget(const FGeometry& MyGeometry, con
 					else
 					{
 						FWidgetBlueprintEditorUtils::ImportPropertiesFromText(Slot, SelectedDragDropOp->ExportedSlotProperties);
-						Slot->SetDesiredPosition(NewPosition);
 					}
 
 					DropPreviewParent = NewParent;
diff --git a/Engine/Source/Editor/UMGEditor/Private/Designer/STransformHandle.cpp b/Engine/Source/Editor/UMGEditor/Private/Designer/STransformHandle.cpp
--- a/Engine/Source/Editor/UMGEditor/Private/Designer/STransformHandle.cpp
+++ b/Engine/Source/Editor/UMGEditor/Private/Designer/STransformHandle.cpp
@@ -52,7 +52,7 @@ EVisibility STransformHandle::GetHandleVisibility() const
 			{
 				if ( UPanelSlot* TemplateSlot = SelectedWidget.GetTemplate()->Slot )
 				{
-					if ( TemplateSlot->CanResize(DragDirection) )
+					if ( CanResize(TemplateSlot, DragDirection) )
 					{
 						return EVisibility::Visible;
 					}
@@ -74,6 +74,16 @@ FReply STransformHandle::OnMouseButtonDown(const FGeometry& MyGeometry, const FP
 	{
 		Action = ComputeActionAtLocation(MyGeometry, MouseEvent);
 
+		FWidgetReference SelectedWidget = Designer->GetSelectedWidget();
+		UWidget* Preview = SelectedWidget.GetPreview();
+
+		if ( UCanvasPanelSlot* Slot = Cast<UCanvasPanelSlot>(Preview->Slot) )
+		{
+			StartingOffsets = Slot->GetOffsets();
+		}
+
+		MouseDownPosition = MouseEvent.GetScreenSpacePosition();
+
 		return FReply::Handled().CaptureMouse(SharedThis(this));
 	}
 
@@ -103,66 +113,104 @@ FReply STransformHandle::OnMouseMove(const FGeometry& MyGeometry, const FPointer
 		UWidget* Template = SelectedWidget.GetTemplate();
 		UWidget* Preview = SelectedWidget.GetPreview();
 
+		{
+			FVector2D Delta = MouseEvent.GetScreenSpacePosition() - MouseDownPosition;
+			FVector2D TranslateAmount = Delta * ( 1.0f / Designer->GetPreviewScale() );
+
+			Resize(Cast<UCanvasPanelSlot>(Preview->Slot), DragDirection, TranslateAmount);
+			Resize(Cast<UCanvasPanelSlot>(Template->Slot), DragDirection, TranslateAmount);
+		}
+
 		ETransformMode::Type TransformMode = Designer->GetTransformMode();
-		if ( TransformMode == ETransformMode::Layout )
+		if ( TransformMode == ETransformMode::Render )
 		{
-			//FGeometry MyGeometry = Designer->GetWidgetGeometry(SelectedWidget);
-			//FGeometry ParentGeometry = Designer->GetWidgetParentGeometry(SelectedWidget);
+			FWidgetTransform RenderTransform = Preview->RenderTransform;
 
-			//FVector2D LocalPosition = MyGeometry.Position;
-			//FVector2D LocalDragPosition = ParentGeometry.AbsoluteToLocal(MouseEvent.GetScreenSpacePosition());
+			static const FName RenderTransformName(TEXT("RenderTransform"));
 
-			//FVector2D NewSize = LocalDragPosition - LocalPosition;
+			FObjectEditorUtils::SetPropertyValue<UWidget, FWidgetTransform>(Preview, RenderTransformName, RenderTransform);
+			FObjectEditorUtils::SetPropertyValue<UWidget, FWidgetTransform>(Template, RenderTransformName, RenderTransform);
+		}
+	}
 
-			FVector2D TranslateAmount = MouseEvent.GetCursorDelta() * ( 1.0f / Designer->GetPreviewScale() );
+	return FReply::Unhandled();
+}
 
-			if ( Preview->Slot )
-			{
-				Preview->Slot->Resize(DragDirection, TranslateAmount);
-			}
+bool STransformHandle::CanResize(UPanelSlot* Slot, const FVector2D& Direction) const
+{
+	return Cast<UCanvasPanelSlot>(Slot) != nullptr;
+}
 
-			if ( Template->Slot )
-			{
-				Template->Slot->Resize(DragDirection, TranslateAmount);
-			}
-		}
-		else if ( TransformMode == ETransformMode::Render )
-		{
-			FWidgetTransform RenderTransform = Preview->RenderTransform;
+void STransformHandle::Resize(UCanvasPanelSlot* Slot, const FVector2D& Direction, const FVector2D& Amount)
+{
+	if ( Slot == nullptr )
+	{
+		return;
+	}
 
-			//if ( Action == ETransformAction::Primary )
-			//{
-				FVector2D TranslateAmount = MouseEvent.GetCursorDelta() * ( 1.0f / Designer->GetPreviewScale() );
+	FMargin Offsets = StartingOffsets;
+	const FAnchorData& LayoutData = Slot->LayoutData;
 
-				if ( Preview->Slot )
-				{
-					Preview->Slot->Resize(DragDirection, TranslateAmount);
-				}
+	FVector2D Movement = Amount * Direction;
+	FVector2D PositionMovement = Movement * ( FVector2D(1.0f, 1.0f) - LayoutData.Alignment );
+	FVector2D SizeMovement = Movement;
 
-				if ( Template->Slot )
-				{
-					Template->Slot->Resize(DragDirection, TranslateAmount);
-				}
-			//}
-			//else
-			//{
-			//	FVector2D RotationDelta = MouseEvent.GetCursorDelta();
-			//	RenderTransform.Angle += RotationDelta.Size();
-			//}
+	if ( Direction.X < 0 )
+	{
+		Offsets.Left -= PositionMovement.X;
 
-			static const FName RenderTransformName(TEXT("RenderTransform"));
+		// Keep the right side stable if it's not stretched and we're resizing left.
+		if ( !LayoutData.Anchors.IsStretchedHorizontal() )
+		{
+			// TODO UMG This doesn't work correctly with Alignment, fix or remove alignment
+			Offsets.Right += SizeMovement.X;
+		}
+	}
 
-			FObjectEditorUtils::SetPropertyValue<UWidget, FWidgetTransform>(Preview, RenderTransformName, RenderTransform);
-			FObjectEditorUtils::SetPropertyValue<UWidget, FWidgetTransform>(Template, RenderTransformName, RenderTransform);
+	if ( Direction.Y < 0 )
+	{
+		Offsets.Top -= PositionMovement.Y;
+
+		// Keep the bottom side stable if it's not stretched and we're resizing top.
+		if ( !LayoutData.Anchors.IsStretchedVertical() )
+		{
+			// TODO UMG This doesn't work correctly with Alignment, fix or remove alignment
+			Offsets.Bottom += SizeMovement.Y;
 		}
 	}
 
-	return FReply::Unhandled();
+	if ( Direction.X > 0 )
+	{
+		if ( LayoutData.Anchors.IsStretchedHorizontal() )
+		{
+			Offsets.Right -= Amount.X * Direction.X;
+		}
+		else
+		{
+			Offsets.Left += (Movement * LayoutData.Alignment).X;
+			Offsets.Right += Amount.X * Direction.X;
+		}
+	}
+
+	if ( Direction.Y > 0 )
+	{
+		if ( LayoutData.Anchors.IsStretchedVertical() )
+		{
+			Offsets.Bottom -= Amount.Y * Direction.Y;
+		}
+		else
+		{
+			Offsets.Top += ( Movement * LayoutData.Alignment ).Y;
+			Offsets.Bottom += Amount.Y * Direction.Y;
+		}
+	}
+
+	Slot->SetOffsets(Offsets);
 }
 
 FCursorReply STransformHandle::OnCursorQuery(const FGeometry& MyGeometry, const FPointerEvent& MouseEvent) const
 {
-	ETransformAction::Type CurrentAction = Action;
+	ETransformAction CurrentAction = Action;
 	if ( CurrentAction == ETransformAction::None )
 	{
 		CurrentAction = ComputeActionAtLocation(MyGeometry, MouseEvent);
@@ -170,28 +218,21 @@ FCursorReply STransformHandle::OnCursorQuery(const FGeometry& MyGeometry, const
 
 	ETransformMode::Type TransformMode = Designer->GetTransformMode();
 
-	//if ( TransformMode == ETransformMode::Layout || CurrentAction == ETransformAction::Primary )
-	//{
-		switch ( TransformDirection )
-		{
-		case ETransformDirection::TopLeft:
-		case ETransformDirection::BottomRight:
-			return FCursorReply::Cursor(EMouseCursor::ResizeSouthEast);
-		case ETransformDirection::TopRight:
-		case ETransformDirection::BottomLeft:
-			return FCursorReply::Cursor(EMouseCursor::ResizeSouthWest);
-		case ETransformDirection::TopCenter:
-		case ETransformDirection::BottomCenter:
-			return FCursorReply::Cursor(EMouseCursor::ResizeUpDown);
-		case ETransformDirection::CenterLeft:
-		case ETransformDirection::CenterRight:
-			return FCursorReply::Cursor(EMouseCursor::ResizeLeftRight);
-		}
-	//}
-	//else
-	//{
-	//	return FCursorReply::Cursor(EMouseCursor::EyeDropper);
-	//}
+	switch ( TransformDirection )
+	{
+	case ETransformDirection::TopLeft:
+	case ETransformDirection::BottomRight:
+		return FCursorReply::Cursor(EMouseCursor::ResizeSouthEast);
+	case ETransformDirection::TopRight:
+	case ETransformDirection::BottomLeft:
+		return FCursorReply::Cursor(EMouseCursor::ResizeSouthWest);
+	case ETransformDirection::TopCenter:
+	case ETransformDirection::BottomCenter:
+		return FCursorReply::Cursor(EMouseCursor::ResizeUpDown);
+	case ETransformDirection::CenterLeft:
+	case ETransformDirection::CenterRight:
+		return FCursorReply::Cursor(EMouseCursor::ResizeLeftRight);
+	}
 
 	return FCursorReply::Unhandled();
 }
@@ -252,7 +293,7 @@ FVector2D STransformHandle::ComputeOrigin(ETransformDirection::Type InTransformD
 	return FVector2D(0, 0);
 }
 
-ETransformAction::Type STransformHandle::ComputeActionAtLocation(const FGeometry& MyGeometry, const FPointerEvent& MouseEvent) const
+ETransformAction STransformHandle::ComputeActionAtLocation(const FGeometry& MyGeometry, const FPointerEvent& MouseEvent) const
 {
 	FVector2D LocalPosition = MyGeometry.AbsoluteToLocal(MouseEvent.GetScreenSpacePosition());
 	FVector2D GrabOriginOffset = LocalPosition - DragOrigin;
diff --git a/Engine/Source/Editor/UMGEditor/Private/Designer/STransformHandle.h b/Engine/Source/Editor/UMGEditor/Private/Designer/STransformHandle.h
--- a/Engine/Source/Editor/UMGEditor/Private/Designer/STransformHandle.h
+++ b/Engine/Source/Editor/UMGEditor/Private/Designer/STransformHandle.h
@@ -24,15 +24,12 @@ namespace ETransformDirection
 }
 
 
-namespace ETransformAction
+enum class ETransformAction
 {
-	enum Type
-	{
-		None,
-		Primary,
-		Secondary
-	};
-}
+	None,
+	Primary,
+	Secondary
+};
 
 class STransformHandle : public SCompoundWidget
 {
@@ -53,7 +50,10 @@ public:
 protected:
 	EVisibility GetHandleVisibility() const;
 
-	ETransformAction::Type ComputeActionAtLocation(const FGeometry& MyGeometry, const FPointerEvent& MouseEvent) const;
+	bool CanResize(UPanelSlot* Slot, const FVector2D& Direction) const;
+	void Resize(UCanvasPanelSlot* Slot, const FVector2D& Direction, const FVector2D& Amount);
+
+	ETransformAction ComputeActionAtLocation(const FGeometry& MyGeometry, const FPointerEvent& MouseEvent) const;
 
 protected:
 	FVector2D ComputeDragDirection(ETransformDirection::Type InTransformDirection) const;
@@ -62,8 +62,11 @@ protected:
 protected:
 	IUMGDesigner* Designer;
 	ETransformDirection::Type TransformDirection;
-	ETransformAction::Type Action;
+	ETransformAction Action;
 
 	FVector2D DragDirection;
 	FVector2D DragOrigin;
+
+	FVector2D MouseDownPosition;
+	FMargin StartingOffsets;
 };
diff --git a/Engine/Source/Editor/UMGEditor/Private/Hierarchy/SHierarchyViewItem.cpp b/Engine/Source/Editor/UMGEditor/Private/Hierarchy/SHierarchyViewItem.cpp
--- a/Engine/Source/Editor/UMGEditor/Private/Hierarchy/SHierarchyViewItem.cpp
+++ b/Engine/Source/Editor/UMGEditor/Private/Hierarchy/SHierarchyViewItem.cpp
@@ -200,7 +200,6 @@ FReply ProcessHierarchyDragDrop(const FDragDropEvent& DragDropEvent, bool bIsDro
 
 				// Import the old slot properties
 				FWidgetBlueprintEditorUtils::ImportPropertiesFromText(NewSlot, HierarchyDragDropOp->ExportedSlotProperties);
-				NewSlot->SetDesiredPosition(FVector2D::ZeroVector);
 				FBlueprintEditorUtils::MarkBlueprintAsStructurallyModified(Blueprint);
 
 				TSet<FWidgetReference> SelectedTemplates;
diff --git a/Engine/Source/Editor/UMGEditor/Private/Utility/WidgetSlotPair.cpp b/Engine/Source/Editor/UMGEditor/Private/Utility/WidgetSlotPair.cpp
new file mode 100644
--- /dev/null
+++ b/Engine/Source/Editor/UMGEditor/Private/Utility/WidgetSlotPair.cpp
@@ -0,0 +1,42 @@
+// Copyright 1998-2015 Epic Games, Inc. All Rights Reserved.
+
+#include "UMGEditorPrivatePCH.h"
+#include "WidgetSlotPair.h"
+
+UWidgetSlotPair::UWidgetSlotPair()
+{
+}
+
+void UWidgetSlotPair::SetWidgetName(FName InWidgetName)
+{
+	WidgetName = InWidgetName;
+}
+
+void UWidgetSlotPair::SetWidget(UWidget* InWidget)
+{
+	WidgetName = InWidget->GetFName();
+
+	TMap<FName, FString> ExportedSlotProperties;
+	FWidgetBlueprintEditorUtils::ExportPropertiesToText(InWidget->Slot, ExportedSlotProperties);
+
+	for ( const auto& Entry : ExportedSlotProperties )
+	{
+		SlotPropertyNames.Add(Entry.Key);
+		SlotPropertyValues.Add(Entry.Value);
+	}
+}
+
+FName UWidgetSlotPair::GetWidgetName() const
+{
+	return WidgetName;
+}
+
+void UWidgetSlotPair::GetSlotProperties(TMap<FName, FString>& OutSlotProperties) const
+{
+	checkSlow(SlotPropertyNames.Num() == SlotPropertyValues.Num());
+
+	for ( int32 PropertyIndex = 0; PropertyIndex < SlotPropertyNames.Num(); PropertyIndex++ )
+	{
+		OutSlotProperties.Add(SlotPropertyNames[PropertyIndex], SlotPropertyValues[PropertyIndex]);
+	}
+}
diff --git a/Engine/Source/Editor/UMGEditor/Private/Utility/WidgetSlotPair.h b/Engine/Source/Editor/UMGEditor/Private/Utility/WidgetSlotPair.h
new file mode 100644
--- /dev/null
+++ b/Engine/Source/Editor/UMGEditor/Private/Utility/WidgetSlotPair.h
@@ -0,0 +1,33 @@
+// Copyright 1998-2015 Epic Games, Inc. All Rights Reserved.
+
+#pragma once
+
+#include "WidgetSlotPair.generated.h"
+
+UCLASS()
+class UWidgetSlotPair : public UObject
+{
+	GENERATED_BODY()
+	
+public:
+
+	UWidgetSlotPair();
+
+	void SetWidget(UWidget* Widget);
+
+	void SetWidgetName(FName WidgetName);
+	FName GetWidgetName() const;
+
+	void GetSlotProperties(TMap<FName, FString>& OutSlotProperties) const;
+
+private:
+
+	UPROPERTY()
+	FName WidgetName;
+
+	UPROPERTY()
+	TArray<FName> SlotPropertyNames;
+
+	UPROPERTY()
+	TArray<FString> SlotPropertyValues;
+};
\ No newline at end of file
diff --git a/Engine/Source/Editor/UMGEditor/Private/WidgetBlueprintEditor.cpp b/Engine/Source/Editor/UMGEditor/Private/WidgetBlueprintEditor.cpp
--- a/Engine/Source/Editor/UMGEditor/Private/WidgetBlueprintEditor.cpp
+++ b/Engine/Source/Editor/UMGEditor/Private/WidgetBlueprintEditor.cpp
@@ -341,7 +341,7 @@ void FWidgetBlueprintEditor::PasteWidgets()
 	TSet<FWidgetReference> Widgets = GetSelectedWidgets();
 	FWidgetReference Target = Widgets.Num() > 0 ? *Widgets.CreateIterator() : FWidgetReference();
 
-	FWidgetBlueprintEditorUtils::PasteWidgets(GetWidgetBlueprintObj(), Target, PasteDropLocation);
+	FWidgetBlueprintEditorUtils::PasteWidgets(SharedThis(this), GetWidgetBlueprintObj(), Target, PasteDropLocation);
 
 	//TODO UMG - Select the newly selected pasted widgets.
 }
@@ -709,4 +709,14 @@ float FWidgetBlueprintEditor::GetHoveredWidgetTime() const
 	return HoverTime;
 }
 
+void FWidgetBlueprintEditor::AddPostDesignerLayoutAction(TFunction<void()> Action)
+{
+	QueuedDesignerActions.Add(Action);
+}
+
+TArray< TFunction<void()> >& FWidgetBlueprintEditor::GetQueuedDesignerActions()
+{
+	return QueuedDesignerActions;
+}
+
 #undef LOCTEXT_NAMESPACE
diff --git a/Engine/Source/Editor/UMGEditor/Private/WidgetBlueprintEditor.h b/Engine/Source/Editor/UMGEditor/Private/WidgetBlueprintEditor.h
--- a/Engine/Source/Editor/UMGEditor/Private/WidgetBlueprintEditor.h
+++ b/Engine/Source/Editor/UMGEditor/Private/WidgetBlueprintEditor.h
@@ -115,6 +115,10 @@ public:
 
 	float GetHoveredWidgetTime() const;
 
+	void AddPostDesignerLayoutAction(TFunction<void()> Action);
+
+	TArray< TFunction<void()> >& GetQueuedDesignerActions();
+
 public:
 	/** Fires whenever the selected set of widgets changing */
 	FOnSelectedWidgetsChanged OnSelectedWidgetsChanging;
@@ -211,4 +215,6 @@ private:
 
 	/**  */
 	bool bIsRealTime;
+
+	TArray< TFunction<void()> > QueuedDesignerActions;
 };
diff --git a/Engine/Source/Editor/UMGEditor/Private/WidgetBlueprintEditorUtils.cpp b/Engine/Source/Editor/UMGEditor/Private/WidgetBlueprintEditorUtils.cpp
--- a/Engine/Source/Editor/UMGEditor/Private/WidgetBlueprintEditorUtils.cpp
+++ b/Engine/Source/Editor/UMGEditor/Private/WidgetBlueprintEditorUtils.cpp
@@ -17,6 +17,7 @@
 #include "ScopedTransaction.h"
 #include "K2Node_ComponentBoundEvent.h"
 #include "CanvasPanel.h"
+#include "WidgetSlotPair.h"
 
 #define LOCTEXT_NAMESPACE "UMG"
 
@@ -34,26 +35,22 @@ public:
 	{
 		const bool bIsWidget = ObjectClass->IsChildOf(UWidget::StaticClass());
 		const bool bIsSlot = ObjectClass->IsChildOf(UPanelSlot::StaticClass());
+		const bool bIsSlotMetaData = ObjectClass->IsChildOf(UWidgetSlotPair::StaticClass());
 
-		return bIsWidget || bIsSlot;
+		return bIsWidget || bIsSlot || bIsSlotMetaData;
 	}
 
 	virtual void ProcessConstructedObject(UObject* NewObject) override
 	{
 		check(NewObject);
 
-		// Add it to the new object map
-		NewObjectMap.Add(NewObject->GetFName(), Cast<UWidget>(NewObject));
-
-		// If this is a scene component and it has a parent
-		UWidget* Widget = Cast<UWidget>(NewObject);
-		if ( Widget && Widget->Slot )
+		if ( UWidget* Widget = Cast<UWidget>(NewObject) )
 		{
-			// Add an entry to the child->parent name map
-//			ParentMap.Add(NewObject->GetFName(), Widget->AttachParent->GetFName());
-
-			// Clear this so it isn't used when constructing the new SCS node
-			//Widget->AttachParent = NULL;
+			NewWidgetMap.Add(Widget->GetFName(), Widget);
+		}
+		else if ( UWidgetSlotPair* SlotMetaData = Cast<UWidgetSlotPair>(NewObject) )
+		{
+			MissingSlotData.Add(SlotMetaData->GetWidgetName(), SlotMetaData);
 		}
 	}
 
@@ -61,11 +58,11 @@ public:
 
 public:
 
-	// Child->Parent name map
-	TMap<FName, FName> ParentMap;
-
 	// Name->Instance object mapping
-	TMap<FName, UWidget*> NewObjectMap;
+	TMap<FName, UWidget*> NewWidgetMap;
+
+	// Instance->OldSlotMetaData that didn't survive the journey because it wasn't copied.
+	TMap<FName, UWidgetSlotPair*> MissingSlotData;
 };
 
 bool FWidgetBlueprintEditorUtils::RenameWidget(TSharedRef<FWidgetBlueprintEditor> BlueprintEditor, const FName& OldName, const FName& NewName)
@@ -451,6 +448,8 @@ void FWidgetBlueprintEditorUtils::CopyWidgets(UWidgetBlueprint* BP, TSet<FWidget
 		UWidget* ParentWidget = Widget.GetTemplate();
 		CopyableWidets.Add(ParentWidget);
 
+		// When copying a widget users expect all sub widgets to be copied as well, so we need to ensure that
+		// we gather all the child widgets and copy them as well.
 		UWidget::GatherAllChildren(ParentWidget, CopyableWidets);
 	}
 
@@ -468,25 +467,33 @@ void FWidgetBlueprintEditorUtils::ExportWidgetsToText(TSet<UWidget*> WidgetsToEx
 	const FExportObjectInnerContext Context;
 
 	// Export each of the selected nodes
-	UObject* LastOuter = NULL;
+	UObject* LastOuter = nullptr;
 	for ( UWidget* Widget : WidgetsToExport )
 	{
 		// The nodes should all be from the same scope
 		UObject* ThisOuter = Widget->GetOuter();
-		check(( LastOuter == ThisOuter ) || ( LastOuter == NULL ));
+		check(( LastOuter == ThisOuter ) || ( LastOuter == nullptr ));
 		LastOuter = ThisOuter;
 
-		UExporter::ExportToOutputDevice(&Context, Widget, NULL, Archive, TEXT("copy"), 0, PPF_ExportsNotFullyQualified | PPF_Copy | PPF_Delimited, false, ThisOuter);
+		UExporter::ExportToOutputDevice(&Context, Widget, nullptr, Archive, TEXT("copy"), 0, PPF_ExportsNotFullyQualified | PPF_Copy | PPF_Delimited, false, ThisOuter);
+
+		if ( Widget->GetParent() == nullptr || !WidgetsToExport.Contains(Widget->GetParent()) )
+		{
+			UWidgetSlotPair* SlotMetaData = new UWidgetSlotPair();
+			SlotMetaData->SetWidget(Widget);
+
+			UExporter::ExportToOutputDevice(&Context, SlotMetaData, nullptr, Archive, TEXT("copy"), 0, PPF_ExportsNotFullyQualified | PPF_Copy | PPF_Delimited, false, nullptr);
+		}
 	}
 
 	ExportedText = Archive;
 }
 
-void FWidgetBlueprintEditorUtils::PasteWidgets(UWidgetBlueprint* BP, FWidgetReference ParentWidgetRef, FVector2D PasteLocation)
+void FWidgetBlueprintEditorUtils::PasteWidgets(TSharedRef<FWidgetBlueprintEditor> BlueprintEditor, UWidgetBlueprint* BP, FWidgetReference ParentWidgetRef, FVector2D PasteLocation)
 {
 	const FScopedTransaction Transaction(FGenericCommands::Get().Paste->GetDescription());
 
-	UPanelWidget* ParentWidget = NULL;
+	UPanelWidget* ParentWidget = nullptr;
 	
 	if ( ParentWidgetRef.IsValid() )
 	{
@@ -511,7 +518,8 @@ void FWidgetBlueprintEditorUtils::PasteWidgets(UWidgetBlueprint* BP, FWidgetRefe
 
 	// Import the nodes
 	TSet<UWidget*> PastedWidgets;
-	FWidgetBlueprintEditorUtils::ImportWidgetsFromText(BP, TextToImport, /*out*/ PastedWidgets);
+	TMap<FName, UWidgetSlotPair*> PastedExtraSlotData;
+	FWidgetBlueprintEditorUtils::ImportWidgetsFromText(BP, TextToImport, /*out*/ PastedWidgets, /*out*/ PastedExtraSlotData);
 
 	// Ignore an empty set of widget paste data.
 	if ( PastedWidgets.Num() == 0 )
@@ -525,7 +533,7 @@ void FWidgetBlueprintEditorUtils::PasteWidgets(UWidgetBlueprint* BP, FWidgetRefe
 		// Widgets with a null parent mean that they were the root most widget of their selection set when
 		// they were copied and thus we need to paste only the root most widgets.  All their children will be added
 		// automatically.
-		if ( NewWidget->GetParent() == NULL )
+		if ( NewWidget->GetParent() == nullptr )
 		{
 			RootPasteWidgets.Add(NewWidget);
 		}
@@ -560,9 +568,31 @@ void FWidgetBlueprintEditorUtils::PasteWidgets(UWidgetBlueprint* BP, FWidgetRefe
 			UPanelSlot* Slot = ParentWidget->AddChild(NewWidget);
 			if ( Slot )
 			{
-				Slot->SetDesiredPosition(PasteLocation);
+				if ( UWidgetSlotPair* OldSlotData = PastedExtraSlotData.FindRef(NewWidget->GetFName()) )
+				{
+					TMap<FName, FString> OldSlotProperties;
+					OldSlotData->GetSlotProperties(OldSlotProperties);
+					FWidgetBlueprintEditorUtils::ImportPropertiesFromText(Slot, OldSlotProperties);
+				}
+
+				BlueprintEditor->AddPostDesignerLayoutAction(
+				[=] {
+					FWidgetReference WidgetRef = BlueprintEditor->GetReferenceFromTemplate(NewWidget);
+					UPanelSlot* PreviewSlot = WidgetRef.GetPreview()->Slot;
+					UPanelSlot* TemplateSlot = WidgetRef.GetTemplate()->Slot;
+					
+					if ( UCanvasPanelSlot* CanvasSlot = Cast<UCanvasPanelSlot>(PreviewSlot) )
+					{
+						CanvasSlot->SaveBaseLayout();
+						CanvasSlot->SetDesiredPosition(PasteLocation);
+						CanvasSlot->RebaseLayout();
+					}
+
+					TMap<FName, FString> SlotProperties;
+					FWidgetBlueprintEditorUtils::ExportPropertiesToText(PreviewSlot, SlotProperties);
+					FWidgetBlueprintEditorUtils::ImportPropertiesFromText(TemplateSlot, SlotProperties);
+				});
 			}
-			//TODO UMG - The paste location needs to be relative from the most upper left hand corner of other widgets in their container.
 		}
 
 		FBlueprintEditorUtils::MarkBlueprintAsStructurallyModified(BP);
@@ -584,7 +614,7 @@ void FWidgetBlueprintEditorUtils::PasteWidgets(UWidgetBlueprint* BP, FWidgetRefe
 	}
 }
 
-void FWidgetBlueprintEditorUtils::ImportWidgetsFromText(UWidgetBlueprint* BP, const FString& TextToImport, /*out*/ TSet<UWidget*>& ImportedWidgetSet)
+void FWidgetBlueprintEditorUtils::ImportWidgetsFromText(UWidgetBlueprint* BP, const FString& TextToImport, /*out*/ TSet<UWidget*>& ImportedWidgetSet, /*out*/ TMap<FName, UWidgetSlotPair*>& PastedExtraSlotData)
 {
 	// We create our own transient package here so that we can deserialize the data in isolation and ensure unreferenced
 	// objects not part of the deserialization set are unresolved.
@@ -595,7 +625,9 @@ void FWidgetBlueprintEditorUtils::ImportWidgetsFromText(UWidgetBlueprint* BP, co
 	FWidgetObjectTextFactory Factory;
 	Factory.ProcessBuffer(TempPackage, RF_Transactional, TextToImport);
 
-	for ( auto& Entry : Factory.NewObjectMap )
+	PastedExtraSlotData = Factory.MissingSlotData;
+
+	for ( auto& Entry : Factory.NewWidgetMap )
 	{
 		UWidget* Widget = Entry.Value;
 
@@ -604,13 +636,26 @@ void FWidgetBlueprintEditorUtils::ImportWidgetsFromText(UWidgetBlueprint* BP, co
 		Widget->SetFlags(RF_Transactional);
 
 		// If there is an existing widget with the same name, rename the newly placed widget.
-		if ( FindObject<UObject>(BP->WidgetTree, *Widget->GetName()) )
+		FString WidgetOldName = Widget->GetName();
+		if ( FindObject<UObject>(BP->WidgetTree, *WidgetOldName) )
 		{
+			UWidgetSlotPair* SlotData = PastedExtraSlotData.FindRef(Widget->GetFName());
+			if ( SlotData )
+			{
+				PastedExtraSlotData.Remove(Widget->GetFName());
+			}
+
 			Widget->Rename(nullptr, BP->WidgetTree);
+
+			if ( SlotData )
+			{
+				SlotData->SetWidgetName(Widget->GetFName());
+				PastedExtraSlotData.Add(Widget->GetFName(), SlotData);
+			}
 		}
 		else
 		{
-			Widget->Rename(*Widget->GetName(), BP->WidgetTree);
+			Widget->Rename(*WidgetOldName, BP->WidgetTree);
 		}
 	}
 
@@ -648,6 +693,9 @@ void FWidgetBlueprintEditorUtils::ImportPropertiesFromText(UObject* Object, cons
 			if ( UProperty* Property = FindField<UProperty>(Object->GetClass(), Entry.Key) )
 			{
 				Property->ImportText(*Entry.Value, Property->ContainerPtrToValuePtr<uint8>(Object), 0, Object);
+
+				FPropertyChangedEvent ChangedEvent(Property);
+				Object->PostEditChangeProperty(ChangedEvent);
 			}
 		}
 	}
diff --git a/Engine/Source/Editor/UMGEditor/Private/WidgetBlueprintEditorUtils.h b/Engine/Source/Editor/UMGEditor/Private/WidgetBlueprintEditorUtils.h
--- a/Engine/Source/Editor/UMGEditor/Private/WidgetBlueprintEditorUtils.h
+++ b/Engine/Source/Editor/UMGEditor/Private/WidgetBlueprintEditorUtils.h
@@ -4,6 +4,7 @@
 
 class UWidgetTree;
 class UWidgetBlueprint;
+class UWidgetSlotPair;
 struct FWidgetReference;
 
 //////////////////////////////////////////////////////////////////////////
@@ -18,7 +19,7 @@ public:
 
 	static void CopyWidgets(UWidgetBlueprint* BP, TSet<FWidgetReference> Widgets);
 
-	static void PasteWidgets(UWidgetBlueprint* BP, FWidgetReference ParentWidget, FVector2D PasteLocation);
+	static void PasteWidgets(TSharedRef<FWidgetBlueprintEditor> BlueprintEditor, UWidgetBlueprint* BP, FWidgetReference ParentWidget, FVector2D PasteLocation);
 
 	static void DeleteWidgets(UWidgetBlueprint* BP, TSet<FWidgetReference> Widgets);
 
@@ -29,7 +30,7 @@ public:
 public:
 	static void ExportWidgetsToText(TSet<UWidget*> WidgetsToExport, /*out*/ FString& ExportedText);
 
-	static void ImportWidgetsFromText(UWidgetBlueprint* BP, const FString& TextToImport, /*out*/ TSet<UWidget*>& ImportedWidgetSet);
+	static void ImportWidgetsFromText(UWidgetBlueprint* BP, const FString& TextToImport, /*out*/ TSet<UWidget*>& ImportedWidgetSet, /*out*/ TMap<FName, UWidgetSlotPair*>& PastedExtraSlotData);
 
 	/** Exports the individual properties of an object to text and stores them in a map. */
 	static void ExportPropertiesToText(UObject* Object, TMap<FName, FString>& ExportedProperties);
diff --git a/Engine/Source/Editor/UMGEditor/UMGEditor.Build.cs b/Engine/Source/Editor/UMGEditor/UMGEditor.Build.cs
--- a/Engine/Source/Editor/UMGEditor/UMGEditor.Build.cs
+++ b/Engine/Source/Editor/UMGEditor/UMGEditor.Build.cs
@@ -19,6 +19,7 @@ public class UMGEditor : ModuleRules
 				"Editor/UMGEditor/Private/Palette",
 				"Editor/UMGEditor/Private/Details",
 				"Editor/UMGEditor/Private/DragDrop",
+                "Editor/UMGEditor/Private/Utility",
 			});
 
 		PrivateIncludePathModuleNames.AddRange(
diff --git a/Engine/Source/Runtime/UMG/Private/Components/CanvasPanel.cpp b/Engine/Source/Runtime/UMG/Private/Components/CanvasPanel.cpp
--- a/Engine/Source/Runtime/UMG/Private/Components/CanvasPanel.cpp
+++ b/Engine/Source/Runtime/UMG/Private/Components/CanvasPanel.cpp
@@ -97,7 +97,7 @@ bool UCanvasPanel::GetGeometryForSlot(UCanvasPanelSlot* Slot, FGeometry& Arrange
 
 		for ( int32 ChildIndex = 0; ChildIndex < ArrangedChildren.Num(); ChildIndex++ )
 		{
-			if ( ArrangedChildren[ChildIndex].Widget == Slot->Content->TakeWidget() )
+			if ( ArrangedChildren[ChildIndex].Widget == Slot->Content->GetCachedWidget() )
 			{
 				ArrangedGeometry = ArrangedChildren[ChildIndex].Geometry;
 				return true;
diff --git a/Engine/Source/Runtime/UMG/Private/Components/CanvasPanelSlot.cpp b/Engine/Source/Runtime/UMG/Private/Components/CanvasPanelSlot.cpp
--- a/Engine/Source/Runtime/UMG/Private/Components/CanvasPanelSlot.cpp
+++ b/Engine/Source/Runtime/UMG/Private/Components/CanvasPanelSlot.cpp
@@ -33,97 +33,6 @@ void UCanvasPanelSlot::BuildSlot(TSharedRef<SConstraintCanvas> Canvas)
 	SynchronizeProperties();
 }
 
-void UCanvasPanelSlot::SetDesiredPosition(FVector2D InPosition)
-{
-	FVector2D Delta = FVector2D(InPosition.X - LayoutData.Offsets.Left, InPosition.Y - LayoutData.Offsets.Top);
-
-	FMargin NewOffset = LayoutData.Offsets;
-	NewOffset.Left += Delta.X;
-	NewOffset.Top += Delta.Y;
-
-	// If the slot is stretched horizontally we need to move the right side as it no longer represents width, but
-	// now represents margin from the right stretched side.
-	if ( LayoutData.Anchors.IsStretchedHorizontal() )
-	{
-		NewOffset.Right -= Delta.X;
-	}
-
-	// If the slot is stretched vertically we need to move the bottom side as it no longer represents width, but
-	// now represents margin from the bottom stretched side.
-	if ( LayoutData.Anchors.IsStretchedVertical() )
-	{
-		NewOffset.Bottom -= Delta.Y;
-	}
-
-	SetOffsets(NewOffset);
-}
-
-void UCanvasPanelSlot::SetDesiredSize(FVector2D InSize)
-{
-	SetOffsets(FMargin(LayoutData.Offsets.Left, LayoutData.Offsets.Top, InSize.X, InSize.Y));
-}
-
-void UCanvasPanelSlot::Resize(const FVector2D& Direction, const FVector2D& Amount)
-{
-	if ( Direction.X < 0 )
-	{
-		LayoutData.Offsets.Left -= Amount.X * Direction.X;
-
-		// Keep the right side stable if it's not stretched and we're resizing left.
-		if ( !LayoutData.Anchors.IsStretchedHorizontal() )
-		{
-			// TODO UMG This doesn't work correctly with Alignment, fix or remove alignment
-			LayoutData.Offsets.Right += Amount.X * Direction.X;
-		}
-	}
-	
-	if ( Direction.Y < 0 )
-	{
-		LayoutData.Offsets.Top -= Amount.Y * Direction.Y;
-
-		// Keep the bottom side stable if it's not stretched and we're resizing top.
-		if ( !LayoutData.Anchors.IsStretchedVertical() )
-		{
-			// TODO UMG This doesn't work correctly with Alignment, fix or remove alignment
-			LayoutData.Offsets.Bottom += Amount.Y * Direction.Y;
-		}
-	}
-	
-	if ( Direction.X > 0 )
-	{
-		if ( LayoutData.Anchors.IsStretchedHorizontal() )
-		{
-			LayoutData.Offsets.Right -= Amount.X * Direction.X;
-		}
-		else
-		{
-			LayoutData.Offsets.Right += Amount.X * Direction.X;
-		}
-	}
-
-	if ( Direction.Y > 0 )
-	{
-		if ( LayoutData.Anchors.IsStretchedVertical() )
-		{
-			LayoutData.Offsets.Bottom -= Amount.Y * Direction.Y;
-		}
-		else
-		{
-			LayoutData.Offsets.Bottom += Amount.Y * Direction.Y;
-		}
-	}
-
-	if ( Slot )
-	{
-		Slot->Offset(LayoutData.Offsets);
-	}
-}
-
-bool UCanvasPanelSlot::CanResize(const FVector2D& Direction) const
-{
-	return true;
-}
-
 void UCanvasPanelSlot::SetLayout(const FAnchorData& InLayoutData)
 {
 	LayoutData = InLayoutData;
@@ -359,6 +268,11 @@ void UCanvasPanelSlot::SaveBaseLayout()
 	}
 }
 
+void UCanvasPanelSlot::SetDesiredPosition(FVector2D InPosition)
+{
+	DesiredPosition = InPosition;
+}
+
 void UCanvasPanelSlot::RebaseLayout(bool PreserveSize)
 {
 	// Ensure we have a parent canvas
@@ -376,14 +290,19 @@ void UCanvasPanelSlot::RebaseLayout(bool PreserveSize)
 				LayoutData.Anchors.Maximum.Y * CanvasSize.Y);
 			FVector2D DefaultAnchorPosition = FVector2D(AnchorPositions.Left, AnchorPositions.Top);
 
+			// Determine the amount that would be offset from the anchor position if alignment was applied.
+			FVector2D AlignmentOffset = LayoutData.Alignment * PreEditGeometry.Size;
+
+			FVector2D MoveDelta = Geometry.Position - PreEditGeometry.Position;
+
 			// Determine where the widget's new position needs to be to maintain a stable location when the anchors change.
 			FVector2D LeftTopDelta = PreEditGeometry.Position - DefaultAnchorPosition;
 
-			if ( PreEditLayoutData.Anchors.Minimum != LayoutData.Anchors.Minimum || PreEditLayoutData.Anchors.Maximum != LayoutData.Anchors.Maximum )
-			{
-				// Determine the amount that would be offset from the anchor position if alignment was applied.
-				FVector2D AlignmentOffset = LayoutData.Alignment * PreEditGeometry.Size;
+			const bool bAnchorsMoved = PreEditLayoutData.Anchors.Minimum != LayoutData.Anchors.Minimum || PreEditLayoutData.Anchors.Maximum != LayoutData.Anchors.Maximum;
+			const bool bMoved = PreEditLayoutData.Offsets.Left != LayoutData.Offsets.Left || PreEditLayoutData.Offsets.Top != LayoutData.Offsets.Top;
 
+			if ( bAnchorsMoved )
+			{
 				// Adjust the size to remain constant
 				if ( !LayoutData.Anchors.IsStretchedHorizontal() && PreEditLayoutData.Anchors.IsStretchedHorizontal() )
 				{
@@ -433,23 +352,60 @@ void UCanvasPanelSlot::RebaseLayout(bool PreserveSize)
 					LayoutData.Offsets.Top = LeftTopDelta.Y + AlignmentOffset.Y;
 				}
 			}
-			else if ( PreEditLayoutData.Offsets.Left != LayoutData.Offsets.Left || PreEditLayoutData.Offsets.Top != LayoutData.Offsets.Top )
+			else if ( DesiredPosition.IsSet() )
+			{
+				FVector2D NewLocalPosition = DesiredPosition.GetValue();
+
+				LayoutData.Offsets.Left = NewLocalPosition.X - AnchorPositions.Left;
+				LayoutData.Offsets.Top = NewLocalPosition.Y - AnchorPositions.Top;
+
+				if ( LayoutData.Anchors.IsStretchedHorizontal() )
+				{
+					LayoutData.Offsets.Right -= LayoutData.Offsets.Left - PreEditLayoutData.Offsets.Left;
+				}
+				else
+				{
+					LayoutData.Offsets.Left += AlignmentOffset.X;
+				}
+
+				if ( LayoutData.Anchors.IsStretchedVertical() )
+				{
+					LayoutData.Offsets.Bottom -= LayoutData.Offsets.Top - PreEditLayoutData.Offsets.Top;
+				}
+				else
+				{
+					LayoutData.Offsets.Top += AlignmentOffset.Y;
+				}
+
+				DesiredPosition.Reset();
+			}
+			else if ( bMoved )
 			{
-				LayoutData.Offsets.Left += LeftTopDelta.X;
-				LayoutData.Offsets.Top += LeftTopDelta.Y;
+				//LayoutData.Offsets.Left += LeftTopDelta.X;
+				//LayoutData.Offsets.Top += LeftTopDelta.Y;
+				LayoutData.Offsets.Left -= DefaultAnchorPosition.X;
+				LayoutData.Offsets.Top -= DefaultAnchorPosition.Y;
 
 				// If the slot is stretched horizontally we need to move the right side as it no longer represents width, but
 				// now represents margin from the right stretched side.
 				if ( LayoutData.Anchors.IsStretchedHorizontal() )
 				{
-					LayoutData.Offsets.Right -= LeftTopDelta.X;
+					//LayoutData.Offsets.Right = PreEditLayoutData.Offsets.Top;
+				}
+				else
+				{
+					LayoutData.Offsets.Left += AlignmentOffset.X;
 				}
 
 				// If the slot is stretched vertically we need to move the bottom side as it no longer represents width, but
 				// now represents margin from the bottom stretched side.
 				if ( LayoutData.Anchors.IsStretchedVertical() )
 				{
-					LayoutData.Offsets.Bottom -= LeftTopDelta.Y;
+					//LayoutData.Offsets.Bottom -= MoveDelta.Y;
+				}
+				else
+				{
+					LayoutData.Offsets.Top += AlignmentOffset.Y;
 				}
 			}
 		}
diff --git a/Engine/Source/Runtime/UMG/Private/Components/PanelSlot.cpp b/Engine/Source/Runtime/UMG/Private/Components/PanelSlot.cpp
--- a/Engine/Source/Runtime/UMG/Private/Components/PanelSlot.cpp
+++ b/Engine/Source/Runtime/UMG/Private/Components/PanelSlot.cpp
@@ -10,36 +10,6 @@ UPanelSlot::UPanelSlot(const FObjectInitializer& ObjectInitializer)
 {
 }
 
-void UPanelSlot::SetDesiredPosition(FVector2D InPosition)
-{
-
-}
-
-void UPanelSlot::SetDesiredSize(FVector2D InSize)
-{
-
-}
-
-void UPanelSlot::Resize(const FVector2D& Direction, const FVector2D& Amount)
-{
-
-}
-
-bool UPanelSlot::CanResize(const FVector2D& Direction) const
-{
-	return false;
-}
-
-void UPanelSlot::MoveTo(const FVector2D& Location)
-{
-
-}
-
-bool UPanelSlot::CanMove() const
-{
-	return false;
-}
-
 bool UPanelSlot::IsDesignTime() const
 {
 	return Parent->IsDesignTime();
diff --git a/Engine/Source/Runtime/UMG/Private/Slate/SObjectWidget.cpp b/Engine/Source/Runtime/UMG/Private/Slate/SObjectWidget.cpp
--- a/Engine/Source/Runtime/UMG/Private/Slate/SObjectWidget.cpp
+++ b/Engine/Source/Runtime/UMG/Private/Slate/SObjectWidget.cpp
@@ -153,6 +153,8 @@ FReply SObjectWidget::OnAnalogValueChanged(const FGeometry& MyGeometry, const FA
 
 FReply SObjectWidget::OnMouseButtonDown(const FGeometry& MyGeometry, const FPointerEvent& MouseEvent)
 {
+	SCompoundWidget::OnMouseButtonDown(MyGeometry, MouseEvent);
+
 	if ( WidgetObject && !WidgetObject->IsDesignTime() )
 	{
 		return WidgetObject->OnMouseButtonDown(MyGeometry, MouseEvent).NativeReply;
@@ -163,6 +165,8 @@ FReply SObjectWidget::OnMouseButtonDown(const FGeometry& MyGeometry, const FPoin
 
 FReply SObjectWidget::OnPreviewMouseButtonDown(const FGeometry& MyGeometry, const FPointerEvent& MouseEvent)
 {
+	SCompoundWidget::OnPreviewMouseButtonDown(MyGeometry, MouseEvent);
+
 	if ( WidgetObject && !WidgetObject->IsDesignTime() )
 	{
 		return WidgetObject->OnPreviewMouseButtonDown(MyGeometry, MouseEvent).NativeReply;
@@ -173,6 +177,8 @@ FReply SObjectWidget::OnPreviewMouseButtonDown(const FGeometry& MyGeometry, cons
 
 FReply SObjectWidget::OnMouseButtonUp(const FGeometry& MyGeometry, const FPointerEvent& MouseEvent)
 {
+	SCompoundWidget::OnMouseButtonUp(MyGeometry, MouseEvent);
+
 	if ( WidgetObject && !WidgetObject->IsDesignTime() )
 	{
 		return WidgetObject->OnMouseButtonUp(MyGeometry, MouseEvent).NativeReply;
diff --git a/Engine/Source/Runtime/UMG/Public/Blueprint/UserWidget.h b/Engine/Source/Runtime/UMG/Public/Blueprint/UserWidget.h
--- a/Engine/Source/Runtime/UMG/Public/Blueprint/UserWidget.h
+++ b/Engine/Source/Runtime/UMG/Public/Blueprint/UserWidget.h
@@ -677,7 +677,7 @@ private:
 template< class T >
 T* CreateWidget(UWorld* World, UClass* UserWidgetClass)
 {
-	if ( !UserWidgetClass->IsChildOf(UUserWidget::StaticClass()) )
+	if ( World == nullptr || !UserWidgetClass->IsChildOf(UUserWidget::StaticClass()) )
 	{
 		// TODO UMG Error?
 		return nullptr;
@@ -697,7 +697,7 @@ T* CreateWidget(UWorld* World, UClass* UserWidgetClass)
 template< class T >
 T* CreateWidget(APlayerController* OwningPlayer, UClass* UserWidgetClass)
 {
-	if ( !UserWidgetClass->IsChildOf(UUserWidget::StaticClass()) )
+	if ( OwningPlayer == nullptr || !UserWidgetClass->IsChildOf(UUserWidget::StaticClass()) )
 	{
 		// TODO UMG Error?
 		return nullptr;
@@ -713,3 +713,24 @@ T* CreateWidget(APlayerController* OwningPlayer, UClass* UserWidgetClass)
 
 	return Cast<T>(NewWidget);
 }
+
+template< class T >
+T* CreateWidget(UGameInstance* OwningGame, UClass* UserWidgetClass)
+{
+	if ( OwningGame == nullptr || !UserWidgetClass->IsChildOf(UUserWidget::StaticClass()) )
+	{
+		// TODO UMG Error?
+		return nullptr;
+	}
+
+	UUserWidget* NewWidget = ConstructObject<UUserWidget>(UserWidgetClass, OwningGame);
+
+	if ( APlayerController* PC = OwningGame->GetFirstLocalPlayerController() )
+	{
+		NewWidget->SetPlayerContext(FLocalPlayerContext(PC));
+	}
+	
+	NewWidget->Initialize();
+
+	return Cast<T>(NewWidget);
+}
diff --git a/Engine/Source/Runtime/UMG/Public/Components/CanvasPanelSlot.h b/Engine/Source/Runtime/UMG/Public/Components/CanvasPanelSlot.h
--- a/Engine/Source/Runtime/UMG/Public/Components/CanvasPanelSlot.h
+++ b/Engine/Source/Runtime/UMG/Public/Components/CanvasPanelSlot.h
@@ -126,14 +126,6 @@ public:
 
 	void BuildSlot(TSharedRef<SConstraintCanvas> Canvas);
 
-	virtual void SetDesiredPosition(FVector2D InPosition) override;
-
-	virtual void SetDesiredSize(FVector2D InSize) override;
-
-	virtual void Resize(const FVector2D& Direction, const FVector2D& Amount) override;
-
-	virtual bool CanResize(const FVector2D& Direction) const override;
-
 	// UPanelSlot interface
 	virtual void SynchronizeProperties() override;
 	// End of UPanelSlot interface
@@ -149,6 +141,8 @@ public:
 	/** Stores the current layout information about the slot and parent canvas. */
 	void SaveBaseLayout();
 
+	void SetDesiredPosition(FVector2D InPosition);
+
 	/** Compares the saved base layout against the current state.  Updates the necessary properties to maintain a stable position. */
 	void RebaseLayout(bool PreserveSize = true);
 #endif
@@ -156,8 +150,10 @@ public:
 private:
 	SConstraintCanvas::FSlot* Slot;
 
-#if WITH_EDITOR
+#if WITH_EDITORONLY_DATA
 	FGeometry PreEditGeometry;
 	FAnchorData PreEditLayoutData;
+
+	TOptional<FVector2D> DesiredPosition;
 #endif
 };
\ No newline at end of file
diff --git a/Engine/Source/Runtime/UMG/Public/Components/PanelSlot.h b/Engine/Source/Runtime/UMG/Public/Components/PanelSlot.h
--- a/Engine/Source/Runtime/UMG/Public/Components/PanelSlot.h
+++ b/Engine/Source/Runtime/UMG/Public/Components/PanelSlot.h
@@ -20,18 +20,6 @@ public:
 	
 	bool IsDesignTime() const;
 
-	virtual void SetDesiredPosition(FVector2D InPosition);
-
-	virtual void SetDesiredSize(FVector2D InSize);
-
-	virtual void Resize(const FVector2D& Direction, const FVector2D& Amount);
-
-	virtual bool CanResize(const FVector2D& Direction) const;
-
-	virtual void MoveTo(const FVector2D& Location);
-
-	virtual bool CanMove() const;
-
 	virtual void ReleaseSlateResources(bool bReleaseChildren) override;
 
 	/** Applies all properties to the live slot if possible. */
