UMG - Widget navigation customization is now enabled, only basic functionality exposed for now.

diff --git a/Engine/Source/Editor/UMGEditor/Private/Customizations/WidgetNavigationCustomization.cpp b/Engine/Source/Editor/UMGEditor/Private/Customizations/WidgetNavigationCustomization.cpp
--- a/Engine/Source/Editor/UMGEditor/Private/Customizations/WidgetNavigationCustomization.cpp
+++ b/Engine/Source/Editor/UMGEditor/Private/Customizations/WidgetNavigationCustomization.cpp
@@ -17,62 +17,16 @@
 
 void FWidgetNavigationCustomization::CustomizeHeader(TSharedRef<IPropertyHandle> PropertyHandle, FDetailWidgetRow& HeaderRow, IPropertyTypeCustomizationUtils& CustomizationUtils)
 {
-	// Disabled for now.
-	if (true) { return; }
-
-	TWeakPtr<IPropertyHandle> PropertyHandlePtr(PropertyHandle);
-
-	HeaderRow
-	.WholeRowContent()
-	[
-		SNew(SVerticalBox)
-
-		+ SVerticalBox::Slot()
-		.AutoHeight()
-		.Padding(2, 4, 2, 4)
-		[
-			SNew(SButton)
-			.HAlign(HAlign_Center)
-			.OnClicked(this, &FWidgetNavigationCustomization::OnCustomizeNavigation, PropertyHandlePtr)
-			[
-				SNew(STextBlock)
-				.Text(LOCTEXT("CustomizeNavigation", "Customize Navigation"))
-			]
-		]
-
-		+ SVerticalBox::Slot()
-		.AutoHeight()
-		.Padding(2, 4, 2, 4)
-		[
-			MakeNavRow(PropertyHandlePtr, EUINavigation::Left, LOCTEXT("LeftNavigation", "Left"))
-		]
-
-		+ SVerticalBox::Slot()
-		.AutoHeight()
-		.Padding(2, 4, 2, 4)
-		[
-			MakeNavRow(PropertyHandlePtr, EUINavigation::Right, LOCTEXT("RightNavigation", "Right"))
-		]
-
-		+ SVerticalBox::Slot()
-		.AutoHeight()
-		.Padding(2, 4, 2, 4)
-		[
-			MakeNavRow(PropertyHandlePtr, EUINavigation::Up, LOCTEXT("UpNavigation", "Up"))
-		]
-
-		+ SVerticalBox::Slot()
-		.AutoHeight()
-		.Padding(2, 4, 2, 4)
-		[
-			MakeNavRow(PropertyHandlePtr, EUINavigation::Down, LOCTEXT("DownNavigation", "Down"))
-		]
-	];
 }
 
 void FWidgetNavigationCustomization::CustomizeChildren(TSharedRef<IPropertyHandle> PropertyHandle, IDetailChildrenBuilder& ChildBuilder, IPropertyTypeCustomizationUtils& CustomizationUtils)
 {
-	//FillOutChildren(PropertyHandle, ChildBuilder, CustomizationUtils);
+	TWeakPtr<IPropertyHandle> PropertyHandlePtr(PropertyHandle);
+
+	MakeNavRow(PropertyHandlePtr, ChildBuilder, EUINavigation::Left, LOCTEXT("LeftNavigation", "Left"));
+	MakeNavRow(PropertyHandlePtr, ChildBuilder, EUINavigation::Right, LOCTEXT("RightNavigation", "Right"));
+	MakeNavRow(PropertyHandlePtr, ChildBuilder, EUINavigation::Up, LOCTEXT("UpNavigation", "Up"));
+	MakeNavRow(PropertyHandlePtr, ChildBuilder, EUINavigation::Down, LOCTEXT("DownNavigation", "Down"));
 }
 
 void FWidgetNavigationCustomization::FillOutChildren(TSharedRef<IPropertyHandle> PropertyHandle, IDetailChildrenBuilder& ChildBuilder, IPropertyTypeCustomizationUtils& CustomizationUtils)
@@ -105,7 +59,7 @@ FReply FWidgetNavigationCustomization::OnCustomizeNavigation(TWeakPtr<IPropertyH
 	TArray<UObject*> OuterObjects;
 	PropertyHandle.Pin()->GetOuterObjects(OuterObjects);
 
-	const FScopedTransaction Transaction(LOCTEXT("InitializeNavigation", "Initialize Navigation"));
+	const FScopedTransaction Transaction(LOCTEXT("UpdateNavigation", "Update Navigation"));
 
 	for ( UObject* OuterObject : OuterObjects )
 	{
@@ -161,22 +115,16 @@ FText FWidgetNavigationCustomization::HandleNavigationText(TWeakPtr<IPropertyHan
 	return FText::GetEmpty();
 }
 
-TSharedRef<class SWidget> FWidgetNavigationCustomization::MakeNavRow(TWeakPtr<IPropertyHandle> PropertyHandle, EUINavigation Nav, FText NavName)
+void FWidgetNavigationCustomization::MakeNavRow(TWeakPtr<IPropertyHandle> PropertyHandle, IDetailChildrenBuilder& ChildBuilder, EUINavigation Nav, FText NavName)
 {
-	TSharedRef<SWidget> Row =
-		SNew(SHorizontalBox)
-
-		+ SHorizontalBox::Slot()
-		.AutoWidth()
-		.Padding(0)
+	ChildBuilder.AddChildContent(NavName)
+		.NameContent()
 		[
 			SNew(STextBlock)
+			.Font(IDetailLayoutBuilder::GetDetailFont())
 			.Text(NavName)
 		]
-
-		+ SHorizontalBox::Slot()
-		.AutoWidth()
-		.Padding(0)
+		.ValueContent()
 		[
 			SNew(SComboButton)
 			.HAlign(HAlign_Center)
@@ -185,14 +133,12 @@ TSharedRef<class SWidget> FWidgetNavigationCustomization::MakeNavRow(TWeakPtr<IP
 				SNew(STextBlock)
 				.Text(this, &FWidgetNavigationCustomization::HandleNavigationText, PropertyHandle, Nav)
 			]
-			.ContentPadding(FMargin(4.0f, 2.0f))
+			//.ContentPadding(FMargin(4.0f, 2.0f))
 			.MenuContent()
 			[
 				MakeNavMenu(PropertyHandle, Nav)
 			]
 		];
-
-	return Row;
 }
 
 TSharedRef<class SWidget> FWidgetNavigationCustomization::MakeNavMenu(TWeakPtr<IPropertyHandle> PropertyHandle, EUINavigation Nav)
@@ -250,7 +196,6 @@ void FWidgetNavigationCustomization::SetNav(UWidget* Widget, EUINavigation Nav,
 		WidgetNavigation = ConstructObject<UWidgetNavigation>(UWidgetNavigation::StaticClass(), Widget);
 	}
 
-	//@Todo: Clear out the navigation if all defaults.
 	switch (Nav)
 	{
 	case EUINavigation::Left:
@@ -267,7 +212,16 @@ void FWidgetNavigationCustomization::SetNav(UWidget* Widget, EUINavigation Nav,
 		break;
 	}
 
-	Widget->Navigation = WidgetNavigation;
+	if ( WidgetNavigation->IsDefault() )
+	{
+		// If the navigation rules are all set to the defaults, remove the navigation
+		// information from the widget.
+		Widget->Navigation = nullptr;
+	}
+	else
+	{
+		Widget->Navigation = WidgetNavigation;
+	}
 }
 
 #undef LOCTEXT_NAMESPACE
