UMG -  missing files needed for widget navigation.

diff --git a/Engine/Source/Editor/UMGEditor/Private/Customizations/WidgetNavigationCustomization.h b/Engine/Source/Editor/UMGEditor/Private/Customizations/WidgetNavigationCustomization.h
--- a/Engine/Source/Editor/UMGEditor/Private/Customizations/WidgetNavigationCustomization.h
+++ b/Engine/Source/Editor/UMGEditor/Private/Customizations/WidgetNavigationCustomization.h
@@ -30,7 +30,7 @@ private:
 	
 	FText HandleNavigationText(TWeakPtr<IPropertyHandle> PropertyHandle, EUINavigation Nav) const;
 
-	TSharedRef<class SWidget> MakeNavRow(TWeakPtr<IPropertyHandle> PropertyHandle, EUINavigation Nav, FText NavName);
+	void MakeNavRow(TWeakPtr<IPropertyHandle> PropertyHandle, IDetailChildrenBuilder& ChildBuilder, EUINavigation Nav, FText NavName);
 
 	TSharedRef<class SWidget> MakeNavMenu(TWeakPtr<IPropertyHandle> PropertyHandle, EUINavigation Nav);
 
diff --git a/Engine/Source/Runtime/UMG/Public/Blueprint/WidgetNavigation.h b/Engine/Source/Runtime/UMG/Public/Blueprint/WidgetNavigation.h
--- a/Engine/Source/Runtime/UMG/Public/Blueprint/WidgetNavigation.h
+++ b/Engine/Source/Runtime/UMG/Public/Blueprint/WidgetNavigation.h
@@ -72,6 +72,14 @@ public:
 		UpdateMetaDataEntry(MetaData, Right, EUINavigation::Right);
 	}
 
+	bool IsDefault() const
+	{
+		return Up.Rule == EUINavigationRule::Escape &&
+			Down.Rule == EUINavigationRule::Escape &&
+			Left.Rule == EUINavigationRule::Escape &&
+			Right.Rule == EUINavigationRule::Escape;
+	}
+
 private:
 
 	void UpdateMetaDataEntry(TSharedRef<FNavigationMetaData> MetaData, const FWidgetNavigationData & NavData, EUINavigation Nav)
