Fix non-unity in WidgetSlotPair.cpp

diff --git a/Engine/Source/Editor/UMGEditor/Private/Utility/WidgetSlotPair.cpp b/Engine/Source/Editor/UMGEditor/Private/Utility/WidgetSlotPair.cpp
--- a/Engine/Source/Editor/UMGEditor/Private/Utility/WidgetSlotPair.cpp
+++ b/Engine/Source/Editor/UMGEditor/Private/Utility/WidgetSlotPair.cpp
@@ -2,6 +2,7 @@
 
 #include "UMGEditorPrivatePCH.h"
 #include "WidgetSlotPair.h"
+#include "WidgetBlueprintEditorUtils.h"
 
 UWidgetSlotPair::UWidgetSlotPair()
 {
