UMG - Fixing a designer resizing issue with changing the upper left/top corner of a canvas item with stretched anchors.

diff --git a/Engine/Source/Editor/UMGEditor/Private/Designer/STransformHandle.cpp b/Engine/Source/Editor/UMGEditor/Private/Designer/STransformHandle.cpp
--- a/Engine/Source/Editor/UMGEditor/Private/Designer/STransformHandle.cpp
+++ b/Engine/Source/Editor/UMGEditor/Private/Designer/STransformHandle.cpp
@@ -157,24 +157,26 @@ void STransformHandle::Resize(UCanvasPanelSlot* Slot, const FVector2D& Direction
 
 	if ( Direction.X < 0 )
 	{
-		Offsets.Left -= PositionMovement.X;
-
-		// Keep the right side stable if it's not stretched and we're resizing left.
-		if ( !LayoutData.Anchors.IsStretchedHorizontal() )
+		if ( LayoutData.Anchors.IsStretchedHorizontal() )
+		{
+			Offsets.Left -= Amount.X * Direction.X;
+		}
+		else
 		{
-			// TODO UMG This doesn't work correctly with Alignment, fix or remove alignment
+			Offsets.Left -= PositionMovement.X;
 			Offsets.Right += SizeMovement.X;
 		}
 	}
 
 	if ( Direction.Y < 0 )
 	{
-		Offsets.Top -= PositionMovement.Y;
-
-		// Keep the bottom side stable if it's not stretched and we're resizing top.
-		if ( !LayoutData.Anchors.IsStretchedVertical() )
+		if ( LayoutData.Anchors.IsStretchedVertical() )
+		{
+			Offsets.Top -= Amount.Y * Direction.Y;
+		}
+		else
 		{
-			// TODO UMG This doesn't work correctly with Alignment, fix or remove alignment
+			Offsets.Top -= PositionMovement.Y;
 			Offsets.Bottom += SizeMovement.Y;
 		}
 	}
