---- Merging with SlateDev branch ----

diff --git a/Engine/Source/Editor/LevelEditor/Private/SLevelEditor.cpp b/Engine/Source/Editor/LevelEditor/Private/SLevelEditor.cpp
--- a/Engine/Source/Editor/LevelEditor/Private/SLevelEditor.cpp
+++ b/Engine/Source/Editor/LevelEditor/Private/SLevelEditor.cpp
@@ -1410,14 +1410,6 @@ UWorld* SLevelEditor::GetWorld() const
 	return World;
 }
 
-void SLevelEditor::Tick( const FGeometry& AllottedGeometry, const double InCurrentTime, const float InDeltaTime )
-{
-	if(ThumbnailPool.IsValid())
-	{
-		ThumbnailPool->Tick( InDeltaTime );
-	}
-}
-
 void SLevelEditor::HandleEditorMapChange( uint32 MapChangeFlags )
 {
 	ResetViewportTabInfo();
diff --git a/Engine/Source/Editor/LevelEditor/Private/SLevelEditor.h b/Engine/Source/Editor/LevelEditor/Private/SLevelEditor.h
--- a/Engine/Source/Editor/LevelEditor/Private/SLevelEditor.h
+++ b/Engine/Source/Editor/LevelEditor/Private/SLevelEditor.h
@@ -121,8 +121,6 @@ public:
 		return true;
 	}
 
-	virtual void Tick( const FGeometry& AllottedGeometry, const double InCurrentTime, const float InDeltaTime ) override;
-
 	// Tab Management
 	TSharedRef<FTabManager> GetTabManager() const;
 	
diff --git a/Engine/Source/Editor/LevelEditor/Private/SLevelViewport.cpp b/Engine/Source/Editor/LevelEditor/Private/SLevelViewport.cpp
--- a/Engine/Source/Editor/LevelEditor/Private/SLevelViewport.cpp
+++ b/Engine/Source/Editor/LevelEditor/Private/SLevelViewport.cpp
@@ -815,12 +815,10 @@ FReply SLevelViewport::OnDrop( const FGeometry& MyGeometry, const FDragDropEvent
 
 void SLevelViewport::Tick( const FGeometry& AllottedGeometry, const double InCurrentTime, const float InDeltaTime )
 {
-	SEditorViewport::Tick(AllottedGeometry, InCurrentTime, InDeltaTime);
+	SEditorViewport::Tick( AllottedGeometry, InCurrentTime, InDeltaTime );
 
 	const bool bContainsFocus = HasFocusedDescendants();
 
-	LastTickTime = FPlatformTime::Seconds();
-
 	// When we have focus we update the 'Allow Throttling' option in slate to be disabled so that interactions in the
 	// viewport with Slate widgets that are part of the game, don't throttle.
 	if ( bPIEContainsFocus != bContainsFocus )
@@ -848,7 +846,7 @@ void SLevelViewport::Tick( const FGeometry& AllottedGeometry, const double InCur
 	// user could miss part of the animation, or even the whole thing!
 	if( bViewTransitionAnimPending )
 	{
-		ViewTransitionAnim.Play();
+		ViewTransitionAnim.Play(this->AsShared());
 		bViewTransitionAnimPending = false;
 	}
 
@@ -861,19 +859,19 @@ void SLevelViewport::Tick( const FGeometry& AllottedGeometry, const double InCur
 			if(PIEOverlaySlotIndex)
 			{
 				PIEOverlayAnim = FCurveSequence(0.0f, SLevelViewportPIEAnimation::MouseControlLabelFadeout, ECurveEaseFunction::CubicInOut);
-				PIEOverlayAnim.Play();
+				PIEOverlayAnim.Play(this->AsShared());
 			}
 		}
 		ViewTransitionType = EViewTransition::None;
 		ViewTransitionAnim = FCurveSequence( 0.0f, 0.25f, ECurveEaseFunction::QuadOut );
-		ViewTransitionAnim.PlayReverse();
+		ViewTransitionAnim.PlayReverse(this->AsShared());
 	}
 
 	if(IsPlayInEditorViewportActive() && bPIEHasFocus != ActiveViewport->HasMouseCapture())
 	{
 		bPIEHasFocus = ActiveViewport->HasMouseCapture();
 		PIEOverlayAnim = FCurveSequence(0.0f, SLevelViewportPIEAnimation::MouseControlLabelFadeout, ECurveEaseFunction::CubicInOut);
-		PIEOverlayAnim.Play();
+		PIEOverlayAnim.Play(this->AsShared());
 	}
 
 	// Update actor preview viewports, if we have any
@@ -2640,7 +2638,7 @@ void SActorPreview::Construct( const FArguments& InArgs )
 		FadeSequence = FCurveSequence( TimeBeforeFadingIn, FadeTime );
 
 		// Start fading in!
-		FadeSequence.Play( TimeBeforeFadingIn );	// Skip the initial time delay and just fade straight in
+		FadeSequence.Play(this->AsShared(), false, TimeBeforeFadingIn);	// Skip the initial time delay and just fade straight in
 	}
 
 	HighlightSequence = FCurveSequence(0.f, 0.5f, ECurveEaseFunction::Linear);
@@ -2720,7 +2718,7 @@ void SActorPreview::OnMouseEnter( const FGeometry& MyGeometry, const FPointerEve
 		}
 		else
 		{
-			FadeSequence.PlayReverse();
+			FadeSequence.PlayReverse(this->AsShared());
 		}
 	}
 }
@@ -2743,7 +2741,7 @@ void SActorPreview::OnMouseLeave( const FPointerEvent& MouseEvent )
 		}
 		else
 		{
-			FadeSequence.Play();
+			FadeSequence.Play(this->AsShared());
 		}
 	}
 
@@ -2785,7 +2783,7 @@ void SActorPreview::OnActorSelected(UObject* InActor)
 void SActorPreview::Highlight()
 {
 	HighlightSequence.JumpToStart();
-	HighlightSequence.Play();
+	HighlightSequence.Play(this->AsShared());
 }
 
 FSlateColor SActorPreview::GetBorderColorAndOpacity() const
diff --git a/Engine/Source/Editor/MaterialEditor/Private/SMaterialEditorViewport.cpp b/Engine/Source/Editor/MaterialEditor/Private/SMaterialEditorViewport.cpp
--- a/Engine/Source/Editor/MaterialEditor/Private/SMaterialEditorViewport.cpp
+++ b/Engine/Source/Editor/MaterialEditor/Private/SMaterialEditorViewport.cpp
@@ -230,11 +230,6 @@ void SMaterialEditorViewport::RefreshViewport()
 	SceneViewport->InvalidateDisplay();
 }
 
-void SMaterialEditorViewport::Tick(const FGeometry& AllottedGeometry, const double InCurrentTime, const float InDeltaTime)
-{
-	SEditorViewport::Tick(AllottedGeometry, InCurrentTime, InDeltaTime);
-}
-
 bool SMaterialEditorViewport::SetPreviewMesh(UStaticMesh* InStaticMesh, USkeletalMesh* InSkeletalMesh)
 {
 	// Only permit the use of a skeletal mesh if the material has bUsedWithSkeltalMesh.
diff --git a/Engine/Source/Editor/MaterialEditor/Private/SMaterialEditorViewport.h b/Engine/Source/Editor/MaterialEditor/Private/SMaterialEditorViewport.h
--- a/Engine/Source/Editor/MaterialEditor/Private/SMaterialEditorViewport.h
+++ b/Engine/Source/Editor/MaterialEditor/Private/SMaterialEditorViewport.h
@@ -21,7 +21,6 @@ public:
 	
 	virtual void AddReferencedObjects( FReferenceCollector& Collector ) override;
 
-	virtual void Tick(const FGeometry& AllottedGeometry, const double InCurrentTime, const float InDeltaTime) override;
 	void RefreshViewport();
 	
 	/**
diff --git a/Engine/Source/Editor/Matinee/Private/SMatineeRecorder.cpp b/Engine/Source/Editor/Matinee/Private/SMatineeRecorder.cpp
--- a/Engine/Source/Editor/Matinee/Private/SMatineeRecorder.cpp
+++ b/Engine/Source/Editor/Matinee/Private/SMatineeRecorder.cpp
@@ -32,11 +32,6 @@ SMatineeRecorder::~SMatineeRecorder()
 	LevelViewportClient.Reset();
 }
 
-void SMatineeRecorder::Tick(const FGeometry& AllottedGeometry, const double InCurrentTime, const float InDeltaTime)
-{
-	SCompoundWidget::Tick(AllottedGeometry, InCurrentTime, InDeltaTime);
-}
-
 void SMatineeRecorder::RefreshViewport()
 {
 	Viewport->InvalidateDisplay();
diff --git a/Engine/Source/Editor/MeshPaint/Private/SMeshPaint.cpp b/Engine/Source/Editor/MeshPaint/Private/SMeshPaint.cpp
--- a/Engine/Source/Editor/MeshPaint/Private/SMeshPaint.cpp
+++ b/Engine/Source/Editor/MeshPaint/Private/SMeshPaint.cpp
@@ -2087,14 +2087,10 @@ void SMeshPaint::Construct(const FArguments& InArgs, TSharedRef<FMeshPaintToolKi
 		]
 	];
 }
-END_SLATE_FUNCTION_BUILD_OPTIMIZATION
 
 void SMeshPaint::Tick( const FGeometry& AllottedGeometry, const double InCurrentTime, const float InDeltaTime )
 {
-	if( MeshPaintEditMode->UpdateTextureList() == true )
-	{
-		MeshPaintEditMode->UpdateTexturePaintTargetList();
-	}
+	MeshPaintEditMode->UpdateTexturePaintTargetList();
 }
 
 class FEdMode* SMeshPaint::GetEditorMode() const
diff --git a/Engine/Source/Editor/PListEditor/Private/SPlistEditor.cpp b/Engine/Source/Editor/PListEditor/Private/SPlistEditor.cpp
--- a/Engine/Source/Editor/PListEditor/Private/SPlistEditor.cpp
+++ b/Engine/Source/Editor/PListEditor/Private/SPlistEditor.cpp
@@ -11,12 +11,6 @@
 #include "SNotificationList.h"
 #define LOCTEXT_NAMESPACE "PListEditor"
 
-/** Ticks any internals of the widget as necessary. */
-void SPListEditorPanel::Tick( const FGeometry& /*AllottedGeometry*/, const double /*InCurrentTime*/, const float /*InDeltaTime*/ )
-{
-	PerformDisplayNotifications();
-}
-
 /** Constructs the main widget for the editor */
 void SPListEditorPanel::Construct( const FArguments& InArgs )
 {
@@ -2076,41 +2070,47 @@ void SPListEditorPanel::OnPopupTextChosen(const FString& ChosenText)
 	FSlateApplication::Get().DismissAllMenus();
 }
 
-/** Helper method to display queued notifications */
-void SPListEditorPanel::PerformDisplayNotifications()
+EActiveTickReturnType SPListEditorPanel::DisplayDeferredNotifications( double InCurrentTime, float InDeltaTime )
 {
-	if(FramesToSkip == 0)
-		return;
-
-	if(--FramesToSkip == 0)
+	if ( --FramesToSkip == 0 )
 	{
-		for(int32 i = 0; i < QueuedNotifications.Num(); ++i)
+		for ( int32 i = 0; i < QueuedNotifications.Num(); ++i )
 		{
-			if(QueuedNotifications[i].NotificationType() == NTF_Normal)
+			if ( QueuedNotifications[i].NotificationType() == NTF_Normal )
 			{
-				NotificationListPtr->AddNotification( FNotificationInfo(QueuedNotifications[i].Notification()) );
+				NotificationListPtr->AddNotification( FNotificationInfo( QueuedNotifications[i].Notification() ) );
 			}
-			else if(QueuedNotifications[i].NotificationType() == NTF_Success)
+			else if ( QueuedNotifications[i].NotificationType() == NTF_Success )
 			{
-				FNotificationInfo info(QueuedNotifications[i].Notification());
-				TWeakPtr<SNotificationItem> PendingProgressPtr = NotificationListPtr->AddNotification(info);
-				PendingProgressPtr.Pin()->SetCompletionState(SNotificationItem::CS_Success);
+				FNotificationInfo info( QueuedNotifications[i].Notification() );
+				TWeakPtr<SNotificationItem> PendingProgressPtr = NotificationListPtr->AddNotification( info );
+				PendingProgressPtr.Pin()->SetCompletionState( SNotificationItem::CS_Success );
 			}
 			else // if(QueuedNotifications[i].NotificationType() == NTF_Failed)
 			{
-				FNotificationInfo info(QueuedNotifications[i].Notification());
-				TWeakPtr<SNotificationItem> PendingProgressPtr = NotificationListPtr->AddNotification(info);
-				PendingProgressPtr.Pin()->SetCompletionState(SNotificationItem::CS_Fail);
+				FNotificationInfo info( QueuedNotifications[i].Notification() );
+				TWeakPtr<SNotificationItem> PendingProgressPtr = NotificationListPtr->AddNotification( info );
+				PendingProgressPtr.Pin()->SetCompletionState( SNotificationItem::CS_Fail );
 			}
 		}
 
 		QueuedNotifications.Empty();
+
+		return EActiveTickReturnType::StopTicking;
 	}
+
+	return EActiveTickReturnType::KeepTicking;
 }
 
 /** Helper function to display notifications in the current tab */
 void SPListEditorPanel::DisplayNotification(const FText& ToDisplay, const ENTF_Types NotificationType)
 {
+	// Register the active tick if it isn't already
+	if ( FramesToSkip == 0 )
+	{
+		RegisterActiveTick( 0.f, FWidgetActiveTickDelegate::CreateSP( this, &SPListEditorPanel::DisplayDeferredNotifications ) );
+	}
+
 	QueuedNotifications.Add( FQueuedNotification( ToDisplay, NotificationType ) );
 	FramesToSkip = 15; // Hack to get notifications to always show full animations (would break if displaying >1 notification within 'FramesToSkip' frames)
 }
diff --git a/Engine/Source/Editor/PListEditor/Public/SPlistEditor.h b/Engine/Source/Editor/PListEditor/Public/SPlistEditor.h
--- a/Engine/Source/Editor/PListEditor/Public/SPlistEditor.h
+++ b/Engine/Source/Editor/PListEditor/Public/SPlistEditor.h
@@ -117,7 +117,6 @@ public:
 	}
 	
 	// SWidget interface
-	virtual void Tick( const FGeometry& AllottedGeometry, const double InCurrentTime, const float InDeltaTime ) override;
 	virtual FReply OnKeyDown( const FGeometry& MyGeometry, const FKeyEvent& InKeyEvent ) override;
 	virtual FReply OnMouseButtonDown( const FGeometry& MyGeometry, const FPointerEvent& MouseEvent ) override;
 	// End of SWidget interface
@@ -129,7 +128,7 @@ public:
 	bool OnTabClose();
 
 private:
-
+	
 	/** An internal array holding parsed data from a loaded plist file */
 	TArray<TSharedPtr<IPListNode> > PListNodes;
 	/** The list widget, needed so we can request refreshes when we change the list's contents */
@@ -233,11 +232,14 @@ public:
 	void MarkDirty();
 
 private:
+	/** Displays queued notifications */
+	EActiveTickReturnType DisplayDeferredNotifications( double InCurrentTime, float InDeltaTime );
+	/** Helper method to display queued notifications */
+	void PerformDisplayNotifications();
+
 
 	/** Helper method to parse through and load an Xml tree into an internal intermediate format for Slate */
 	bool ParseXmlTree(FXmlFile& Doc, FString& OutError);
-	/** Helper method to display queued notifications */
-	void PerformDisplayNotifications();
 	/** Helper method to write out the contents of PListNodes using a valid FileWriter */
 	void SerializePListNodes(FArchive* Writer);
 	/** Helper function to check if PListNodes is valid (every element being valid) */
diff --git a/Engine/Source/Editor/Persona/Private/SAnimCompositeEditor.cpp b/Engine/Source/Editor/Persona/Private/SAnimCompositeEditor.cpp
--- a/Engine/Source/Editor/Persona/Private/SAnimCompositeEditor.cpp
+++ b/Engine/Source/Editor/Persona/Private/SAnimCompositeEditor.cpp
@@ -22,7 +22,7 @@ TSharedRef<SWidget> SAnimCompositeEditor::CreateDocumentAnchor()
 
 void SAnimCompositeEditor::Construct(const FArguments& InArgs)
 {
-	bRebuildPanel = false;
+	bIsActiveTickRegistered = false;
 	PersonaPtr = InArgs._Persona;
 	CompositeObj = InArgs._Composite;
 	check(CompositeObj);
@@ -106,16 +106,16 @@ void SAnimCompositeEditor::RebuildPanel()
 {
 	SortAndUpdateComposite();
 	AnimCompositePanel->Update();
-	bRebuildPanel = false;
 }
 
-void SAnimCompositeEditor::OnCompositeChange(class UObject *EditorAnimBaseObj, bool Rebuild)
+void SAnimCompositeEditor::OnCompositeChange(class UObject *EditorAnimBaseObj, bool bRebuild)
 {
-	if(CompositeObj!=NULL)
+	if ( CompositeObj != nullptr )
 	{
-		if(Rebuild)
+		if(bRebuild && !bIsActiveTickRegistered)
 		{
-			bRebuildPanel = true;
+			bIsActiveTickRegistered = true;
+			RegisterActiveTick(0.f, FWidgetActiveTickDelegate::CreateSP(this, &SAnimCompositeEditor::TriggerRebuildPanel));
 		} 
 		else
 		{
@@ -128,7 +128,7 @@ void SAnimCompositeEditor::OnCompositeChange(class UObject *EditorAnimBaseObj, b
 
 void SAnimCompositeEditor::CollapseComposite()
 {
-	if (CompositeObj==NULL)
+	if ( CompositeObj == nullptr )
 	{
 		return;
 	}
@@ -140,7 +140,11 @@ void SAnimCompositeEditor::CollapseComposite()
 
 void SAnimCompositeEditor::PostUndo()
 {
-	bRebuildPanel=true;
+	if (!bIsActiveTickRegistered)
+	{
+		bIsActiveTickRegistered = true;
+		RegisterActiveTick(0.f, FWidgetActiveTickDelegate::CreateSP(this, &SAnimCompositeEditor::TriggerRebuildPanel));
+	}
 
 	// when undo or redo happens, we still have to recalculate length, so we can't rely on sequence length changes or not
 	if (CompositeObj->SequenceLength)
@@ -154,16 +158,14 @@ void SAnimCompositeEditor::InitDetailsViewEditorObject(UEditorAnimBaseObj* EdObj
 	EdObj->InitFromAnim(CompositeObj, FOnAnimObjectChange::CreateSP( SharedThis(this), &SAnimCompositeEditor::OnCompositeChange ));
 }
 
-void SAnimCompositeEditor::Tick( const FGeometry& AllottedGeometry, const double InCurrentTime, const float InDeltaTime )
+EActiveTickReturnType SAnimCompositeEditor::TriggerRebuildPanel(double InCurrentTime, float InDeltaTime)
 {
 	// we should not update any property related within PostEditChange, 
 	// so this is deferred to Tick, when it needs to rebuild, just mark it and this will update in the next tick
-	if (bRebuildPanel)
-	{
-		RebuildPanel();
-	}
+	RebuildPanel();
 
-	SWidget::Tick(AllottedGeometry, InCurrentTime, InDeltaTime);
+	bIsActiveTickRegistered = false;
+	return EActiveTickReturnType::StopTicking;
 }
 
 float SAnimCompositeEditor::CalculateSequenceLengthOfEditorObject() const
@@ -173,7 +175,7 @@ float SAnimCompositeEditor::CalculateSequenceLengthOfEditorObject() const
 
 void SAnimCompositeEditor::SortAndUpdateComposite()
 {
-	if (CompositeObj == NULL)
+	if (CompositeObj == nullptr)
 	{
 		return;
 	}
diff --git a/Engine/Source/Editor/Persona/Private/SAnimCompositeEditor.h b/Engine/Source/Editor/Persona/Private/SAnimCompositeEditor.h
--- a/Engine/Source/Editor/Persona/Private/SAnimCompositeEditor.h
+++ b/Engine/Source/Editor/Persona/Private/SAnimCompositeEditor.h
@@ -43,7 +43,7 @@ private:
 	void RebuildPanel();
 
 	/** Handler for when composite is edited in details view */
-	void OnCompositeChange(class UObject *EditorAnimBaseObj, bool Rebuild);
+	void OnCompositeChange(class UObject *EditorAnimBaseObj, bool bRebuild);
 
 	/** This will remove empty spaces in the composite's anim segments but not resort. e.g. - all cached indexes remains valid. UI IS NOT REBUILT after this */
 	void CollapseComposite();
@@ -69,16 +69,18 @@ private:
 	/** This will sort all components of the montage and update (recreate) the UI */
 	void SortAndUpdateComposite();
 
+	/** One-off active tick to trigger a panel rebuild */
+	EActiveTickReturnType TriggerRebuildPanel(double InCurrentTime, float InDeltaTime);
+
+	/** Whether the active tick to trigger a panel rebuild is currently registered */
+	bool bIsActiveTickRegistered;
+
 public:
 
 	/** delegate handlers for when the composite is being editor */
 	void			PreAnimUpdate();
 	void			PostAnimUpdate();
 
-	// SWidget interface start
-	virtual void Tick( const FGeometry& AllottedGeometry, const double InCurrentTime, const float InDeltaTime ) override;
-	// SWidget interface end
-
 	// Begin SAnimEditorBase interface
 	virtual TSharedRef<SWidget> CreateDocumentAnchor() override;
 	// End SAnimEditorBase interface
diff --git a/Engine/Source/Editor/Persona/Private/SAnimNotifyPanel.cpp b/Engine/Source/Editor/Persona/Private/SAnimNotifyPanel.cpp
--- a/Engine/Source/Editor/Persona/Private/SAnimNotifyPanel.cpp
+++ b/Engine/Source/Editor/Persona/Private/SAnimNotifyPanel.cpp
@@ -72,7 +72,7 @@ public:
 	virtual FReply OnDragDetected( const FGeometry& MyGeometry, const FPointerEvent& MouseEvent) override;
 	virtual FReply OnMouseMove( const FGeometry& MyGeometry, const FPointerEvent& MouseEvent) override;
 	virtual FReply OnMouseButtonUp( const FGeometry& MyGeometry, const FPointerEvent& MouseEvent ) override;
-	virtual void Tick(const FGeometry& AllottedGeometry, const double InCurrentTime, const float InDeltaTime) override;
+	virtual void Tick( const FGeometry& AllottedGeometry, const double InCurrentTime, const float InDeltaTime ) override;
 	
 	// End of SWidget interface
 
@@ -242,7 +242,7 @@ public:
 	void Construct(const FArguments& InArgs);
 
 	// SWidget interface
-	virtual void Tick( const FGeometry& AllottedGeometry, const double InCurrentTime, const float InDeltaTime ) {UpdateCachedGeometry(AllottedGeometry);}
+	virtual void Tick( const FGeometry& AllottedGeometry, const double InCurrentTime, const float InDeltaTime ) override { UpdateCachedGeometry( AllottedGeometry ); }
 	virtual int32 OnPaint(const FPaintArgs& Args, const FGeometry& AllottedGeometry, const FSlateRect& MyClippingRect, FSlateWindowElementList& OutDrawElements, int32 LayerId, const FWidgetStyle& InWidgetStyle, bool bParentEnabled) const override;
 	virtual FReply OnMouseWheel(const FGeometry& MyGeometry, const FPointerEvent& MouseEvent) override;
 	virtual FCursorReply OnCursorQuery(const FGeometry& MyGeometry, const FPointerEvent& CursorEvent) const override;
@@ -1520,7 +1520,7 @@ void SAnimNotifyNode::DrawHandleOffset( const float& Offset, const float& Handle
 		);
 }
 
-void SAnimNotifyNode::Tick(const FGeometry& AllottedGeometry, const double InCurrentTime, const float InDeltaTime)
+void SAnimNotifyNode::Tick( const FGeometry& AllottedGeometry, const double InCurrentTime, const float InDeltaTime )
 {
 	ScreenPosition = AllottedGeometry.AbsolutePosition;
 }
diff --git a/Engine/Source/Editor/Persona/Private/SAnimTrackPanel.cpp b/Engine/Source/Editor/Persona/Private/SAnimTrackPanel.cpp
--- a/Engine/Source/Editor/Persona/Private/SAnimTrackPanel.cpp
+++ b/Engine/Source/Editor/Persona/Private/SAnimTrackPanel.cpp
@@ -131,7 +131,7 @@ FReply SAnimTrackPanel::OnMouseMove( const FGeometry& InMyGeometry, const FPoint
 		if( !bPanning )
 		{
 			PanningDistance += FMath::Abs(InMouseEvent.GetCursorDelta().X);
-			if ( PanningDistance > FSlateApplication::Get().GetDragTriggerDistnace() )
+			if ( PanningDistance > FSlateApplication::Get().GetDragTriggerDistance() )
 			{
 				bPanning = true;
 				UE_LOG(LogAnimation, Log, TEXT("MouseMove (Capturing Mouse) %d, %0.5f"), bPanning, PanningDistance);
diff --git a/Engine/Source/Editor/Persona/Private/SAnimationBlendSpaceBase.cpp b/Engine/Source/Editor/Persona/Private/SAnimationBlendSpaceBase.cpp
--- a/Engine/Source/Editor/Persona/Private/SAnimationBlendSpaceBase.cpp
+++ b/Engine/Source/Editor/Persona/Private/SAnimationBlendSpaceBase.cpp
@@ -1180,26 +1180,23 @@ SBlendSpaceEditorBase::~SBlendSpaceEditorBase()
 
 void SBlendSpaceEditorBase::Construct(const FArguments& InArgs)
 {
+	bIsActiveTickRegistered = false;
 	BlendSpace = InArgs._BlendSpace;
 	PersonaPtr = InArgs._Persona;
 	PersonaPtr.Pin()->RegisterOnPostUndo(FPersona::FOnPostUndo::CreateSP( this, &SBlendSpaceEditorBase::PostUndo ) );
 }
 
-/**
-* Ticks this widget.  Override in derived classes, but always call the parent implementation.
-*
-* @param  AllottedGeometry The space allotted for this widget
-* @param  InCurrentTime  Current absolute real time
-* @param  InDeltaTime  Real time passed since last tick
-*/
-void SBlendSpaceEditorBase::Tick( const FGeometry& AllottedGeometry, const double InCurrentTime, const float InDeltaTime )
+EActiveTickReturnType SBlendSpaceEditorBase::UpdatePreview( double InCurrentTime, float InDeltaTime )
 {
-	BlendSpaceWidget->bPreviewOn = IsPreviewOn()==ECheckBoxState::Checked;
+	// Update the preview as long as its enabled
 	if (BlendSpaceWidget->bPreviewOn)
 	{
 		UpdatePreviewParameter();
+		return EActiveTickReturnType::KeepTicking;
 	}
-	SWidget::Tick(AllottedGeometry, InCurrentTime, InDeltaTime);
+
+	bIsActiveTickRegistered = false;
+	return EActiveTickReturnType::StopTicking;
 }
 
 void SBlendSpaceEditorBase::PostUndo()
@@ -1288,12 +1285,20 @@ ECheckBoxState SBlendSpaceEditorBase::IsPreviewOn() const
 
 void SBlendSpaceEditorBase::ShowPreview_OnIsCheckedChanged( ECheckBoxState NewValue )
 {
-	bool bPreviewOn = (NewValue != ECheckBoxState::Unchecked);
+	bool bPreviewEnabled = (NewValue != ECheckBoxState::Unchecked);
 	class UDebugSkelMeshComponent* Component = PersonaPtr.Pin()->GetPreviewMeshComponent();
 
 	if ( Component != NULL )
 	{
-		Component->EnablePreview(bPreviewOn, BlendSpace, NULL);
+		Component->EnablePreview( bPreviewEnabled, BlendSpace, NULL );
+
+		BlendSpaceWidget->bPreviewOn = bPreviewEnabled;
+
+		if ( bPreviewEnabled && !bIsActiveTickRegistered )
+		{
+			bIsActiveTickRegistered = true;
+			RegisterActiveTick( 0.f, FWidgetActiveTickDelegate::CreateSP( this, &SBlendSpaceEditorBase::UpdatePreview ) );
+		}
 	}
 }
 
@@ -1309,7 +1314,7 @@ void SBlendSpaceEditorBase::UpdatePreviewParameter() const
 {
 	class UDebugSkelMeshComponent* Component = PersonaPtr.Pin()->GetPreviewMeshComponent();
 
-	if (Component != NULL && Component->IsPreviewOn())
+	if (Component != nullptr && Component->IsPreviewOn())
 	{
 		if (Component->PreviewInstance->CurrentAsset == BlendSpace)
 		{
diff --git a/Engine/Source/Editor/Persona/Private/SAnimationBlendSpaceBase.h b/Engine/Source/Editor/Persona/Private/SAnimationBlendSpaceBase.h
--- a/Engine/Source/Editor/Persona/Private/SAnimationBlendSpaceBase.h
+++ b/Engine/Source/Editor/Persona/Private/SAnimationBlendSpaceBase.h
@@ -357,13 +357,11 @@ public:
 
 	void Construct(const FArguments& InArgs);
 
-	// SWidget interface
-	virtual void	Tick(const FGeometry& AllottedGeometry, const double InCurrentTime, const float InDeltaTime) override;
-	// End of SWidget
-
 	void PostUndo();
 
 protected:
+	/** Updates the preview */
+	EActiveTickReturnType UpdatePreview( double InCurrentTime, float InDeltaTime );
 
 	/** Creates the editor heading UI */
 	TSharedRef<SWidget> MakeEditorHeader() const;
@@ -413,5 +411,8 @@ protected:
 
 	/** This is used to tell the user when they do something wrong */
 	TSharedPtr<SNotificationList> NotificationListPtr;
+
+	/** Whether the active tick to update the preview is registered */
+	bool bIsActiveTickRegistered;
 };
 
diff --git a/Engine/Source/Editor/Persona/Private/SAnimationEditorViewport.cpp b/Engine/Source/Editor/Persona/Private/SAnimationEditorViewport.cpp
--- a/Engine/Source/Editor/Persona/Private/SAnimationEditorViewport.cpp
+++ b/Engine/Source/Editor/Persona/Private/SAnimationEditorViewport.cpp
@@ -269,11 +269,6 @@ FString SAnimationEditorViewportTabBody::GetDisplayString() const
 	}
 }
 
-void SAnimationEditorViewportTabBody::Tick(const FGeometry& AllottedGeometry, const double InCurrentTime, const float InDeltaTime)
-{
-	SCompoundWidget::Tick(AllottedGeometry, InCurrentTime, InDeltaTime);
-}
-
 void SAnimationEditorViewportTabBody::RefreshViewport()
 {
 	LevelViewportClient->Invalidate();
diff --git a/Engine/Source/Editor/Persona/Private/SAnimationScrubPanel.cpp b/Engine/Source/Editor/Persona/Private/SAnimationScrubPanel.cpp
--- a/Engine/Source/Editor/Persona/Private/SAnimationScrubPanel.cpp
+++ b/Engine/Source/Editor/Persona/Private/SAnimationScrubPanel.cpp
@@ -349,8 +349,6 @@ void SAnimationScrubPanel::Tick( const FGeometry& AllottedGeometry, const double
 			bSliderBeingDragged = false;
 		}
 	}
-
-	SCompoundWidget::Tick( AllottedGeometry, InCurrentTime, InDeltaTime );
 }
 
 class UAnimSingleNodeInstance* SAnimationScrubPanel::GetPreviewInstance() const
diff --git a/Engine/Source/Editor/Persona/Private/SAnimationSequenceBrowser.cpp b/Engine/Source/Editor/Persona/Private/SAnimationSequenceBrowser.cpp
--- a/Engine/Source/Editor/Persona/Private/SAnimationSequenceBrowser.cpp
+++ b/Engine/Source/Editor/Persona/Private/SAnimationSequenceBrowser.cpp
@@ -401,6 +401,10 @@ void SAnimationSequenceBrowser::Construct(const FArguments& InArgs)
 	CurrentAssetHistoryIndex = INDEX_NONE;
 	bTriedToCacheOrginalAsset = false;
 
+	bIsActiveTickRegistered = false;
+	bToolTipVisualizedThisFrame = false;
+	bToolTipClosedThisFrame = false;
+
 	FContentBrowserModule& ContentBrowserModule = FModuleManager::Get().LoadModuleChecked<FContentBrowserModule>(TEXT("ContentBrowser"));
 
 	CreateAssetTooltipResources();
@@ -438,6 +442,7 @@ void SAnimationSequenceBrowser::Construct(const FArguments& InArgs)
 	
 	Config.OnGetCustomAssetToolTip = FOnGetCustomAssetToolTip::CreateSP(this, &SAnimationSequenceBrowser::CreateCustomAssetToolTip);
 	Config.OnVisualizeAssetToolTip = FOnVisualizeAssetToolTip::CreateSP(this, &SAnimationSequenceBrowser::OnVisualizeAssetToolTip);
+	Config.OnAssetToolTipClosing = FOnAssetToolTipClosing::CreateSP( this, &SAnimationSequenceBrowser::OnAssetToolTipClosing );
 
 	TWeakPtr< SMenuAnchor > BackMenuAnchorPtr;
 	TWeakPtr< SMenuAnchor > FwdMenuAnchorPtr;	
@@ -903,7 +908,7 @@ bool SAnimationSequenceBrowser::OnVisualizeAssetToolTip(const TSharedPtr<SWidget
 				PreviewComponent->SetSkeletalMesh(MeshToUse);
 			}
 
-			PreviewComponent->EnablePreview(true, Asset, NULL);
+			PreviewComponent->EnablePreview(true, Asset, nullptr);
 			PreviewComponent->PreviewInstance->PlayAnim(true);
 
 			float HalfFov = FMath::DegreesToRadians(ViewportClient->ViewFOV) / 2.0f;
@@ -913,6 +918,14 @@ bool SAnimationSequenceBrowser::OnVisualizeAssetToolTip(const TSharedPtr<SWidget
 			ViewportClient->SetViewLocationForOrbiting(FVector(0.0f, 0.0f, MeshToUse->Bounds.BoxExtent.Z / 2.0f), TargetDist);
 
 			ViewportWidget->SetVisibility(EVisibility::Visible);
+			
+			// Update the preview as long as the tooltip is visible
+			if ( !bIsActiveTickRegistered )
+			{
+				bIsActiveTickRegistered = true;
+				RegisterActiveTick(0.f, FWidgetActiveTickDelegate::CreateSP(this, &SAnimationSequenceBrowser::UpdateTootipPreview));
+			}
+			bToolTipVisualizedThisFrame = true;
 		}
 		else
 		{
@@ -925,6 +938,15 @@ bool SAnimationSequenceBrowser::OnVisualizeAssetToolTip(const TSharedPtr<SWidget
 	return false;
 }
 
+void SAnimationSequenceBrowser::OnAssetToolTipClosing()
+{
+	// Make sure that the tooltip isn't about to preview another animation
+	if (!bToolTipVisualizedThisFrame)
+	{
+		ViewportWidget->SetVisibility(EVisibility::Hidden);
+	}
+}
+
 void SAnimationSequenceBrowser::CleanupPreviewSceneComponent(USceneComponent* Component)
 {
 	if(Component)
@@ -939,15 +961,21 @@ void SAnimationSequenceBrowser::CleanupPreviewSceneComponent(USceneComponent* Co
 	}
 }
 
-void SAnimationSequenceBrowser::Tick(const FGeometry& AllottedGeometry, const double InCurrentTime, const float InDeltaTime)
+EActiveTickReturnType SAnimationSequenceBrowser::UpdateTootipPreview( double InCurrentTime, float InDeltaTime )
 {
-	SCompoundWidget::Tick(AllottedGeometry, InCurrentTime, InDeltaTime);
-
-	if(PreviewComponent && ViewportWidget->GetVisibility() == EVisibility::Visible)
+	bToolTipVisualizedThisFrame = false;
+	if ( PreviewComponent && IsToolTipPreviewVisible() )
 	{
 		// Tick the world to update preview viewport for tooltips
-		PreviewComponent->GetScene()->GetWorld()->Tick(LEVELTICK_All, InDeltaTime);
+		PreviewComponent->GetScene()->GetWorld()->Tick( LEVELTICK_All, InDeltaTime );
+	}
+	else
+	{
+		bIsActiveTickRegistered = false;
+		return EActiveTickReturnType::StopTicking;
 	}
+
+	return EActiveTickReturnType::KeepTicking;
 }
 
 bool SAnimationSequenceBrowser::IsToolTipPreviewVisible()
diff --git a/Engine/Source/Editor/Persona/Private/SAnimationSequenceBrowser.h b/Engine/Source/Editor/Persona/Private/SAnimationSequenceBrowser.h
--- a/Engine/Source/Editor/Persona/Private/SAnimationSequenceBrowser.h
+++ b/Engine/Source/Editor/Persona/Private/SAnimationSequenceBrowser.h
@@ -75,11 +75,6 @@ public:
 
 	/** Refresh list */
 	void SelectAsset(UAnimationAsset * AnimAsset);
-	
-	/**
-	 * Update the widget
-	 */
-	virtual void Tick(const FGeometry& AllottedGeometry, const double InCurrentTime, const float InDeltaTime) override;
 
 protected:
 	bool CanShowColumnForAssetRegistryTag(FName AssetType, FName TagName) const;
@@ -138,6 +133,11 @@ protected:
 	 */
 	bool OnVisualizeAssetToolTip(const TSharedPtr<SWidget>& TooltipContent, FAssetData& AssetData);
 
+	/**
+	* Called as a tooltip is closing
+	*/
+	void OnAssetToolTipClosing();
+
 	/**
 	 * Cleaup a component from the preview scene along with any attached children
 	 */
@@ -198,4 +198,17 @@ protected:
 	FGetCurrentSelectionDelegate GetCurrentSelectionDelegate;
 
 	void RetargetAnimationHandler(USkeleton* OldSkeleton, USkeleton* NewSkeleton, bool bRemapReferencedAssets, bool bConvertSpaces, TArray<TWeakObjectPtr<UObject>> InAnimAssets);
+
+private:
+	/** Updates the animation preview in the tooltip */
+	EActiveTickReturnType UpdateTootipPreview( double InCurrentTime, float InDeltaTime );
+
+	/** Whether the tooltip was visualized this frame */
+	uint8 bToolTipClosedThisFrame : 1;
+
+	/** Whether the active tick is currently registered */
+	uint8 bIsActiveTickRegistered : 1;
+
+	/** Whether the active tick should stop */
+	uint8 bToolTipVisualizedThisFrame : 1;
 };
\ No newline at end of file
diff --git a/Engine/Source/Editor/Persona/Private/SMontageEditor.cpp b/Engine/Source/Editor/Persona/Private/SMontageEditor.cpp
--- a/Engine/Source/Editor/Persona/Private/SMontageEditor.cpp
+++ b/Engine/Source/Editor/Persona/Private/SMontageEditor.cpp
@@ -31,7 +31,7 @@ void SMontageEditor::Construct(const FArguments& InArgs)
 	check(MontageObj);
 
 	bDragging = false;
-	bRebuildMontagePanel = false;
+	bIsActiveTickRegistered = false;
 
 	SAnimEditorBase::Construct( SAnimEditorBase::FArguments()
 		.Persona(InArgs._Persona)
@@ -297,18 +297,26 @@ void SMontageEditor::RebuildMontagePanel()
 {
 	SortAndUpdateMontage();
 	AnimMontageSectionsPanel->Update();
-	bRebuildMontagePanel = false;
+}
+
+EActiveTickReturnType SMontageEditor::TriggerRebuildMontagePanel(double InCurrentTime, float InDeltaTime)
+{
+	RebuildMontagePanel();
+
+	bIsActiveTickRegistered = false;
+	return EActiveTickReturnType::StopTicking;
 }
 
 void SMontageEditor::OnMontageChange(class UObject *EditorAnimBaseObj, bool Rebuild)
 {
 	bDragging = false;
 
-	if(MontageObj!=NULL)
+	if ( MontageObj != nullptr )
 	{
-		if(Rebuild)
+		if(Rebuild && !bIsActiveTickRegistered)
 		{
-			bRebuildMontagePanel = true;
+			bIsActiveTickRegistered = true;
+			RegisterActiveTick(0.f, FWidgetActiveTickDelegate::CreateSP(this, &SMontageEditor::TriggerRebuildMontagePanel));
 		} 
 		else
 		{
@@ -465,7 +473,7 @@ bool SMontageEditor::ClampToEndTime(float NewEndTime)
 
 void SMontageEditor::AddNewSection(float StartTime, FString SectionName)
 {
-	if(MontageObj!=NULL)
+	if ( MontageObj != nullptr )
 	{
 		const FScopedTransaction Transaction( LOCTEXT("AddNewSection", "Add New Section") );
 		MontageObj->Modify();
@@ -519,7 +527,7 @@ void SMontageEditor::RenameSlotNode(int32 SlotIndex, FString NewSlotName)
 
 void SMontageEditor::AddNewMontageSlot( FString NewSlotName )
 {
-	if(MontageObj != NULL)
+	if ( MontageObj != nullptr )
 	{
 		const FScopedTransaction Transaction( LOCTEXT("AddSlot", "Add Slot") );
 		MontageObj->Modify();
@@ -544,7 +552,7 @@ FText SMontageEditor::GetMontageSlotName(int32 SlotIndex) const
 
 void SMontageEditor::RemoveMontageSlot(int32 AnimSlotIndex)
 {
-	if(MontageObj != NULL && MontageObj->SlotAnimTracks.IsValidIndex(AnimSlotIndex))
+	if ( MontageObj != nullptr && MontageObj->SlotAnimTracks.IsValidIndex( AnimSlotIndex ) )
 	{
 		const FScopedTransaction Transaction( LOCTEXT("RemoveSlot", "Remove Slot") );
 		MontageObj->Modify();
@@ -558,7 +566,7 @@ void SMontageEditor::RemoveMontageSlot(int32 AnimSlotIndex)
 void SMontageEditor::ShowSectionInDetailsView(int32 SectionIndex)
 {
 	UEditorCompositeSection *Obj = Cast<UEditorCompositeSection>(ShowInDetailsView(UEditorCompositeSection::StaticClass()));
-	if(Obj != NULL)
+	if ( Obj != nullptr )
 	{
 		Obj->InitSection(SectionIndex);
 	}
@@ -594,7 +602,7 @@ void SMontageEditor::RestartPreviewPlayAllSections()
 
 void SMontageEditor::MakeDefaultSequentialSections()
 {
-	check(MontageObj!=NULL);
+	check( MontageObj != nullptr );
 	SortSections();
 	for(int32 SectionIdx=0; SectionIdx < MontageObj->CompositeSections.Num(); SectionIdx++)
 	{
@@ -605,7 +613,7 @@ void SMontageEditor::MakeDefaultSequentialSections()
 
 void SMontageEditor::ClearSquenceOrdering()
 {
-	check(MontageObj!=NULL);
+	check( MontageObj != nullptr );
 	SortSections();
 	for(int32 SectionIdx=0; SectionIdx < MontageObj->CompositeSections.Num(); SectionIdx++)
 	{
@@ -614,18 +622,6 @@ void SMontageEditor::ClearSquenceOrdering()
 	RestartPreview();
 }
 
-void SMontageEditor::Tick( const FGeometry& AllottedGeometry, const double InCurrentTime, const float InDeltaTime )
-{
-	// we should not update any property related within PostEditChange, 
-	// so this is defered to Tick, when it needs to rebuild, just mark it and this will update in the next tick
-	if (bRebuildMontagePanel)
-	{
-		RebuildMontagePanel();
-	}
-
-	SWidget::Tick(AllottedGeometry, InCurrentTime, InDeltaTime);
-}
-
 void SMontageEditor::PostUndo()
 {
 	// when undo or redo happens, we still have to recalculate length, so we can't rely on sequence length changes or not
@@ -634,7 +630,7 @@ void SMontageEditor::PostUndo()
 		MontageObj->SequenceLength = 0.f;
 	}
 
-	RebuildMontagePanel(); //Rebuild here, undoing adds can cause slate to crash later on if we don't
+	RebuildMontagePanel(); //Rebuild here, undoing adds can cause slate to crash later on if we don't (using dummy args since they aren't used by the method
 }
 #undef LOCTEXT_NAMESPACE
 
diff --git a/Engine/Source/Editor/Persona/Private/SMontageEditor.h b/Engine/Source/Editor/Persona/Private/SMontageEditor.h
--- a/Engine/Source/Editor/Persona/Private/SMontageEditor.h
+++ b/Engine/Source/Editor/Persona/Private/SMontageEditor.h
@@ -39,10 +39,6 @@ private:
 	TSharedPtr<class SAnimMontageSectionsPanel> AnimMontageSectionsPanel;
 	TSharedPtr<class SAnimMontageScrubPanel> AnimMontageScrubPanel;
 
-	/** do I have to update montage **/
-	bool bRebuildMontagePanel;
-	void RebuildMontagePanel();
-
 protected:
 	// Begin SAnimEditorBase interface
 	virtual TSharedRef<class SAnimationScrubPanel> ConstructAnimScrubPanel() override;
@@ -63,18 +59,21 @@ private:
 	/** Pointer to the animation sequence being edited */
 	UAnimMontage* MontageObj;
 
+	/** If previewing section, it is section used to restart previewing when play button is pushed */
+	int32 PreviewingStartSectionIdx;
+
 	/** If currently previewing all or selected section */
-	bool bPreviewingAllSections;
+	bool bPreviewingAllSections : 1;
 
 	/** If currently previewing tracks instead of sections */
-	bool bPreviewingTracks;
-
-	/** If previewing section, it is section used to restart previewing when play button is pushed */
-	int32 PreviewingStartSectionIdx;
+	bool bPreviewingTracks : 1;
 
 	/** If user is currently dragging an item */
-	bool bDragging;
-	
+	bool bDragging : 1;
+
+	/** If the active tick to trigger a montage panel rebuild is currently registered */
+	bool bIsActiveTickRegistered : 1;
+
 	virtual float CalculateSequenceLengthOfEditorObject() const override;
 	void SortAndUpdateMontage();
 	void CollapseMontage();
@@ -93,6 +92,12 @@ private:
 	/** Updates Notify trigger offsets to take into account current montage state */
 	void RefreshNotifyTriggerOffsets();
 
+	/** One-off active tick to trigger a montage panel rebuild */
+	EActiveTickReturnType TriggerRebuildMontagePanel(double InCurrentTime, float InDeltaTime);
+
+	/** Rebuilds the montage panel */
+	void RebuildMontagePanel();
+
 protected:
 	virtual void InitDetailsViewEditorObject(class UEditorAnimBaseObj* EdObj) override;
 
@@ -127,10 +132,6 @@ public:
 	void			PreAnimUpdate();
 	void			PostAnimUpdate();
 
-	// SWidget interface start
-	virtual void Tick( const FGeometry& AllottedGeometry, const double InCurrentTime, const float InDeltaTime ) override;
-	// SWidget interface end
-
 	// Begin SAnimEditorBase interface
 	virtual TSharedRef<SWidget> CreateDocumentAnchor() override;
 	// End SAnimEditorBase interface
diff --git a/Engine/Source/Editor/PhAT/Private/PhAT.cpp b/Engine/Source/Editor/PhAT/Private/PhAT.cpp
--- a/Engine/Source/Editor/PhAT/Private/PhAT.cpp
+++ b/Engine/Source/Editor/PhAT/Private/PhAT.cpp
@@ -520,7 +520,7 @@ void FPhAT::RefreshHierachyTree()
 
 		// Force the tree to refresh now instead of next tick
 		const FGeometry Stub;
-		Hierarchy->Tick(Stub, 0.0, 0.0f);
+		Hierarchy->Tick( Stub, 0.0f, 0.0f );
 
 		if (!InsideSelChanged && (Hierarchy->GetNumItemsSelected() > 0))
 		{
diff --git a/Engine/Source/Editor/PhAT/Private/PhATPreviewViewportClient.cpp b/Engine/Source/Editor/PhAT/Private/PhATPreviewViewportClient.cpp
--- a/Engine/Source/Editor/PhAT/Private/PhATPreviewViewportClient.cpp
+++ b/Engine/Source/Editor/PhAT/Private/PhATPreviewViewportClient.cpp
@@ -650,6 +650,9 @@ void FPhATEdPreviewViewportClient::Tick(float DeltaSeconds)
 {
 	FEditorViewportClient::Tick(DeltaSeconds);
 
+	UWorld* World = SharedData->PreviewScene.GetWorld();
+	World->Tick(LEVELTICK_All, DeltaSeconds);
+
 	if (SharedData->bRunningSimulation)
 	{
 		// check if PIE disabled the realtime viewport and quit sim if so
@@ -660,7 +663,6 @@ void FPhATEdPreviewViewportClient::Tick(float DeltaSeconds)
 			Invalidate();
 		}
 
-		UWorld* World = SharedData->PreviewScene.GetWorld();
 		AWorldSettings* Setting = World->GetWorldSettings();
 		Setting->WorldGravityZ = SharedData->bNoGravitySimulation ? 0.0f : UPhysicsSettings::Get()->DefaultGravityZ*SharedData->EditorSimOptions->GravScale;
 		Setting->bWorldGravitySet = true;
diff --git a/Engine/Source/Editor/PhAT/Private/SPhATPreviewViewport.cpp b/Engine/Source/Editor/PhAT/Private/SPhATPreviewViewport.cpp
--- a/Engine/Source/Editor/PhAT/Private/SPhATPreviewViewport.cpp
+++ b/Engine/Source/Editor/PhAT/Private/SPhATPreviewViewport.cpp
@@ -72,13 +72,6 @@ void SPhATPreviewViewport::RefreshViewport()
 	Viewport->InvalidateDisplay();
 }
 
-void SPhATPreviewViewport::Tick(const FGeometry& AllottedGeometry, const double InCurrentTime, const float InDeltaTime)
-{
-	SCompoundWidget::Tick(AllottedGeometry, InCurrentTime, InDeltaTime);
-
-	ViewportClient->GetScene()->GetWorld()->Tick(LEVELTICK_All, InDeltaTime);
-}
-
 bool SPhATPreviewViewport::IsVisible() const
 {
 	return ViewportWidget.IsValid() && (!ParentTab.IsValid() || ParentTab.Pin()->IsForeground());
diff --git a/Engine/Source/Editor/PlacementMode/Private/SPlacementModeTools.h b/Engine/Source/Editor/PlacementMode/Private/SPlacementModeTools.h
--- a/Engine/Source/Editor/PlacementMode/Private/SPlacementModeTools.h
+++ b/Engine/Source/Editor/PlacementMode/Private/SPlacementModeTools.h
@@ -99,7 +99,7 @@ private:
 	void RefreshVolumes();
 
 	// Begin SWidget
-	void Tick( const FGeometry& AllottedGeometry, const double InCurrentTime, const float InDeltaTime ) override;
+	virtual void Tick( const FGeometry& AllottedGeometry, const double InCurrentTime, const float InDeltaTime ) override;
 	virtual FReply OnKeyDown( const FGeometry& MyGeometry, const FKeyEvent& InKeyEvent ) override;
 	// End SWidget
 
diff --git a/Engine/Source/Editor/PropertyEditor/Private/PropertyEditorToolkit.cpp b/Engine/Source/Editor/PropertyEditor/Private/PropertyEditorToolkit.cpp
--- a/Engine/Source/Editor/PropertyEditor/Private/PropertyEditorToolkit.cpp
+++ b/Engine/Source/Editor/PropertyEditor/Private/PropertyEditorToolkit.cpp
@@ -607,11 +607,11 @@ void FPropertyEditorToolkit::TickPinColorAndOpacity()
 	{
 		if ( PinSequence.IsAtStart() )
 		{
-			PinSequence.Play();
+			PinSequence.Play( PropertyTree.ToSharedRef() );
 		}
 		else
 		{
-			PinSequence.PlayReverse();
+			PinSequence.PlayReverse( PropertyTree.ToSharedRef() );
 		}
 	}
 
diff --git a/Engine/Source/Editor/PropertyEditor/Private/SDetailsViewBase.cpp b/Engine/Source/Editor/PropertyEditor/Private/SDetailsViewBase.cpp
--- a/Engine/Source/Editor/PropertyEditor/Private/SDetailsViewBase.cpp
+++ b/Engine/Source/Editor/PropertyEditor/Private/SDetailsViewBase.cpp
@@ -642,7 +642,7 @@ FReply SDetailsViewBase::OnFocusReceived(const FGeometry& MyGeometry, const FFoc
 }
 
 /** Ticks the property view.  This function performs a data consistency check */
-void SDetailsViewBase::Tick(const FGeometry& AllottedGeometry, const double InCurrentTime, const float InDeltaTime)
+void SDetailsViewBase::Tick( const FGeometry& AllottedGeometry, const double InCurrentTime, const float InDeltaTime )
 {
 	for (int32 i = 0; i < CustomizationClassInstancesPendingDelete.Num(); ++i)
 	{
@@ -742,11 +742,6 @@ void SDetailsViewBase::Tick(const FGeometry& AllottedGeometry, const double InCu
 		}
 	}
 
-	if (ThumbnailPool.IsValid())
-	{
-		ThumbnailPool->Tick(InDeltaTime);
-	}
-
 	if (DetailLayout.IsValid())
 	{
 		DetailLayout->Tick(InDeltaTime);
diff --git a/Engine/Source/Editor/PropertyEditor/Private/UserInterface/PropertyTable/SPropertyTable.h b/Engine/Source/Editor/PropertyEditor/Private/UserInterface/PropertyTable/SPropertyTable.h
--- a/Engine/Source/Editor/PropertyEditor/Private/UserInterface/PropertyTable/SPropertyTable.h
+++ b/Engine/Source/Editor/PropertyEditor/Private/UserInterface/PropertyTable/SPropertyTable.h
@@ -106,10 +106,10 @@ public:
 		}
 	}
 
-	void Tick( const FGeometry& AllottedGeometry, const double InCurrentTime, const float InDeltaTime ) override
+	virtual void Tick( const FGeometry& AllottedGeometry, const double InCurrentTime, const float InDeltaTime ) override
 	{
 		Table->Tick();
-		STreeView< TSharedRef< IPropertyTableRow > >::Tick( AllottedGeometry, InCurrentTime, InDeltaTime  );
+		STreeView< TSharedRef< IPropertyTableRow > >::Tick(AllottedGeometry, InCurrentTime, InDeltaTime);
 	}
 
 	virtual void RequestRefresh() override
diff --git a/Engine/Source/Editor/PropertyEditor/Private/UserInterface/PropertyTable/SPropertyTableCell.cpp b/Engine/Source/Editor/PropertyEditor/Private/UserInterface/PropertyTable/SPropertyTableCell.cpp
--- a/Engine/Source/Editor/PropertyEditor/Private/UserInterface/PropertyTable/SPropertyTableCell.cpp
+++ b/Engine/Source/Editor/PropertyEditor/Private/UserInterface/PropertyTable/SPropertyTableCell.cpp
@@ -91,8 +91,8 @@ void SPropertyTableCell::EnteredEditMode()
 {
 	if (Cell->IsValid())
 	{
-		// We delay the activation of editing mode till Tick due to mouse related input replys stomping on the focus
-		bEnterEditingMode = true;
+		// We delay the activation of editing mode till Tick due to mouse related input replies stomping on the focus
+		RegisterActiveTick(0.f, FWidgetActiveTickDelegate::CreateSP(this, &SPropertyTableCell::TriggerEnterEditingMode));
 	}
 }
 
@@ -110,33 +110,28 @@ void SPropertyTableCell::ExitedEditMode()
 	}
 }
 
-void SPropertyTableCell::Tick( const FGeometry& AllottedGeometry, const double InCurrentTime, const float InDeltaTime )
+EActiveTickReturnType SPropertyTableCell::TriggerEnterEditingMode(double InCurrentTime, float InDeltaTime)
 {
-	SCompoundWidget::Tick( AllottedGeometry, InCurrentTime, InDeltaTime );
-
-	if ( bEnterEditingMode )
+	if (Cell->GetTable()->GetCurrentCell() == Cell)
 	{
-		if ( Cell->GetTable()->GetCurrentCell() == Cell )
+		if (Presenter.IsValid())
 		{
-			if ( Presenter.IsValid() )
-			{
-				SetContent( ConstructCellContents() );
+			SetContent(ConstructCellContents());
 
-				if ( DropDownAnchor.IsValid() && Presenter->RequiresDropDown() )
-				{
-					DropDownAnchor->SetIsOpen( true, false );
-				}
-
-				FSlateApplication::Get().SetKeyboardFocus( Presenter->WidgetToFocusOnEdit(), EFocusCause::SetDirectly );
-			}
-			else
+			if (DropDownAnchor.IsValid() && Presenter->RequiresDropDown())
 			{
-				FSlateApplication::Get().SetKeyboardFocus( ChildSlot.GetChildAt( 0 ), EFocusCause::SetDirectly );
+				DropDownAnchor->SetIsOpen(true, false);
 			}
-		}
 
-		bEnterEditingMode = false;
+			FSlateApplication::Get().SetKeyboardFocus(Presenter->WidgetToFocusOnEdit(), EFocusCause::SetDirectly);
+		}
+		else
+		{
+			FSlateApplication::Get().SetKeyboardFocus(ChildSlot.GetChildAt(0), EFocusCause::SetDirectly);
+		}
 	}
+
+	return EActiveTickReturnType::StopTicking;
 }
 
 int32 SPropertyTableCell::OnPaint( const FPaintArgs& Args, const FGeometry& AllottedGeometry, const FSlateRect& MyClippingRect, FSlateWindowElementList& OutDrawElements, int32 LayerId, const FWidgetStyle& InWidgetStyle, bool bParentEnabled ) const
diff --git a/Engine/Source/Editor/PropertyEditor/Private/UserInterface/PropertyTable/SPropertyTableCell.h b/Engine/Source/Editor/PropertyEditor/Private/UserInterface/PropertyTable/SPropertyTableCell.h
--- a/Engine/Source/Editor/PropertyEditor/Private/UserInterface/PropertyTable/SPropertyTableCell.h
+++ b/Engine/Source/Editor/PropertyEditor/Private/UserInterface/PropertyTable/SPropertyTableCell.h
@@ -36,8 +36,6 @@ class SPropertyTableCell : public SCompoundWidget
 
 	virtual FReply OnMouseButtonDoubleClick( const FGeometry& InMyGeometry, const FPointerEvent& InMouseEvent ) override;
 
-	virtual void Tick( const FGeometry& AllottedGeometry, const double InCurrentTime, const float InDeltaTime ) override;
-
 	virtual int32 OnPaint( const FPaintArgs& Args, const FGeometry& AllottedGeometry, const FSlateRect& MyClippingRect, FSlateWindowElementList& OutDrawElements, int32 LayerId, const FWidgetStyle& InWidgetStyle, bool bParentEnabled ) const override;
 
 
@@ -63,6 +61,8 @@ private:
 
 	void OnCellValueChanged( UObject* Object, FPropertyChangedEvent& PropertyChangedEvent );
 
+	/** One-off active tick to trigger entering the editing mode */
+	EActiveTickReturnType TriggerEnterEditingMode(double InCurrentTime, float InDeltaTime);
 
 private:
 
diff --git a/Engine/Source/Editor/PropertyEditor/Private/UserInterface/PropertyTable/SPropertyTableRow.h b/Engine/Source/Editor/PropertyEditor/Private/UserInterface/PropertyTable/SPropertyTableRow.h
--- a/Engine/Source/Editor/PropertyEditor/Private/UserInterface/PropertyTable/SPropertyTableRow.h
+++ b/Engine/Source/Editor/PropertyEditor/Private/UserInterface/PropertyTable/SPropertyTableRow.h
@@ -151,10 +151,9 @@ public:
 		return Reply;
 	}
 
-	void Tick( const FGeometry& AllottedGeometry, const double InCurrentTime, const float InDeltaTime )
+	virtual void Tick( const FGeometry& AllottedGeometry, const double InCurrentTime, const float InDeltaTime ) override
 	{
 		Row->Tick();
-		SPropertyTableRowBase::Tick( AllottedGeometry, InCurrentTime, InDeltaTime );
 	}
 
 
diff --git a/Engine/Source/Editor/ReferenceViewer/Private/SReferenceViewer.cpp b/Engine/Source/Editor/ReferenceViewer/Private/SReferenceViewer.cpp
--- a/Engine/Source/Editor/ReferenceViewer/Private/SReferenceViewer.cpp
+++ b/Engine/Source/Editor/ReferenceViewer/Private/SReferenceViewer.cpp
@@ -253,16 +253,6 @@ void SReferenceViewer::SetGraphRootPackageNames(const TArray<FName>& NewGraphRoo
 	HistoryManager.AddHistoryData();
 }
 
-void SReferenceViewer::Tick(const FGeometry& AllottedGeometry, const double InCurrentTime, const float InDeltaTime)
-{
-	SCompoundWidget::Tick(AllottedGeometry, InCurrentTime, InDeltaTime);
-
-	if ( GraphObj && GraphObj->GetAssetThumbnailPool().IsValid() )
-	{
-		GraphObj->GetAssetThumbnailPool()->Tick(InDeltaTime);
-	}
-}
-
 void SReferenceViewer::OnNodeDoubleClicked(UEdGraphNode* Node)
 {
 	TSet<UObject*> Nodes;
diff --git a/Engine/Source/Editor/ReferenceViewer/Private/SReferenceViewer.h b/Engine/Source/Editor/ReferenceViewer/Private/SReferenceViewer.h
--- a/Engine/Source/Editor/ReferenceViewer/Private/SReferenceViewer.h
+++ b/Engine/Source/Editor/ReferenceViewer/Private/SReferenceViewer.h
@@ -22,10 +22,6 @@ public:
 	/** Sets a new root package name */
 	void SetGraphRootPackageNames(const TArray<FName>& NewGraphRootPackageNames);
 
-	// SWidget implementation
-	virtual void Tick(const FGeometry& AllottedGeometry, const double InCurrentTime, const float InDeltaTime) override;
-	// End SWidget implementation
-
 private:
 
 	void OnNodeDoubleClicked(class UEdGraphNode* Node);
diff --git a/Engine/Source/Editor/SceneOutliner/Private/SSceneOutliner.cpp b/Engine/Source/Editor/SceneOutliner/Private/SSceneOutliner.cpp
--- a/Engine/Source/Editor/SceneOutliner/Private/SSceneOutliner.cpp
+++ b/Engine/Source/Editor/SceneOutliner/Private/SSceneOutliner.cpp
@@ -2242,10 +2242,7 @@ namespace SceneOutliner
 	}
 
 	void SSceneOutliner::Tick( const FGeometry& AllottedGeometry, const double InCurrentTime, const float InDeltaTime )
-	{
-		// Call parent implementation
-		SCompoundWidget::Tick( AllottedGeometry, InCurrentTime, InDeltaTime );
-
+	{	
 		for (auto& Pair : Columns)
 		{
 			Pair.Value->Tick(InCurrentTime, InDeltaTime);
diff --git a/Engine/Source/Editor/Sequencer/Private/SSequencer.cpp b/Engine/Source/Editor/Sequencer/Private/SSequencer.cpp
--- a/Engine/Source/Editor/Sequencer/Private/SSequencer.cpp
+++ b/Engine/Source/Editor/Sequencer/Private/SSequencer.cpp
@@ -46,7 +46,7 @@ public:
 		Sequencer = InSequencer;
 	}
 
-	virtual void Tick( const FGeometry& AllottedGeometry, const double InCurrentTime, const float InDeltaTime )
+	virtual void Tick( const FGeometry& AllottedGeometry, const double InCurrentTime, const float InDeltaTime ) override
 	{
 		if (Sequencer.Pin()->IsShotFilteringOn())
 		{
@@ -143,6 +143,7 @@ private:
 void SSequencer::Construct( const FArguments& InArgs, TSharedRef< class FSequencer > InSequencer )
 {
 	Sequencer = InSequencer;
+	bIsActiveTickRegistered = false;
 
 	// Create a node tree which contains a tree of movie scene data to display in the sequence
 	SequencerNodeTree = MakeShareable( new FSequencerNodeTree( InSequencer.Get() ) );
@@ -373,15 +374,28 @@ SSequencer::~SSequencer()
 {
 }
 
-void SSequencer::Tick( const FGeometry& AllottedGeometry, const double InCurrentTime, const float InDeltaTime )
+void SSequencer::RegisterActiveTickForPlayback()
 {
-	Sequencer.Pin()->Tick(InDeltaTime);
+	if (!bIsActiveTickRegistered)
+	{
+		bIsActiveTickRegistered = true;
+		RegisterActiveTick(0.f, FWidgetActiveTickDelegate::CreateSP(this, &SSequencer::EnsureSlateTickDuringPlayback));
+	}
+}
 
-	const TArray< TSharedPtr<FMovieSceneTrackEditor> >& TrackEditors = Sequencer.Pin()->GetTrackEditors();
-	for (int32 EditorIndex = 0; EditorIndex < TrackEditors.Num(); ++EditorIndex)
+EActiveTickReturnType SSequencer::EnsureSlateTickDuringPlayback(double InCurrentTime, float InDeltaTime)
+{
+	if (Sequencer.IsValid())
 	{
-		TrackEditors[EditorIndex]->Tick(InCurrentTime);
+		auto PlaybackStatus = Sequencer.Pin()->GetPlaybackStatus();
+		if (PlaybackStatus == EMovieScenePlayerStatus::Playing || PlaybackStatus == EMovieScenePlayerStatus::Recording)
+		{
+			return EActiveTickReturnType::KeepTicking;
+		}
 	}
+
+	bIsActiveTickRegistered = false;
+	return EActiveTickReturnType::StopTicking;
 }
 
 void SSequencer::UpdateLayoutTree()
diff --git a/Engine/Source/Editor/Sequencer/Private/SSequencer.h b/Engine/Source/Editor/Sequencer/Private/SSequencer.h
--- a/Engine/Source/Editor/Sequencer/Private/SSequencer.h
+++ b/Engine/Source/Editor/Sequencer/Private/SSequencer.h
@@ -72,12 +72,13 @@ public:
 
 	~SSequencer();
 
-	virtual void Tick( const FGeometry& AllottedGeometry, const double InCurrentTime, const float InDeltaTime ) override;
 	virtual bool SupportsKeyboardFocus() const override { return true; }
 
 	/** Updates the layout node tree from movie scene data */
 	void UpdateLayoutTree();
 
+	/** Causes the widget to register an empty active tick that persists until Sequencer playback stops */
+	void RegisterActiveTickForPlayback();
 
 	/**
 	 * Updates the breadcrumbs from a change in the shot filter state
@@ -91,6 +92,9 @@ public:
 	void DeleteSelectedNodes();
 
 private:
+	/** Empty active tick to ensure Slate ticks during Sequencer playback */
+	EActiveTickReturnType EnsureSlateTickDuringPlayback(double InCurrentTime, float InDeltaTime);	
+
 	/** Makes the toolbar for the outline section. */
 	TSharedRef<SWidget> MakeToolBar();
 
@@ -203,4 +207,7 @@ private:
 	TAttribute<bool> CleanViewEnabled;
 	/** The main sequencer interface */
 	TWeakPtr<FSequencer> Sequencer;
+
+	/** Whether the active tick is currently registered */
+	bool bIsActiveTickRegistered;
 };
diff --git a/Engine/Source/Editor/Sequencer/Private/Sequencer.cpp b/Engine/Source/Editor/Sequencer/Private/Sequencer.cpp
--- a/Engine/Source/Editor/Sequencer/Private/Sequencer.cpp
+++ b/Engine/Source/Editor/Sequencer/Private/Sequencer.cpp
@@ -158,6 +158,57 @@ FSequencer::~FSequencer()
 	SequencerWidget.Reset();
 }
 
+void FSequencer::Tick(float InDeltaTime)
+{
+	if (bNeedTreeRefresh)
+	{
+		// @todo - Sequencer Will be called too often
+		UpdateRuntimeInstances();
+
+		SequencerWidget->UpdateLayoutTree();
+		bNeedTreeRefresh = false;
+	}
+
+	float NewTime = GetGlobalTime() + InDeltaTime;
+	if (PlaybackState == EMovieScenePlayerStatus::Playing ||
+		PlaybackState == EMovieScenePlayerStatus::Recording)
+	{
+		TRange<float> TimeBounds = GetTimeBounds();
+		if (!TimeBounds.IsEmpty())
+		{
+			if (NewTime > TimeBounds.GetUpperBoundValue())
+			{
+				if (bLoopingEnabled)
+				{
+					NewTime -= TimeBounds.Size<float>();
+				}
+				else
+				{
+					NewTime = TimeBounds.GetUpperBoundValue();
+					PlaybackState = EMovieScenePlayerStatus::Stopped;
+				}
+			}
+
+			if (NewTime < TimeBounds.GetLowerBoundValue())
+			{
+				NewTime = TimeBounds.GetLowerBoundValue();
+			}
+
+			SetGlobalTime(NewTime);
+		}
+		else
+		{
+			// no bounds at all, stop playing
+			PlaybackState = EMovieScenePlayerStatus::Stopped;
+		}
+	}
+
+	// Tick all the tools we own as well
+	for (int32 EditorIndex = 0; EditorIndex < TrackEditors.Num(); ++EditorIndex)
+	{
+		TrackEditors[EditorIndex]->Tick(InDeltaTime);
+	}
+}
 
 TArray< UMovieScene* > FSequencer::GetMovieScenesBeingEdited()
 {
@@ -610,52 +661,6 @@ void FSequencer::RemoveMovieSceneInstance( UMovieSceneSection& MovieSceneSection
 	MovieSceneSectionToInstanceMap.Remove( &MovieSceneSection );
 }
 
-void FSequencer::Tick(const float InDeltaTime)
-{
-	if( bNeedTreeRefresh )
-	{
-		// @todo - Sequencer Will be called too often
-		UpdateRuntimeInstances();
-
-		SequencerWidget->UpdateLayoutTree();
-		bNeedTreeRefresh = false;
-	}
-
-	float NewTime = GetGlobalTime() + InDeltaTime;
-	if (PlaybackState == EMovieScenePlayerStatus::Playing ||
-		PlaybackState == EMovieScenePlayerStatus::Recording)
-	{
-		TRange<float> TimeBounds = GetTimeBounds();
-		if (!TimeBounds.IsEmpty())
-		{
-			if (NewTime > TimeBounds.GetUpperBoundValue())
-			{
-				if (bLoopingEnabled)
-				{
-					NewTime -= TimeBounds.Size<float>();
-				}
-				else
-				{
-					NewTime = TimeBounds.GetUpperBoundValue();
-					PlaybackState = EMovieScenePlayerStatus::Stopped;
-				}
-			}
-			
-			if (NewTime < TimeBounds.GetLowerBoundValue())
-			{
-				NewTime = TimeBounds.GetLowerBoundValue();
-			}
-
-			SetGlobalTime(NewTime);
-		}
-		else
-		{
-			// no bounds at all, stop playing
-			PlaybackState = EMovieScenePlayerStatus::Stopped;
-		}
-	}
-}
-
 void FSequencer::AddReferencedObjects( FReferenceCollector& Collector )
 {
 	for( int32 MovieSceneIndex = 0; MovieSceneIndex < MovieSceneStack.Num(); ++MovieSceneIndex )
@@ -694,7 +699,6 @@ void FSequencer::DestroySpawnablesForAllMovieScenes()
 	}
 }
 
-
 FReply FSequencer::OnPlay()
 {
 	if( PlaybackState == EMovieScenePlayerStatus::Playing ||
@@ -715,6 +719,9 @@ FReply FSequencer::OnPlay()
 				SetGlobalTime(TimeBounds.GetLowerBoundValue());
 			}
 			PlaybackState = EMovieScenePlayerStatus::Playing;
+			
+			// Make sure Slate ticks during playback
+			SequencerWidget->RegisterActiveTickForPlayback();
 		}
 	}
 
@@ -726,6 +733,9 @@ FReply FSequencer::OnRecord()
 	if (PlaybackState != EMovieScenePlayerStatus::Recording)
 	{
 		PlaybackState = EMovieScenePlayerStatus::Recording;
+		
+		// Make sure Slate ticks during playback
+		SequencerWidget->RegisterActiveTickForPlayback();
 
 		// @todo sequencer livecapture: Ideally we would support fixed timestep capture from simulation
 		//			Basically we need to run the PIE world at a fixed time step, capturing key frames every frame. 
@@ -849,7 +859,7 @@ void FSequencer::OnViewRangeChanged( TRange<float> NewViewRange, bool bSmoothZoo
 			{
 				LastViewRange = TargetViewRange;
 
-				ZoomAnimation.Play();
+				ZoomAnimation.Play( SequencerWidget.ToSharedRef() );
 			}
 			TargetViewRange = NewViewRange;
 		}
@@ -1253,11 +1263,11 @@ void FSequencer::FilterToShotSections(const TArray< TWeakObjectPtr<class UMovieS
 
 	if (!bWasPreviouslyFiltering && bIsNowFiltering)
 	{
-		OverlayAnimation.Play();
+		OverlayAnimation.Play( SequencerWidget.ToSharedRef() );
 	}
 	else if (bWasPreviouslyFiltering && !bIsNowFiltering) 
 	{
-		OverlayAnimation.PlayReverse();
+		OverlayAnimation.PlayReverse( SequencerWidget.ToSharedRef() );
 	}
 
 	if( bZoomToShotBounds )
diff --git a/Engine/Source/Editor/Sequencer/Private/Sequencer.h b/Engine/Source/Editor/Sequencer/Private/Sequencer.h
--- a/Engine/Source/Editor/Sequencer/Private/Sequencer.h
+++ b/Engine/Source/Editor/Sequencer/Private/Sequencer.h
@@ -36,7 +36,7 @@ struct FSequencerInitParams
 /**
  * Sequencer is the editing tool for MovieScene assets
  */
-class FSequencer : public ISequencer, public FGCObject, public FEditorUndoClient
+class FSequencer : public ISequencer, public FGCObject, public FEditorUndoClient, public FTickableEditorObject
 { 
 
 public:
@@ -58,6 +58,11 @@ public:
 	/** FGCObject interface */
 	virtual void AddReferencedObjects(FReferenceCollector& Collector) override;
 
+	// FTickableEditorObject interface
+	virtual void Tick(float DeltaTime) override;
+	virtual bool IsTickable() const override { return true; }
+	virtual TStatId GetStatId() const override { RETURN_QUICK_DECLARE_CYCLE_STAT(FSequencer, STATGROUP_Tickables); };
+	// End FTickableEditorObject interface
 
 	/** ISequencer interface */
 	virtual TSharedRef<SWidget> GetSequencerWidget() const override { return SequencerWidget.ToSharedRef(); }
@@ -120,9 +125,6 @@ public:
 	 */
 	const TArray< TSharedPtr<FMovieSceneTrackEditor> >& GetTrackEditors() const { return TrackEditors; }
 
-	/** Ticks the sequencer by InDeltaTime */
-	void Tick( const float InDeltaTime );
-
 	/**
 	 * Attempts to add a new spawnable to the MovieScene for the specified asset or class
 	 *
diff --git a/Engine/Source/Editor/Sequencer/Private/TimeSliderController.cpp b/Engine/Source/Editor/Sequencer/Private/TimeSliderController.cpp
--- a/Engine/Source/Editor/Sequencer/Private/TimeSliderController.cpp
+++ b/Engine/Source/Editor/Sequencer/Private/TimeSliderController.cpp
@@ -360,7 +360,7 @@ FReply FSequencerTimeSliderController::OnMouseMove( TSharedRef<SWidget> WidgetOw
 			if (!bPanning)
 			{
 				DistanceDragged += FMath::Abs( MouseEvent.GetCursorDelta().X );
-				if ( DistanceDragged > FSlateApplication::Get().GetDragTriggerDistnace() )
+				if ( DistanceDragged > FSlateApplication::Get().GetDragTriggerDistance() )
 				{
 					bPanning = true;
 				}
@@ -407,7 +407,7 @@ FReply FSequencerTimeSliderController::OnMouseMove( TSharedRef<SWidget> WidgetOw
 			if ( !bDraggingScrubber )
 			{
 				DistanceDragged += FMath::Abs( MouseEvent.GetCursorDelta().X );
-				if ( DistanceDragged > FSlateApplication::Get().GetDragTriggerDistnace() )
+				if ( DistanceDragged > FSlateApplication::Get().GetDragTriggerDistance() )
 				{
 					bDraggingScrubber = true;
 					TimeSliderArgs.OnBeginScrubberMovement.ExecuteIfBound();
diff --git a/Engine/Source/Editor/UMGEditor/Private/Customizations/CanvasSlotCustomization.cpp b/Engine/Source/Editor/UMGEditor/Private/Customizations/CanvasSlotCustomization.cpp
--- a/Engine/Source/Editor/UMGEditor/Private/Customizations/CanvasSlotCustomization.cpp
+++ b/Engine/Source/Editor/UMGEditor/Private/Customizations/CanvasSlotCustomization.cpp
@@ -28,6 +28,7 @@ public:
 		FText LabelText,
 		FAnchors Anchors)
 	{
+		bIsHovered = false;
 		ResizeCurve = FCurveSequence(0, 0.40f);
 
 		ChildSlot
@@ -82,33 +83,40 @@ public:
 		];
 	}
 
-	void Tick(const FGeometry& AllottedGeometry, const double InCurrentTime, const float InDeltaTime) override
+	virtual void Tick( const FGeometry& AllottedGeometry, const double InCurrentTime, const float InDeltaTime )
 	{
-		if ( AllottedGeometry.IsUnderLocation(FSlateApplication::Get().GetCursorPos()) )
+		if ( bIsHovered )
 		{
 			if ( !ResizeCurve.IsPlaying() )
 			{
 				if ( ResizeCurve.IsAtStart() )
 				{
-					ResizeCurve.Play();
+					ResizeCurve.Play( this->AsShared() );
 				}
 				else if ( ResizeCurve.IsAtEnd() )
 				{
-					ResizeCurve.PlayReverse();
+					ResizeCurve.Reverse();
 				}
 			}
 		}
-		else
+		//Make sure the preview animation goes all the way back to the initial position before we stop ticking
+		else if ( !ResizeCurve.IsAtStart() && !ResizeCurve.IsInReverse() )
 		{
-			if ( !ResizeCurve.IsAtStart() && ((!ResizeCurve.IsInReverse() && ResizeCurve.IsPlaying()) || !ResizeCurve.IsPlaying()) )
-			{
-				ResizeCurve.PlayReverse();
-			}
+			ResizeCurve.Reverse();
 		}
 	}
 
-private:
+	virtual void OnMouseEnter( const FGeometry& MyGeometry, const FPointerEvent& MouseEvent )
+	{
+		bIsHovered = true;
+	}
 
+	virtual void OnMouseLeave( const FPointerEvent& MouseEvent )
+	{
+		bIsHovered = false;
+	}
+
+private:
 	FOptionalSize GetCurrentWidth() const
 	{
 		return 48 + ( 16 * ResizeCurve.GetLerp() );
@@ -156,6 +164,8 @@ private:
 
 private:
 	FCurveSequence ResizeCurve;
+	TWeakPtr<FActiveTickHandle> ActiveTickHandle;
+	bool bIsHovered;
 };
 
 
diff --git a/Engine/Source/Editor/UMGEditor/Private/Designer/SDesignSurface.cpp b/Engine/Source/Editor/UMGEditor/Private/Designer/SDesignSurface.cpp
--- a/Engine/Source/Editor/UMGEditor/Private/Designer/SDesignSurface.cpp
+++ b/Engine/Source/Editor/UMGEditor/Private/Designer/SDesignSurface.cpp
@@ -119,12 +119,11 @@ void SDesignSurface::Construct(const FArguments& InArgs)
 	ViewOffset = FVector2D::ZeroVector;
 
 	ZoomLevelFade = FCurveSequence(0.0f, 1.0f);
-	ZoomLevelFade.Play();
+	ZoomLevelFade.Play( this->AsShared() );
 
 	ZoomLevelGraphFade = FCurveSequence(0.0f, 0.5f);
-	ZoomLevelGraphFade.Play();
+	ZoomLevelGraphFade.Play( this->AsShared() );
 
-	bDeferredZoomingToFit = false;
 	bDeferredZoomToExtents = false;
 
 	bAllowContinousZoomInterpolation = false;
@@ -143,41 +142,43 @@ void SDesignSurface::Construct(const FArguments& InArgs)
 	];
 }
 
-void SDesignSurface::Tick(const FGeometry& AllottedGeometry, const double InCurrentTime, const float InDeltaTime)
+EActiveTickReturnType SDesignSurface::HandleZoomToFit( double InCurrentTime, float InDeltaTime )
 {
-	FSlateRect Bounds = ComputeAreaBounds();
+	const FVector2D DesiredViewCenter = ( ZoomTargetTopLeft + ZoomTargetBottomRight ) * 0.5f;
+	const bool bDoneScrolling = ScrollToLocation(CachedGeometry, DesiredViewCenter, bTeleportInsteadOfScrollingWhenZoomingToFit ? 1000.0f : InDeltaTime);
+	const bool bDoneZooming = ZoomToLocation(CachedGeometry.Size, ZoomTargetBottomRight - ZoomTargetTopLeft, bDoneScrolling);
+
+	if (bDoneZooming && bDoneScrolling)
+	{
+		// One final push to make sure we're centered in the end
+		ViewOffset = DesiredViewCenter - ( 0.5f * CachedGeometry.Scale * CachedGeometry.Size / GetZoomAmount() );
 
+		ZoomTargetTopLeft = FVector2D::ZeroVector;
+		ZoomTargetBottomRight = FVector2D::ZeroVector;
+
+		return EActiveTickReturnType::StopTicking;
+	}
+	
+	return EActiveTickReturnType::KeepTicking;
+}
+
+void SDesignSurface::Tick( const FGeometry& AllottedGeometry, const double InCurrentTime, const float InDeltaTime )
+{
+	CachedGeometry = AllottedGeometry;
+	
 	// Zoom to extents
+	FSlateRect Bounds = ComputeAreaBounds();
 	if ( bDeferredZoomToExtents )
 	{
 		bDeferredZoomToExtents = false;
-		//ZoomPadding = NodePanelDefs::DefaultZoomPadding;
 		ZoomTargetTopLeft = FVector2D(Bounds.Left, Bounds.Top);
 		ZoomTargetBottomRight = FVector2D(Bounds.Right, Bounds.Bottom);
-		//bDeferredZoomToSelection = false;
-		bDeferredZoomingToFit = true;
-	}
-
-	if ( bDeferredZoomingToFit )
-	{
-		const FVector2D DesiredViewCenter = ( ZoomTargetTopLeft + ZoomTargetBottomRight ) * 0.5f;
-		const bool bDoneScrolling = ScrollToLocation(AllottedGeometry, DesiredViewCenter, bTeleportInsteadOfScrollingWhenZoomingToFit ? 1000.0f : InDeltaTime);
-		bool bDoneZooming = ZoomToLocation(AllottedGeometry.Size, ZoomTargetBottomRight - ZoomTargetTopLeft, bDoneScrolling);
 
-		if ( bDoneZooming && bDoneScrolling )
+		if (!ActiveTickHandle.IsValid())
 		{
-			// One final push to make sure we centered in the end
-			ViewOffset = DesiredViewCenter - ( 0.5f * AllottedGeometry.Scale * AllottedGeometry.Size / GetZoomAmount() );
-			// Reset ZoomPadding
-			//ZoomPadding = NodePanelDefs::DefaultZoomPadding;
-			ZoomTargetTopLeft = FVector2D::ZeroVector;
-			ZoomTargetBottomRight = FVector2D::ZeroVector;
-
-			bDeferredZoomingToFit = false;
+			RegisterActiveTick(0.f, FWidgetActiveTickDelegate::CreateSP(this, &SDesignSurface::HandleZoomToFit));
 		}
 	}
-	
-	SCompoundWidget::Tick(AllottedGeometry, InCurrentTime, InDeltaTime);
 }
 
 int32 SDesignSurface::OnPaint(const FPaintArgs& Args, const FGeometry& AllottedGeometry, const FSlateRect& MyClippingRect, FSlateWindowElementList& OutDrawElements, int32 LayerId, const FWidgetStyle& InWidgetStyle, bool bParentEnabled) const
@@ -297,7 +298,7 @@ void SDesignSurface::ChangeZoomLevel(int32 ZoomLevelDelta, const FVector2D& Widg
 		PostChangedZoom();
 
 		// Note: This happens even when maxed out at a stop; so the user sees the animation and knows that they're at max zoom in/out
-		ZoomLevelFade.Play();
+		ZoomLevelFade.Play( this->AsShared() );
 
 		// Re-center the screen so that it feels like zooming around the cursor.
 		{
@@ -376,7 +377,7 @@ bool SDesignSurface::ZoomToLocation(const FVector2D& CurrentSizeWithoutZoom, con
 			// Animate to it
 			PreviousZoomLevel = ZoomLevel;
 			ZoomLevel = FMath::Clamp(DesiredZoom, 0, NumZoomLevels - 1);
-			ZoomLevelGraphFade.Play();
+			ZoomLevelGraphFade.Play(this->AsShared());
 			return false;
 		}
 		else
@@ -386,7 +387,7 @@ bool SDesignSurface::ZoomToLocation(const FVector2D& CurrentSizeWithoutZoom, con
 			{
 				// Zooming out; do it instantly
 				ZoomLevel = PreviousZoomLevel = DesiredZoom;
-				ZoomLevelFade.Play();
+				ZoomLevelFade.Play(this->AsShared());
 			}
 			else
 			{
@@ -394,7 +395,7 @@ bool SDesignSurface::ZoomToLocation(const FVector2D& CurrentSizeWithoutZoom, con
 				if ( bDoneScrolling )
 				{
 					ZoomLevel = PreviousZoomLevel = DesiredZoom;
-					ZoomLevelFade.Play();
+					ZoomLevelFade.Play(this->AsShared());
 				}
 			}
 		}
@@ -407,9 +408,8 @@ bool SDesignSurface::ZoomToLocation(const FVector2D& CurrentSizeWithoutZoom, con
 
 void SDesignSurface::ZoomToFit(bool bInstantZoom)
 {
-	bDeferredZoomingToFit = true;
-	bDeferredZoomToExtents = true;
 	bTeleportInsteadOfScrollingWhenZoomingToFit = bInstantZoom;
+	bDeferredZoomToExtents = true;
 }
 
 FText SDesignSurface::GetZoomText() const
diff --git a/Engine/Source/Editor/UMGEditor/Private/Designer/SDesignSurface.h b/Engine/Source/Editor/UMGEditor/Private/Designer/SDesignSurface.h
--- a/Engine/Source/Editor/UMGEditor/Private/Designer/SDesignSurface.h
+++ b/Engine/Source/Editor/UMGEditor/Private/Designer/SDesignSurface.h
@@ -22,7 +22,7 @@ public:
 	void Construct(const FArguments& InArgs);
 
 	// SWidget interface
-	virtual void Tick(const FGeometry& AllottedGeometry, const double InCurrentTime, const float InDeltaTime);
+	virtual void Tick( const FGeometry& AllottedGeometry, const double InCurrentTime, const float InDeltaTime );
 	virtual FReply OnMouseButtonDown(const FGeometry& MyGeometry, const FPointerEvent& MouseEvent) override;
 	virtual FReply OnMouseButtonUp(const FGeometry& MyGeometry, const FPointerEvent& MouseEvent) override;
 	virtual FReply OnMouseMove(const FGeometry& MyGeometry, const FPointerEvent& MouseEvent) override;
@@ -88,12 +88,6 @@ protected:
 	// The interface for mapping ZoomLevel values to actual node scaling values
 	TScopedPointer<FZoomLevelsContainer> ZoomLevels;
 
-	// A flag noting if we have a pending zoom to fit operation to perform next tick.
-	bool bDeferredZoomingToFit;
-
-	// A flag noting if we have a pending zoom to extents operation to perform next tick.
-	bool bDeferredZoomToExtents;
-
 	bool bAllowContinousZoomInterpolation;
 
 	bool bTeleportInsteadOfScrollingWhenZoomingToFit;
@@ -104,4 +98,17 @@ protected:
 
 	/** Does the user need to press Control in order to over-zoom. */
 	bool bRequireControlToOverZoom;
+
+	/** Cached geometry for use within the active tick */
+	FGeometry CachedGeometry;
+
+private:
+	/** Active tick that handles deferred zooming until the target zoom is reached */
+	EActiveTickReturnType HandleZoomToFit(double InCurrentTime, float InDeltaTime);
+
+	/** The handle to the active tick */
+	TWeakPtr<FActiveTickHandle> ActiveTickHandle;
+
+	// A flag noting if we have a pending zoom to extents operation to perform next tick.
+	bool bDeferredZoomToExtents;
 };
diff --git a/Engine/Source/Editor/UMGEditor/Private/Designer/SDesignerView.cpp b/Engine/Source/Editor/UMGEditor/Private/Designer/SDesignerView.cpp
--- a/Engine/Source/Editor/UMGEditor/Private/Designer/SDesignerView.cpp
+++ b/Engine/Source/Editor/UMGEditor/Private/Designer/SDesignerView.cpp
@@ -177,7 +177,9 @@ void SDesignerView::Construct(const FArguments& InArgs, TSharedPtr<FWidgetBluepr
 	SetStartupResolution();
 
 	ResolutionTextFade = FCurveSequence(0.0f, 1.0f);
-	ResolutionTextFade.Play();
+	ResolutionTextFade.Play(this->AsShared());
+
+	HoveredWidgetOutlineFade = FCurveSequence(0.0f, 0.15f);
 
 	bMovingExistingWidget = false;
 
@@ -294,6 +296,8 @@ void SDesignerView::Construct(const FArguments& InArgs, TSharedPtr<FWidgetBluepr
 	);
 
 	BlueprintEditor.Pin()->OnSelectedWidgetsChanged.AddRaw(this, &SDesignerView::OnEditorSelectionChanged);
+	BlueprintEditor.Pin()->OnHoveredWidgetSet.AddRaw(this, &SDesignerView::SetHoveredWidget);
+	BlueprintEditor.Pin()->OnHoveredWidgetCleared.AddRaw(this, &SDesignerView::ClearHoveredWidget);
 
 	ZoomToFit(/*bInstantZoom*/ true);
 }
@@ -520,7 +524,10 @@ SDesignerView::~SDesignerView()
 
 	if ( BlueprintEditor.IsValid() )
 	{
-		BlueprintEditor.Pin()->OnSelectedWidgetsChanged.RemoveAll(this);
+		auto PinnedEditor = BlueprintEditor.Pin();
+		PinnedEditor->OnSelectedWidgetsChanged.RemoveAll(this);
+		PinnedEditor->OnHoveredWidgetSet.RemoveAll(this);
+		PinnedEditor->OnHoveredWidgetCleared.RemoveAll(this);
 	}
 
 	if ( GEditor )
@@ -723,9 +730,23 @@ void SDesignerView::OnEditorSelectionChanged()
 	CreateExtensionWidgetsForSelection();
 }
 
+void SDesignerView::SetHoveredWidget(const FWidgetReference& InHoveredWidget)
+{
+	if (InHoveredWidget != HoveredWidget)
+	{
+		HoveredWidget = InHoveredWidget;
+		HoveredWidgetOutlineFade.Play(this->AsShared());
+	}
+}
+
+void SDesignerView::ClearHoveredWidget()
+{
+	HoveredWidget = FWidgetReference();
+}
+
 FGeometry SDesignerView::GetDesignerGeometry() const
 {
-	return CachedDesignerGeometry;
+	return CachedGeometry;
 }
 
 void SDesignerView::MarkDesignModifed(bool bRequiresRecompile)
@@ -857,8 +878,8 @@ FVector2D SDesignerView::GetExtensionPosition(TSharedRef<FDesignerSurfaceElement
 
 		if ( GetWidgetGeometry(SelectedWidget, SelectedWidgetGeometry_RelativeToDesigner) && GetWidgetParentGeometry(SelectedWidget, SelectedWidgetParentGeometry_RelativeToDesigner) )
 		{
-			SelectedWidgetGeometry_RelativeToDesigner.AppendTransform(FSlateLayoutTransform(Inverse(CachedDesignerGeometry.AbsolutePosition)));
-			SelectedWidgetParentGeometry_RelativeToDesigner.AppendTransform(FSlateLayoutTransform(Inverse(CachedDesignerGeometry.AbsolutePosition)));
+			SelectedWidgetGeometry_RelativeToDesigner.AppendTransform(FSlateLayoutTransform(Inverse(CachedGeometry.AbsolutePosition)));
+			SelectedWidgetParentGeometry_RelativeToDesigner.AppendTransform(FSlateLayoutTransform(Inverse(CachedGeometry.AbsolutePosition)));
 
 			const FVector2D WidgetPostion_DesignerSpace = SelectedWidgetGeometry_RelativeToDesigner.AbsolutePosition;
 			const FVector2D WidgetSize = SelectedWidgetGeometry_RelativeToDesigner.Size * GetPreviewScale();
@@ -1162,7 +1183,11 @@ FReply SDesignerView::OnMouseMove(const FGeometry& MyGeometry, const FPointerEve
 	FWidgetHitResult HitResult;
 	if ( FindWidgetUnderCursor(MyGeometry, MouseEvent, HitResult) )
 	{
-		BlueprintEditor.Pin()->SetHoveredWidget(HitResult.Widget);
+		SetHoveredWidget(HitResult.Widget);
+	}
+	else if (HoveredWidget.IsValid())
+	{
+		ClearHoveredWidget();
 	}
 
 	return FReply::Unhandled();
@@ -1172,14 +1197,14 @@ void SDesignerView::OnMouseEnter(const FGeometry& MyGeometry, const FPointerEven
 {
 	SDesignSurface::OnMouseEnter(MyGeometry, MouseEvent);
 
-	BlueprintEditor.Pin()->ClearHoveredWidget();
+	ClearHoveredWidget();
 }
 
 void SDesignerView::OnMouseLeave(const FPointerEvent& MouseEvent)
 {
 	SDesignSurface::OnMouseLeave(MouseEvent);
 
-	BlueprintEditor.Pin()->ClearHoveredWidget();
+	ClearHoveredWidget();
 }
 
 FReply SDesignerView::OnKeyDown(const FGeometry& MyGeometry, const FKeyEvent& InKeyEvent)
@@ -1242,8 +1267,7 @@ int32 SDesignerView::HandleEffectsPainting(const FOnPaintHandlerParams& PaintArg
 	static const FName SelectionOutlineName("UMGEditor.SelectionOutline");
 
 	static const FLinearColor SelectedTint(0, 1, 0);
-	const FLinearColor HoveredTint(0, 0.5, 1, FMath::Clamp(BlueprintEditor.Pin()->GetHoveredWidgetTime() / HoveredAnimationTime, 0.0f, 1.0f)); // Azure = 0x007FFF
-
+	
 	const FSlateBrush* SelectionOutlineBrush = FEditorStyle::Get().GetBrush(SelectionOutlineName);
 	FVector2D SelectionBrushInflationAmount = FVector2D(16, 16) * FVector2D(SelectionOutlineBrush->Margin.Left, SelectionOutlineBrush->Margin.Top) * ( 1.0f / GetPreviewScale() );
 
@@ -1273,12 +1297,14 @@ int32 SDesignerView::HandleEffectsPainting(const FOnPaintHandlerParams& PaintArg
 		}
 	}
 
-	FWidgetReference HoveredWidget = BlueprintEditor.Pin()->GetHoveredWidget();
 	TSharedPtr<SWidget> HoveredSlateWidget = HoveredWidget.GetPreviewSlate();
 
 	// Don't draw the hovered effect if it's also the selected widget
 	if ( HoveredSlateWidget.IsValid() && !SelectedWidgets.Contains(HoveredWidget) )
 	{
+		//const FLinearColor HoveredTint(0, 0.5, 1, FMath::Clamp(BlueprintEditor.Pin()->GetHoveredWidgetTime() / HoveredAnimationTime, 0.0f, 1.0f)); // Azure = 0x007FFF
+		const FLinearColor HoveredTint(0, 0.5, 1, HoveredWidgetOutlineFade.GetLerp()); // Azure = 0x007FFF
+
 		TSharedRef<SWidget> Widget = HoveredSlateWidget.ToSharedRef();
 
 		FArrangedWidget ArrangedWidget(SNullWidget::NullWidget, FGeometry());
@@ -1344,9 +1370,10 @@ void SDesignerView::UpdatePreviewWidget(bool bForceUpdate)
 	}
 }
 
-void SDesignerView::Tick(const FGeometry& AllottedGeometry, const double InCurrentTime, const float InDeltaTime)
+void SDesignerView::Tick( const FGeometry& AllottedGeometry, const double InCurrentTime, const float InDeltaTime )
 {
-	CachedDesignerGeometry = AllottedGeometry;
+	// Tick the parent first to update CachedGeometry
+	SDesignSurface::Tick(AllottedGeometry, InCurrentTime, InDeltaTime);
 
 	const bool bForceUpdate = false;
 	UpdatePreviewWidget(bForceUpdate);
@@ -1398,8 +1425,6 @@ void SDesignerView::Tick(const FGeometry& AllottedGeometry, const double InCurre
 		TopRuler->SetCursor(TOptional<FVector2D>());
 		SideRuler->SetCursor(TOptional<FVector2D>());
 	}
-
-	SDesignSurface::Tick(AllottedGeometry, InCurrentTime, InDeltaTime);
 }
 
 FReply SDesignerView::OnDragDetected(const FGeometry& MyGeometry, const FPointerEvent& MouseEvent)
@@ -1430,7 +1455,7 @@ void SDesignerView::OnDragEnter(const FGeometry& MyGeometry, const FDragDropEven
 {
 	SDesignSurface::OnDragEnter(MyGeometry, DragDropEvent);
 
-	BlueprintEditor.Pin()->ClearHoveredWidget();
+	ClearHoveredWidget();
 
 	//@TODO UMG Drop Feedback
 }
@@ -1439,7 +1464,7 @@ void SDesignerView::OnDragLeave(const FDragDropEvent& DragDropEvent)
 {
 	SDesignSurface::OnDragLeave(DragDropEvent);
 
-	BlueprintEditor.Pin()->ClearHoveredWidget();
+	ClearHoveredWidget();
 
 	TSharedPtr<FDecoratedDragDropOp> DecoratedDragDropOp = DragDropEvent.GetOperationAs<FDecoratedDragDropOp>();
 	if ( DecoratedDragDropOp.IsValid() )
@@ -1526,7 +1551,7 @@ UWidget* SDesignerView::ProcessDropAndAddWidget(const FGeometry& MyGeometry, con
 	TSharedPtr<FWidgetTemplateDragDropOp> TemplateDragDropOp = DragDropEvent.GetOperationAs<FWidgetTemplateDragDropOp>();
 	if ( TemplateDragDropOp.IsValid() )
 	{
-		BlueprintEditor.Pin()->SetHoveredWidget(HitResult.Widget);
+		SetHoveredWidget(HitResult.Widget);
 
 		TemplateDragDropOp->SetCursorOverride(TOptional<EMouseCursor::Type>());
 
@@ -1660,7 +1685,7 @@ UWidget* SDesignerView::ProcessDropAndAddWidget(const FGeometry& MyGeometry, con
 		}
 
 		FWidgetReference TargetReference = bIsPreview ? BlueprintEditor.Pin()->GetReferenceFromPreview(Target) : BlueprintEditor.Pin()->GetReferenceFromTemplate(Target);
-		BlueprintEditor.Pin()->SetHoveredWidget(TargetReference);
+		SetHoveredWidget(TargetReference);
 
 		// If the widget being hovered over is a panel, attempt to place it into that panel.
 		if ( Target && Target->IsA(UPanelWidget::StaticClass()) )
@@ -1927,7 +1952,7 @@ void SDesignerView::HandleOnCommonResolutionSelected(int32 Width, int32 Height,
 		MarkDesignModifed(/*bRequiresRecompile*/ false);
 	}
 
-	ResolutionTextFade.Play();
+	ResolutionTextFade.Play( this->AsShared() );
 }
 
 bool SDesignerView::HandleIsCommonResolutionSelected(int32 Width, int32 Height) const
diff --git a/Engine/Source/Editor/UMGEditor/Private/Designer/SDesignerView.h b/Engine/Source/Editor/UMGEditor/Private/Designer/SDesignerView.h
--- a/Engine/Source/Editor/UMGEditor/Private/Designer/SDesignerView.h
+++ b/Engine/Source/Editor/UMGEditor/Private/Designer/SDesignerView.h
@@ -48,7 +48,7 @@ public:
 
 	virtual FReply OnKeyDown(const FGeometry& MyGeometry, const FKeyEvent& InKeyEvent) override;
 
-	virtual void Tick(const FGeometry& AllottedGeometry, const double InCurrentTime, const float InDeltaTime) override;
+	virtual void Tick( const FGeometry& AllottedGeometry, const double InCurrentTime, const float InDeltaTime ) override;
 
 	virtual FReply OnDragDetected(const FGeometry& MyGeometry, const FPointerEvent& MouseEvent) override;
 	virtual void OnDragEnter(const FGeometry& MyGeometry, const FDragDropEvent& DragDropEvent) override;
@@ -108,6 +108,12 @@ private:
 
 	void OnEditorSelectionChanged();
 
+	/** Establishes the new hovered widget and restarts the outline fade */
+	void SetHoveredWidget(const FWidgetReference& InHoveredWidget);
+
+	/** Clears the hovered widget reference */
+	void ClearHoveredWidget();
+
 	/** Gets the blueprint being edited by the designer */
 	UWidgetBlueprint* GetBlueprint() const;
 
@@ -246,6 +252,12 @@ private:
 	/** Curve to handle fading of the resolution */
 	FCurveSequence ResolutionTextFade;
 
+	/** Curve to handle the fade-in of the border around the hovered widget */
+	FCurveSequence HoveredWidgetOutlineFade;
+
+	/** The widget currently being hovered over */
+	FWidgetReference HoveredWidget;
+
 	/**  */
 	FWeakWidgetPath SelectedWidgetPath;
 
@@ -263,7 +275,4 @@ private:
 
 	/**  */
 	TMap<TSharedRef<SWidget>, FArrangedWidget> CachedWidgetGeometry;
-
-	/**  */
-	FGeometry CachedDesignerGeometry;
 };
diff --git a/Engine/Source/Editor/UMGEditor/Private/Designer/SDisappearingBar.cpp b/Engine/Source/Editor/UMGEditor/Private/Designer/SDisappearingBar.cpp
--- a/Engine/Source/Editor/UMGEditor/Private/Designer/SDisappearingBar.cpp
+++ b/Engine/Source/Editor/UMGEditor/Private/Designer/SDisappearingBar.cpp
@@ -12,10 +12,6 @@
 void SDisappearingBar::Construct(const FArguments& InArgs)
 {
 	FadeCurve = FCurveSequence(0, 0.25f);
-	//FadeCurve.Play();
-
-	bDisappear = false;
-
 	ColorAndOpacity = TAttribute<FLinearColor>::Create(TAttribute<FLinearColor>::FGetter::CreateSP(this, &SDisappearingBar::GetFadeColorAndOpacity));
 
 	ChildSlot
@@ -24,25 +20,35 @@ void SDisappearingBar::Construct(const FArguments& InArgs)
 	];
 }
 
-void SDisappearingBar::Tick(const FGeometry& AllottedGeometry, const double InCurrentTime, const float InDeltaTime)
+void SDisappearingBar::OnDragEnter( const FGeometry& MyGeometry, const FDragDropEvent& DragDropEvent )
 {
-	if ( AllottedGeometry.IsUnderLocation(FSlateApplication::Get().GetCursorPos()) )
+	if ( !FadeCurve.IsAtEnd() )
 	{
-		if ( !bDisappear )
+		// Play the curve forward to fade the bar out
+		if ( FadeCurve.IsInReverse() && FadeCurve.IsPlaying() )
 		{
-			FadeCurve.Play();
+			FadeCurve.Reverse();
+		}
+		else if ( !FadeCurve.IsPlaying() )
+		{
+			FadeCurve.Play( this->AsShared() );
 		}
-
-		bDisappear = true;
 	}
-	else
+}
+
+void SDisappearingBar::OnDragLeave( const FDragDropEvent& DragDropEvent )
+{
+	if ( !FadeCurve.IsAtStart() )
 	{
-		if ( bDisappear )
+		// Play the curve in reverse to fade the bar back in
+		if ( !FadeCurve.IsInReverse() && FadeCurve.IsPlaying() )
 		{
-			FadeCurve.PlayReverse();
+			FadeCurve.Reverse();
+		}
+		else if ( !FadeCurve.IsPlaying() )
+		{
+			FadeCurve.PlayReverse( this->AsShared() );
 		}
-
-		bDisappear = false;
 	}
 }
 
diff --git a/Engine/Source/Editor/UMGEditor/Private/Designer/SDisappearingBar.h b/Engine/Source/Editor/UMGEditor/Private/Designer/SDisappearingBar.h
--- a/Engine/Source/Editor/UMGEditor/Private/Designer/SDisappearingBar.h
+++ b/Engine/Source/Editor/UMGEditor/Private/Designer/SDisappearingBar.h
@@ -23,11 +23,11 @@ public:
 
 	void Construct(const FArguments& InArgs);
 	
-	void Tick(const FGeometry& AllottedGeometry, const double InCurrentTime, const float InDeltaTime) override;
+	virtual void OnDragEnter( const FGeometry& MyGeometry, const FDragDropEvent& DragDropEvent ) override;
+	virtual void OnDragLeave( const FDragDropEvent& DragDropEvent ) override;
 
 private:
 	FLinearColor GetFadeColorAndOpacity() const;
 
-	bool bDisappear;
 	FCurveSequence FadeCurve;
 };
diff --git a/Engine/Source/Editor/UMGEditor/Private/Extensions/CanvasSlotExtension.cpp b/Engine/Source/Editor/UMGEditor/Private/Extensions/CanvasSlotExtension.cpp
--- a/Engine/Source/Editor/UMGEditor/Private/Extensions/CanvasSlotExtension.cpp
+++ b/Engine/Source/Editor/UMGEditor/Private/Extensions/CanvasSlotExtension.cpp
@@ -308,11 +308,6 @@ void FCanvasSlotExtension::GetCollisionSegmentsFromGeometry(FGeometry ArrangedGe
 	Segments[7] = ArrangedGeometry.Position + ArrangedGeometry.Size;
 }
 
-void FCanvasSlotExtension::Tick(const FGeometry& AllottedGeometry, const double InCurrentTime, const float InDeltaTime)
-{
-
-}
-
 void FCanvasSlotExtension::Paint(const TSet< FWidgetReference >& Selection, const FGeometry& AllottedGeometry, const FSlateRect& MyClippingRect, FSlateWindowElementList& OutDrawElements, int32 LayerId) const
 {
 	//PaintCollisionLines(Selection, AllottedGeometry, MyClippingRect, OutDrawElements, LayerId);
diff --git a/Engine/Source/Editor/UMGEditor/Private/Extensions/CanvasSlotExtension.h b/Engine/Source/Editor/UMGEditor/Private/Extensions/CanvasSlotExtension.h
--- a/Engine/Source/Editor/UMGEditor/Private/Extensions/CanvasSlotExtension.h
+++ b/Engine/Source/Editor/UMGEditor/Private/Extensions/CanvasSlotExtension.h
@@ -40,7 +40,6 @@ public:
 	virtual bool CanExtendSelection(const TArray< FWidgetReference >& Selection) const override;
 	
 	virtual void ExtendSelection(const TArray< FWidgetReference >& Selection, TArray< TSharedRef<FDesignerSurfaceElement> >& SurfaceElements) override;
-	virtual void Tick(const FGeometry& AllottedGeometry, const double InCurrentTime, const float InDeltaTime);
 	virtual void Paint(const TSet< FWidgetReference >& Selection, const FGeometry& AllottedGeometry, const FSlateRect& MyClippingRect, FSlateWindowElementList& OutDrawElements, int32 LayerId) const override;
 
 private:
diff --git a/Engine/Source/Editor/UMGEditor/Private/Hierarchy/SHierarchyView.cpp b/Engine/Source/Editor/UMGEditor/Private/Hierarchy/SHierarchyView.cpp
--- a/Engine/Source/Editor/UMGEditor/Private/Hierarchy/SHierarchyView.cpp
+++ b/Engine/Source/Editor/UMGEditor/Private/Hierarchy/SHierarchyView.cpp
@@ -91,10 +91,8 @@ SHierarchyView::~SHierarchyView()
 	GEditor->OnObjectsReplaced().RemoveAll(this);
 }
 
-void SHierarchyView::Tick(const FGeometry& AllottedGeometry, const double InCurrentTime, const float InDeltaTime)
+void SHierarchyView::Tick( const FGeometry& AllottedGeometry, const double InCurrentTime, const float InDeltaTime )
 {
-	SCompoundWidget::Tick(AllottedGeometry, InCurrentTime, InDeltaTime);
-
 	if ( bRebuildTreeRequested || bRefreshRequested )
 	{
 		if ( bRebuildTreeRequested )
diff --git a/Engine/Source/Editor/UMGEditor/Private/Palette/SPaletteView.h b/Engine/Source/Editor/UMGEditor/Private/Palette/SPaletteView.h
--- a/Engine/Source/Editor/UMGEditor/Private/Palette/SPaletteView.h
+++ b/Engine/Source/Editor/UMGEditor/Private/Palette/SPaletteView.h
@@ -38,7 +38,7 @@ public:
 	void Construct(const FArguments& InArgs, TSharedPtr<FBlueprintEditor> InBlueprintEditor);
 	virtual ~SPaletteView();
 	
-	virtual void Tick(const FGeometry& AllottedGeometry, const double InCurrentTime, const float InDeltaTime);
+	virtual void Tick( const FGeometry& AllottedGeometry, const double InCurrentTime, const float InDeltaTime ) override;
 
 	/** Gets the text currently displayed in the search box. */
 	FText GetSearchText() const;
diff --git a/Engine/Source/Editor/UMGEditor/Private/WidgetBlueprintEditor.cpp b/Engine/Source/Editor/UMGEditor/Private/WidgetBlueprintEditor.cpp
--- a/Engine/Source/Editor/UMGEditor/Private/WidgetBlueprintEditor.cpp
+++ b/Engine/Source/Editor/UMGEditor/Private/WidgetBlueprintEditor.cpp
@@ -28,7 +28,6 @@
 FWidgetBlueprintEditor::FWidgetBlueprintEditor()
 	: PreviewScene(FPreviewScene::ConstructionValues().AllowAudioPlayback(true).ShouldSimulatePhysics(true))
 	, PreviewBlueprint(nullptr)
-	, HoverTime(0)
 	, bIsSimulateEnabled(false)
 	, bIsRealTime(true)
 {
@@ -350,8 +349,6 @@ void FWidgetBlueprintEditor::Tick(float DeltaTime)
 {
 	FBlueprintEditor::Tick(DeltaTime);
 
-	HoverTime += DeltaTime;
-
 	// Tick the preview scene world.
 	if ( !GIntraFrameDebuggingGameThread )
 	{
@@ -686,29 +683,13 @@ FGraphAppearanceInfo FWidgetBlueprintEditor::GetGraphAppearance() const
 
 void FWidgetBlueprintEditor::ClearHoveredWidget()
 {
-	HoveredWidget = FWidgetReference();
-	HoverTime = 0;
+	OnHoveredWidgetCleared.Broadcast();
 }
 
 void FWidgetBlueprintEditor::SetHoveredWidget(FWidgetReference& InHoveredWidget)
 {
-	if ( !( InHoveredWidget == HoveredWidget ) )
-	{
-		HoveredWidget = InHoveredWidget;
-		HoverTime = 0;
-	}
-}
-
-FWidgetReference FWidgetBlueprintEditor::GetHoveredWidget() const
-{
-	return HoveredWidget;
+	OnHoveredWidgetSet.Broadcast(InHoveredWidget);
 }
-
-float FWidgetBlueprintEditor::GetHoveredWidgetTime() const
-{
-	return HoverTime;
-}
-
 void FWidgetBlueprintEditor::AddPostDesignerLayoutAction(TFunction<void()> Action)
 {
 	QueuedDesignerActions.Add(Action);
diff --git a/Engine/Source/Editor/UMGEditor/Private/WidgetBlueprintEditor.h b/Engine/Source/Editor/UMGEditor/Private/WidgetBlueprintEditor.h
--- a/Engine/Source/Editor/UMGEditor/Private/WidgetBlueprintEditor.h
+++ b/Engine/Source/Editor/UMGEditor/Private/WidgetBlueprintEditor.h
@@ -17,6 +17,9 @@ class UUserWidget;
 class FWidgetBlueprintEditor : public FBlueprintEditor
 {
 public:
+	DECLARE_MULTICAST_DELEGATE_OneParam(FOnHoveredWidgetSet, const FWidgetReference&)
+	DECLARE_MULTICAST_DELEGATE(FOnHoveredWidgetCleared);
+
 	DECLARE_MULTICAST_DELEGATE(FOnSelectedWidgetsChanging)
 	DECLARE_MULTICAST_DELEGATE(FOnSelectedWidgetsChanged)
 
@@ -111,15 +114,19 @@ public:
 
 	void ClearHoveredWidget();
 
-	FWidgetReference GetHoveredWidget() const;
-
-	float GetHoveredWidgetTime() const;
+	//FWidgetReference GetHoveredWidget() const;
 
 	void AddPostDesignerLayoutAction(TFunction<void()> Action);
 
 	TArray< TFunction<void()> >& GetQueuedDesignerActions();
 
 public:
+	/** Fires whenever a new widget is being hovered over */
+	FOnHoveredWidgetSet OnHoveredWidgetSet;
+
+	/** Fires when there is no longer any widget being hovered over */
+	FOnHoveredWidgetCleared OnHoveredWidgetCleared;	
+
 	/** Fires whenever the selected set of widgets changing */
 	FOnSelectedWidgetsChanged OnSelectedWidgetsChanging;
 
@@ -201,12 +208,6 @@ private:
 	/** The widget references out in the ether that may need to be updated after being issued. */
 	TArray< TWeakPtr<FWidgetHandle> > WidgetHandlePool;
 
-	/** The wall clock time the user has been hovering over a single widget */
-	float HoverTime;
-
-	/** The current widget being hovered */
-	FWidgetReference HoveredWidget;
-
 	/** The preview becomes invalid and needs to be rebuilt on the next tick. */
 	bool bPreviewInvalidated;
 
diff --git a/Engine/Source/Editor/UMGEditor/Public/WidgetReference.h b/Engine/Source/Editor/UMGEditor/Public/WidgetReference.h
--- a/Engine/Source/Editor/UMGEditor/Public/WidgetReference.h
+++ b/Engine/Source/Editor/UMGEditor/Public/WidgetReference.h
@@ -50,7 +50,7 @@ public:
 		return TSharedPtr<SWidget>();
 	}
 
-	/** Checks if widget reference is the same as nother widget reference, based on the template pointers. */
+	/** Checks if widget reference is the same as another widget reference, based on the template pointers. */
 	bool operator==( const FWidgetReference& Other ) const
 	{
 		if ( TemplateHandle.IsValid() && Other.TemplateHandle.IsValid() )
@@ -61,6 +61,12 @@ public:
 		return false;
 	}
 
+	/** Checks if widget reference is the different from another widget reference, based on the template pointers. */
+	bool operator!=( const FWidgetReference& Other ) const
+	{
+		return !operator==( Other );
+	}
+
 private:
 	FWidgetReference(TSharedPtr<FWidgetBlueprintEditor> WidgetEditor, TSharedPtr< FWidgetHandle > TemplateHandle);
 
diff --git a/Engine/Source/Editor/UndoHistory/Private/Widgets/SUndoHistory.cpp b/Engine/Source/Editor/UndoHistory/Private/Widgets/SUndoHistory.cpp
--- a/Engine/Source/Editor/UndoHistory/Private/Widgets/SUndoHistory.cpp
+++ b/Engine/Source/Editor/UndoHistory/Private/Widgets/SUndoHistory.cpp
@@ -75,16 +75,8 @@ void SUndoHistory::Construct( const FArguments& InArgs )
 
 	ReloadUndoList();
 }
-END_SLATE_FUNCTION_BUILD_OPTIMIZATION
-
-
-/* SUndoHistory implementation
- *****************************************************************************/
-
 void SUndoHistory::Tick( const FGeometry& AllottedGeometry, const double InCurrentTime, const float InDeltaTime )
 {
-	SCompoundWidget::Tick(AllottedGeometry, InCurrentTime, InDeltaTime);
-
 	// reload the transaction list if necessary
 	if ((GEditor == nullptr) || (GEditor->Trans == nullptr) || (LastQueueLength != GEditor->Trans->GetQueueLength()))
 	{
diff --git a/Engine/Source/Editor/UnrealEd/Private/AssetThumbnail.cpp b/Engine/Source/Editor/UnrealEd/Private/AssetThumbnail.cpp
--- a/Engine/Source/Editor/UnrealEd/Private/AssetThumbnail.cpp
+++ b/Engine/Source/Editor/UnrealEd/Private/AssetThumbnail.cpp
@@ -220,8 +220,6 @@ public:
 
 	virtual void Tick( const FGeometry& AllottedGeometry, const double InCurrentTime, const float InDeltaTime ) override
 	{
-		SCompoundWidget::Tick(AllottedGeometry, InCurrentTime, InDeltaTime);
-
 		if ( WidthLastFrame != AllottedGeometry.Size.X )
 		{
 			WidthLastFrame = AllottedGeometry.Size.X;
@@ -403,7 +401,7 @@ private:
 		if ( !bHasRenderedThumbnail && AssetData == AssetThumbnail->GetAssetData() && ShouldRender() )
 		{
 			OnRenderedThumbnailChanged( true );
-			ViewportFadeAnimation.Play();
+			ViewportFadeAnimation.Play( this->AsShared() );
 		}
 	}
 
@@ -829,6 +827,16 @@ void FAssetThumbnailPool::ReleaseResources()
 	}
 }
 
+TStatId FAssetThumbnailPool::GetStatId() const
+{
+	RETURN_QUICK_DECLARE_CYCLE_STAT( FAssetThumbnailPool, STATGROUP_Tickables );
+}
+
+bool FAssetThumbnailPool::IsTickable() const
+{
+	return RecentlyLoadedAssets.Num() > 0 || ThumbnailsToRenderStack.Num() > 0 || RealTimeThumbnails.Num() > 0;
+}
+
 void FAssetThumbnailPool::Tick( float DeltaTime )
 {
 	// If there were any assets loaded since last frame that we are currently displaying thumbnails for, push them on the render stack now.
diff --git a/Engine/Source/Editor/UnrealEd/Private/CreditsScreen.cpp b/Engine/Source/Editor/UnrealEd/Private/CreditsScreen.cpp
--- a/Engine/Source/Editor/UnrealEd/Private/CreditsScreen.cpp
+++ b/Engine/Source/Editor/UnrealEd/Private/CreditsScreen.cpp
@@ -14,6 +14,8 @@ void SCreditsScreen::Construct(const FArguments& InArgs)
 	ScrollPixelsPerSecond = 50.0f;
 	bIsPlaying = true;
 
+	ActiveTickHandle = RegisterActiveTick( 0.f, FWidgetActiveTickDelegate::CreateSP( this, &SCreditsScreen::RollCredits ) );
+
 	const FString Version = GEngineVersion.ToString(FEngineBuildSettings::IsPerforceBuild() ? EVersionComponent::Branch : EVersionComponent::Patch);
 
 	FString CreditsText;
@@ -67,19 +69,33 @@ void SCreditsScreen::Construct(const FArguments& InArgs)
 	];
 }
 
-void SCreditsScreen::Tick(const FGeometry& AllottedGeometry, const double InCurrentTime, const float InDeltaTime)
+EActiveTickReturnType SCreditsScreen::RollCredits( double InCurrentTime, float InDeltaTime )
 {
-	if ( bIsPlaying )
-	{
-		float NewPixelOffset = ( ScrollPixelsPerSecond * InDeltaTime );
-		ScrollBox->SetScrollOffset(ScrollBox->GetScrollOffset() + NewPixelOffset);
-		PreviousScrollPosition = ScrollBox->GetScrollOffset();
-	}
+	float NewPixelOffset = ( ScrollPixelsPerSecond * InDeltaTime );
+	ScrollBox->SetScrollOffset( ScrollBox->GetScrollOffset() + NewPixelOffset );
+	PreviousScrollPosition = ScrollBox->GetScrollOffset();
+
+	return EActiveTickReturnType::KeepTicking;
 }
 
 FReply SCreditsScreen::HandleTogglePlayPause()
 {
-	bIsPlaying = !bIsPlaying;
+	if ( bIsPlaying )
+	{
+		bIsPlaying = false;
+		if ( ActiveTickHandle.IsValid() )
+		{
+			UnRegisterActiveTick( ActiveTickHandle.Pin().ToSharedRef() );
+		}
+	}
+	else
+	{
+		bIsPlaying = true;
+		if ( !ActiveTickHandle.IsValid() )
+		{
+			ActiveTickHandle = RegisterActiveTick( 0.f, FWidgetActiveTickDelegate::CreateSP( this, &SCreditsScreen::RollCredits ) );
+		}
+	}
 
 	return FReply::Handled();
 }
@@ -90,6 +106,10 @@ void SCreditsScreen::HandleUserScrolled(float ScrollOffset)
 	if ( bIsPlaying && ScrollOffset < PreviousScrollPosition )
 	{
 		bIsPlaying = false;
+		if ( ActiveTickHandle.IsValid() )
+		{
+			UnRegisterActiveTick( ActiveTickHandle.Pin().ToSharedRef() );
+		}
 	}
 
 	PreviousScrollPosition = ScrollOffset;
diff --git a/Engine/Source/Editor/UnrealEd/Private/Dialogs/SDeleteAssetsDialog.cpp b/Engine/Source/Editor/UnrealEd/Private/Dialogs/SDeleteAssetsDialog.cpp
--- a/Engine/Source/Editor/UnrealEd/Private/Dialogs/SDeleteAssetsDialog.cpp
+++ b/Engine/Source/Editor/UnrealEd/Private/Dialogs/SDeleteAssetsDialog.cpp
@@ -101,6 +101,10 @@ SDeleteAssetsDialog::~SDeleteAssetsDialog()
 
 void SDeleteAssetsDialog::Construct( const FArguments& InArgs, TSharedRef<FAssetDeleteModel> InDeleteModel )
 {
+	//bIsActiveTickRegistered = false;
+	bIsActiveTickRegistered = true;
+	RegisterActiveTick( 0.f, FWidgetActiveTickDelegate::CreateSP( this, &SDeleteAssetsDialog::TickDeleteModel ) );
+
 	DeleteModel = InDeleteModel;
 
 	// Save off the attributes
@@ -507,10 +511,17 @@ FReply SDeleteAssetsDialog::OnKeyDown( const FGeometry& MyGeometry, const FKeyEv
 	return FReply::Unhandled();
 }
 
-void SDeleteAssetsDialog::Tick( const FGeometry& AllottedGeometry, const double InCurrentTime, const float InDeltaTime )
+EActiveTickReturnType SDeleteAssetsDialog::TickDeleteModel( double InCurrentTime, float InDeltaTime )
 {
-	AssetThumbnailPool->Tick( InDeltaTime );
 	DeleteModel->Tick( InDeltaTime );
+
+	if ( DeleteModel->GetState() == FAssetDeleteModel::EState::Finished )
+	{
+		bIsActiveTickRegistered = false;
+		return EActiveTickReturnType::StopTicking;
+	}
+
+	return EActiveTickReturnType::KeepTicking;
 }
 
 void SDeleteAssetsDialog::HandleDeleteModelStateChanged(FAssetDeleteModel::EState NewState)
@@ -825,7 +836,13 @@ void SDeleteAssetsDialog::ExecuteDeleteReferencers()
 	for ( const FAssetData& SelectedAsset : SelectedAssets )
 	{
 		UObject* ObjectToDelete = SelectedAsset.GetAsset();
+		
 		DeleteModel->AddObjectToDelete( ObjectToDelete );
+		if ( !bIsActiveTickRegistered )
+		{
+			bIsActiveTickRegistered = true;
+			RegisterActiveTick( 0.f, FWidgetActiveTickDelegate::CreateSP( this, &SDeleteAssetsDialog::TickDeleteModel ) );
+		}
 	}
 }
 
diff --git a/Engine/Source/Editor/UnrealEd/Private/EditorViewportClient.cpp b/Engine/Source/Editor/UnrealEd/Private/EditorViewportClient.cpp
--- a/Engine/Source/Editor/UnrealEd/Private/EditorViewportClient.cpp
+++ b/Engine/Source/Editor/UnrealEd/Private/EditorViewportClient.cpp
@@ -90,7 +90,7 @@ public:
 
 
 FViewportCameraTransform::FViewportCameraTransform()
-	: TransitionCurve( new FCurveSequence( 0.0f, FocusConstants::TransitionTime, ECurveEaseFunction::CubicOut ) )
+	: TransitionStartTime(0)
 	, ViewLocation( FVector::ZeroVector )
 	, ViewRotation( FRotator::ZeroRotator )
 	, DesiredLocation( FVector::ZeroVector )
@@ -110,14 +110,14 @@ void FViewportCameraTransform::TransitionToLocation( const FVector& InDesiredLoc
 	if( bInstant )
 	{
 		SetLocation( InDesiredLocation );
-		TransitionCurve->JumpToEnd();
+		TransitionStartTime = FSlateApplication::Get().GetCurrentTime() - FocusConstants::TransitionTime;
 	}
 	else
 	{
 		DesiredLocation = InDesiredLocation;
 		StartLocation = ViewLocation;
 
-		TransitionCurve->Play();
+		TransitionStartTime = FSlateApplication::Get().GetCurrentTime();
 	}
 }
 
@@ -125,10 +125,12 @@ void FViewportCameraTransform::TransitionToLocation( const FVector& InDesiredLoc
 bool FViewportCameraTransform::UpdateTransition()
 {
 	bool bIsAnimating = false;
-	if( TransitionCurve->IsPlaying() || ViewLocation != DesiredLocation )
+	double TransitionProgress = FMath::Clamp( (FSlateApplication::Get().GetCurrentTime() - TransitionStartTime) / FocusConstants::TransitionTime, 0.0, 1.0);
+	if( TransitionProgress < 1.0 || ViewLocation != DesiredLocation )
 	{
-		float LerpWeight = TransitionCurve->GetLerp();
-		
+		const float Offset = (float)TransitionProgress - 1.0f;
+		float LerpWeight = Offset * Offset * Offset + 1.0f;
+
 		if( LerpWeight == 1.0f )
 		{
 			// Failsafe for the value not being exact on lerps
