UMG - Making the navigation dropdowns less bulky in the nav customization.

diff --git a/Engine/Source/Editor/UMGEditor/Private/Customizations/WidgetNavigationCustomization.cpp b/Engine/Source/Editor/UMGEditor/Private/Customizations/WidgetNavigationCustomization.cpp
--- a/Engine/Source/Editor/UMGEditor/Private/Customizations/WidgetNavigationCustomization.cpp
+++ b/Engine/Source/Editor/UMGEditor/Private/Customizations/WidgetNavigationCustomization.cpp
@@ -133,7 +133,7 @@ void FWidgetNavigationCustomization::MakeNavRow(TWeakPtr<IPropertyHandle> Proper
 				SNew(STextBlock)
 				.Text(this, &FWidgetNavigationCustomization::HandleNavigationText, PropertyHandle, Nav)
 			]
-			//.ContentPadding(FMargin(4.0f, 2.0f))
+			.ContentPadding(FMargin(2.0f, 1.0f))
 			.MenuContent()
 			[
 				MakeNavMenu(PropertyHandle, Nav)
