UMG - The hierarchy now appears highlighted when hovering the corresponding widget in the UMG designer

diff --git a/Engine/Source/Editor/UMGEditor/Private/Designer/SDesignerView.cpp b/Engine/Source/Editor/UMGEditor/Private/Designer/SDesignerView.cpp
--- a/Engine/Source/Editor/UMGEditor/Private/Designer/SDesignerView.cpp
+++ b/Engine/Source/Editor/UMGEditor/Private/Designer/SDesignerView.cpp
@@ -295,9 +295,10 @@ void SDesignerView::Construct(const FArguments& InArgs, TSharedPtr<FWidgetBluepr
 		]
 	);
 
-	BlueprintEditor.Pin()->OnSelectedWidgetsChanged.AddRaw(this, &SDesignerView::OnEditorSelectionChanged);
-	BlueprintEditor.Pin()->OnHoveredWidgetSet.AddRaw(this, &SDesignerView::SetHoveredWidget);
-	BlueprintEditor.Pin()->OnHoveredWidgetCleared.AddRaw(this, &SDesignerView::ClearHoveredWidget);
+	auto PinnedBlueprintEditor = BlueprintEditor.Pin();
+	PinnedBlueprintEditor->OnSelectedWidgetsChanged.AddRaw(this, &SDesignerView::OnEditorSelectionChanged);
+	PinnedBlueprintEditor->OnHoveredWidgetSet.AddRaw(this, &SDesignerView::OnHoveredWidgetSet);
+	PinnedBlueprintEditor->OnHoveredWidgetCleared.AddRaw(this, &SDesignerView::OnHoveredWidgetCleared);
 
 	ZoomToFit(/*bInstantZoom*/ true);
 }
@@ -730,18 +731,14 @@ void SDesignerView::OnEditorSelectionChanged()
 	CreateExtensionWidgetsForSelection();
 }
 
-void SDesignerView::SetHoveredWidget(const FWidgetReference& InHoveredWidget)
+void SDesignerView::OnHoveredWidgetSet(const FWidgetReference& InHoveredWidget)
 {
-	if (InHoveredWidget != HoveredWidget)
-	{
-		HoveredWidget = InHoveredWidget;
-		HoveredWidgetOutlineFade.Play(this->AsShared());
-	}
+	HoveredWidgetOutlineFade.Play(this->AsShared());
 }
 
-void SDesignerView::ClearHoveredWidget()
+void SDesignerView::OnHoveredWidgetCleared()
 {
-	HoveredWidget = FWidgetReference();
+	HoveredWidgetOutlineFade.JumpToEnd();
 }
 
 FGeometry SDesignerView::GetDesignerGeometry() const
@@ -1179,15 +1176,17 @@ FReply SDesignerView::OnMouseMove(const FGeometry& MyGeometry, const FPointerEve
 		}
 	}
 	
+
 	// Update the hovered widget under the mouse
+	auto PinnedBlueprintEditor = BlueprintEditor.Pin();
 	FWidgetHitResult HitResult;
 	if ( FindWidgetUnderCursor(MyGeometry, MouseEvent, HitResult) )
 	{
-		SetHoveredWidget(HitResult.Widget);
+		PinnedBlueprintEditor->SetHoveredWidget(HitResult.Widget);
 	}
-	else if (HoveredWidget.IsValid())
+	else if (PinnedBlueprintEditor->GetHoveredWidget().IsValid())
 	{
-		ClearHoveredWidget();
+		PinnedBlueprintEditor->ClearHoveredWidget();
 	}
 
 	return FReply::Unhandled();
@@ -1197,14 +1196,14 @@ void SDesignerView::OnMouseEnter(const FGeometry& MyGeometry, const FPointerEven
 {
 	SDesignSurface::OnMouseEnter(MyGeometry, MouseEvent);
 
-	ClearHoveredWidget();
+	BlueprintEditor.Pin()->ClearHoveredWidget();
 }
 
 void SDesignerView::OnMouseLeave(const FPointerEvent& MouseEvent)
 {
 	SDesignSurface::OnMouseLeave(MouseEvent);
 
-	ClearHoveredWidget();
+	BlueprintEditor.Pin()->ClearHoveredWidget();
 }
 
 FReply SDesignerView::OnKeyDown(const FGeometry& MyGeometry, const FKeyEvent& InKeyEvent)
@@ -1297,12 +1296,12 @@ int32 SDesignerView::HandleEffectsPainting(const FOnPaintHandlerParams& PaintArg
 		}
 	}
 
+	const FWidgetReference& HoveredWidget = BlueprintEditor.Pin()->GetHoveredWidget();
 	TSharedPtr<SWidget> HoveredSlateWidget = HoveredWidget.GetPreviewSlate();
 
 	// Don't draw the hovered effect if it's also the selected widget
 	if ( HoveredSlateWidget.IsValid() && !SelectedWidgets.Contains(HoveredWidget) )
 	{
-		//const FLinearColor HoveredTint(0, 0.5, 1, FMath::Clamp(BlueprintEditor.Pin()->GetHoveredWidgetTime() / HoveredAnimationTime, 0.0f, 1.0f)); // Azure = 0x007FFF
 		const FLinearColor HoveredTint(0, 0.5, 1, HoveredWidgetOutlineFade.GetLerp()); // Azure = 0x007FFF
 
 		TSharedRef<SWidget> Widget = HoveredSlateWidget.ToSharedRef();
@@ -1455,7 +1454,7 @@ void SDesignerView::OnDragEnter(const FGeometry& MyGeometry, const FDragDropEven
 {
 	SDesignSurface::OnDragEnter(MyGeometry, DragDropEvent);
 
-	ClearHoveredWidget();
+	BlueprintEditor.Pin()->ClearHoveredWidget();
 
 	//@TODO UMG Drop Feedback
 }
@@ -1464,7 +1463,7 @@ void SDesignerView::OnDragLeave(const FDragDropEvent& DragDropEvent)
 {
 	SDesignSurface::OnDragLeave(DragDropEvent);
 
-	ClearHoveredWidget();
+	BlueprintEditor.Pin()->ClearHoveredWidget();
 
 	TSharedPtr<FDecoratedDragDropOp> DecoratedDragDropOp = DragDropEvent.GetOperationAs<FDecoratedDragDropOp>();
 	if ( DecoratedDragDropOp.IsValid() )
@@ -1551,7 +1550,7 @@ UWidget* SDesignerView::ProcessDropAndAddWidget(const FGeometry& MyGeometry, con
 	TSharedPtr<FWidgetTemplateDragDropOp> TemplateDragDropOp = DragDropEvent.GetOperationAs<FWidgetTemplateDragDropOp>();
 	if ( TemplateDragDropOp.IsValid() )
 	{
-		SetHoveredWidget(HitResult.Widget);
+		BlueprintEditor.Pin()->SetHoveredWidget(HitResult.Widget);
 
 		TemplateDragDropOp->SetCursorOverride(TOptional<EMouseCursor::Type>());
 
@@ -1685,7 +1684,7 @@ UWidget* SDesignerView::ProcessDropAndAddWidget(const FGeometry& MyGeometry, con
 		}
 
 		FWidgetReference TargetReference = bIsPreview ? BlueprintEditor.Pin()->GetReferenceFromPreview(Target) : BlueprintEditor.Pin()->GetReferenceFromTemplate(Target);
-		SetHoveredWidget(TargetReference);
+		BlueprintEditor.Pin()->SetHoveredWidget(TargetReference);
 
 		// If the widget being hovered over is a panel, attempt to place it into that panel.
 		if ( Target && Target->IsA(UPanelWidget::StaticClass()) )
diff --git a/Engine/Source/Editor/UMGEditor/Private/Designer/SDesignerView.h b/Engine/Source/Editor/UMGEditor/Private/Designer/SDesignerView.h
--- a/Engine/Source/Editor/UMGEditor/Private/Designer/SDesignerView.h
+++ b/Engine/Source/Editor/UMGEditor/Private/Designer/SDesignerView.h
@@ -108,11 +108,11 @@ private:
 
 	void OnEditorSelectionChanged();
 
-	/** Establishes the new hovered widget and restarts the outline fade */
-	void SetHoveredWidget(const FWidgetReference& InHoveredWidget);
+	/** Called when a new widget is being hovered */
+	void OnHoveredWidgetSet(const FWidgetReference& InHoveredWidget);
 
-	/** Clears the hovered widget reference */
-	void ClearHoveredWidget();
+	/** Called when a widget is no longer being hovered */
+	void OnHoveredWidgetCleared();
 
 	/** Gets the blueprint being edited by the designer */
 	UWidgetBlueprint* GetBlueprint() const;
@@ -255,9 +255,6 @@ private:
 	/** Curve to handle the fade-in of the border around the hovered widget */
 	FCurveSequence HoveredWidgetOutlineFade;
 
-	/** The widget currently being hovered over */
-	FWidgetReference HoveredWidget;
-
 	/**  */
 	FWeakWidgetPath SelectedWidgetPath;
 
diff --git a/Engine/Source/Editor/UMGEditor/Private/Hierarchy/SHierarchyViewItem.cpp b/Engine/Source/Editor/UMGEditor/Private/Hierarchy/SHierarchyViewItem.cpp
--- a/Engine/Source/Editor/UMGEditor/Private/Hierarchy/SHierarchyViewItem.cpp
+++ b/Engine/Source/Editor/UMGEditor/Private/Hierarchy/SHierarchyViewItem.cpp
@@ -782,6 +782,11 @@ void FHierarchyWidget::OnMouseLeave()
 	BlueprintEditor.Pin()->ClearHoveredWidget();
 }
 
+bool FHierarchyWidget::IsHovered() const
+{
+	return BlueprintEditor.Pin()->GetHoveredWidget() == Item;
+}
+
 void FHierarchyWidget::UpdateSelection()
 {
 	const TSet<FWidgetReference>& SelectedWidgets = BlueprintEditor.Pin()->GetSelectedWidgets();
@@ -895,6 +900,11 @@ FText SHierarchyViewItem::GetItemText() const
 	return Model->GetText();
 }
 
+bool SHierarchyViewItem::IsHovered() const
+{
+	return bIsHovered || Model->IsHovered();
+}
+
 FReply SHierarchyViewItem::OnDragOver(const FGeometry& MyGeometry, const FDragDropEvent& DragDropEvent)
 {
 	return Model->OnDragOver(MyGeometry, DragDropEvent);
diff --git a/Engine/Source/Editor/UMGEditor/Private/Hierarchy/SHierarchyViewItem.h b/Engine/Source/Editor/UMGEditor/Private/Hierarchy/SHierarchyViewItem.h
--- a/Engine/Source/Editor/UMGEditor/Private/Hierarchy/SHierarchyViewItem.h
+++ b/Engine/Source/Editor/UMGEditor/Private/Hierarchy/SHierarchyViewItem.h
@@ -53,7 +53,8 @@ public:
 	void RefreshSelection();
 	bool ContainsSelection();
 	bool IsSelected() const;
-
+	
+	virtual bool IsHovered() const { return false; }
 	virtual bool IsVisible() const { return true; }
 	virtual bool CanControlVisibility() const { return false; }
 	virtual void SetIsVisible(bool IsVisible) { }
@@ -173,6 +174,8 @@ public:
 	virtual void OnMouseEnter() override;
 	virtual void OnMouseLeave() override;
 
+	virtual bool IsHovered() const override;
+
 	virtual bool IsVisible() const override
 	{
 		if ( UWidget* TemplateWidget = Item.GetTemplate() )
@@ -234,6 +237,7 @@ public:
 	void Construct(const FArguments& InArgs, const TSharedRef< STableViewBase >& InOwnerTableView, TSharedPtr<FHierarchyModel> InModel);
 
 	// Begin SWidget
+	virtual bool IsHovered() const;
 	void OnMouseEnter(const FGeometry& MyGeometry, const FPointerEvent& MouseEvent);
 	void OnMouseLeave(const FPointerEvent& MouseEvent);
 	virtual FReply OnDragOver(const FGeometry& MyGeometry, const FDragDropEvent& DragDropEvent) override;
diff --git a/Engine/Source/Editor/UMGEditor/Private/WidgetBlueprintEditor.cpp b/Engine/Source/Editor/UMGEditor/Private/WidgetBlueprintEditor.cpp
--- a/Engine/Source/Editor/UMGEditor/Private/WidgetBlueprintEditor.cpp
+++ b/Engine/Source/Editor/UMGEditor/Private/WidgetBlueprintEditor.cpp
@@ -683,13 +683,24 @@ FGraphAppearanceInfo FWidgetBlueprintEditor::GetGraphAppearance() const
 
 void FWidgetBlueprintEditor::ClearHoveredWidget()
 {
+	HoveredWidget = FWidgetReference();
 	OnHoveredWidgetCleared.Broadcast();
 }
 
 void FWidgetBlueprintEditor::SetHoveredWidget(FWidgetReference& InHoveredWidget)
 {
-	OnHoveredWidgetSet.Broadcast(InHoveredWidget);
+	if (InHoveredWidget != HoveredWidget)
+	{
+		HoveredWidget = InHoveredWidget;
+		OnHoveredWidgetSet.Broadcast(InHoveredWidget);
+	}
+}
+
+const FWidgetReference& FWidgetBlueprintEditor::GetHoveredWidget() const
+{
+	return HoveredWidget;
 }
+
 void FWidgetBlueprintEditor::AddPostDesignerLayoutAction(TFunction<void()> Action)
 {
 	QueuedDesignerActions.Add(Action);
diff --git a/Engine/Source/Editor/UMGEditor/Private/WidgetBlueprintEditor.h b/Engine/Source/Editor/UMGEditor/Private/WidgetBlueprintEditor.h
--- a/Engine/Source/Editor/UMGEditor/Private/WidgetBlueprintEditor.h
+++ b/Engine/Source/Editor/UMGEditor/Private/WidgetBlueprintEditor.h
@@ -114,7 +114,8 @@ public:
 
 	void ClearHoveredWidget();
 
-	//FWidgetReference GetHoveredWidget() const;
+	/** @return The widget that is currently being hovered over (either in the designer or hierarchy) */
+	const FWidgetReference& GetHoveredWidget() const;
 
 	void AddPostDesignerLayoutAction(TFunction<void()> Action);
 
@@ -208,6 +209,9 @@ private:
 	/** The widget references out in the ether that may need to be updated after being issued. */
 	TArray< TWeakPtr<FWidgetHandle> > WidgetHandlePool;
 
+	/** The widget currently being hovered over */
+	FWidgetReference HoveredWidget;
+
 	/** The preview becomes invalid and needs to be rebuilt on the next tick. */
 	bool bPreviewInvalidated;
 
