UMG - Adding a drop shadow to the text drawn in the canvas slot's anchor designer extension.

diff --git a/Engine/Source/Editor/UMGEditor/Private/Extensions/CanvasSlotExtension.cpp b/Engine/Source/Editor/UMGEditor/Private/Extensions/CanvasSlotExtension.cpp
--- a/Engine/Source/Editor/UMGEditor/Private/Extensions/CanvasSlotExtension.cpp
+++ b/Engine/Source/Editor/UMGEditor/Private/Extensions/CanvasSlotExtension.cpp
@@ -619,16 +619,28 @@ void FCanvasSlotExtension::PaintLineWithText(FVector2D Start, FVector2D End, FTe
 
 	FVector2D TextPos = ( LinePoints[0] + LinePoints[1] ) / 2.0f + Offset;
 
+	// Draw drop shadow
 	FSlateDrawElement::MakeText(
 		OutDrawElements,
 		LayerId,
+		AllottedGeometry.MakeChild(TextPos, TextPos).ToPaintGeometry(FVector2D(1, 1), AllottedGeometry.Size),
+		Text,
+		AnchorFont,
+		MyClippingRect,
+		ESlateDrawEffect::None,
+		FLinearColor::Black
+	);
+
+	FSlateDrawElement::MakeText(
+		OutDrawElements,
+		++LayerId,
 		AllottedGeometry.MakeChild(TextPos, TextPos).ToPaintGeometry(),
 		Text,
 		AnchorFont,
 		MyClippingRect,
 		ESlateDrawEffect::None,
 		FLinearColor::White
-		);
+	);
 }
 
 void FCanvasSlotExtension::PaintCollisionLines(const TSet< FWidgetReference >& Selection, const FGeometry& AllottedGeometry, const FSlateRect& MyClippingRect, FSlateWindowElementList& OutDrawElements, int32 LayerId) const
