Fixed a crash when FUMGEditorModule::PostCompile is called 2 times (without PreCompile)

diff --git a/Engine/Source/Editor/UMGEditor/Private/UMGEditorModule.cpp b/Engine/Source/Editor/UMGEditor/Private/UMGEditorModule.cpp
--- a/Engine/Source/Editor/UMGEditor/Private/UMGEditorModule.cpp
+++ b/Engine/Source/Editor/UMGEditor/Private/UMGEditorModule.cpp
@@ -97,7 +97,11 @@ public:
 
 	void PostCompile(UBlueprint* Blueprint)
 	{
-		delete ReRegister;
+		if (ReRegister)
+		{
+			delete ReRegister;
+			ReRegister = nullptr;
+		}
 		
 		if ( GIsEditor && GEditor )
 		{
