UMG - The component re-registering now waits until all compiling finishes for cases when it's recursive.

diff --git a/Engine/Source/Editor/UMGEditor/Private/UMGEditorModule.cpp b/Engine/Source/Editor/UMGEditor/Private/UMGEditorModule.cpp
--- a/Engine/Source/Editor/UMGEditor/Private/UMGEditorModule.cpp
+++ b/Engine/Source/Editor/UMGEditor/Private/UMGEditorModule.cpp
@@ -26,8 +26,9 @@ class FUMGEditorModule : public IUMGEditorModule, public IBlueprintCompiler
 public:
 	/** Constructor, set up console commands and variables **/
 	FUMGEditorModule()
+		: ReRegister(nullptr)
+		, CompileCount(0)
 	{
-		ReRegister = nullptr;
 	}
 
 	/** Called right after the module DLL has been loaded and the module object has been created */
@@ -82,7 +83,12 @@ public:
 
 	void PreCompile(UBlueprint* Blueprint)
 	{
-		ReRegister = new TComponentReregisterContext<UWidgetComponent>();
+		if ( CompileCount == 0 )
+		{
+			ReRegister = new TComponentReregisterContext<UWidgetComponent>();
+		}
+
+		CompileCount++;
 	}
 
 	void Compile(UBlueprint* Blueprint, const FKismetCompilerOptions& CompileOptions, FCompilerResultsLog& Results, TArray<UObject*>* ObjLoaded)
@@ -97,7 +103,9 @@ public:
 
 	void PostCompile(UBlueprint* Blueprint)
 	{
-		if (ReRegister)
+		CompileCount--;
+
+		if ( ReRegister && CompileCount == 0 )
 		{
 			delete ReRegister;
 			ReRegister = nullptr;
@@ -127,7 +135,14 @@ private:
 	/** All created asset type actions.  Cached here so that we can unregister it during shutdown. */
 	TArray< TSharedPtr<IAssetTypeActions> > CreatedAssetTypeActions;
 
+	/** The temporary variable that captures and reinstances components after compiling finishes. */
 	TComponentReregisterContext<UWidgetComponent>* ReRegister;
+
+	/**
+	 * The current count on the number of compiles that have occurred.  We don't want to re-register components until all
+	 * compiling has stopped.
+	 */
+	int32 CompileCount;
 };
 
 IMPLEMENT_MODULE(FUMGEditorModule, UMGEditor);
