UE-7560 - UMG - Fixing a crash that occurs when compiling a UMG widget without any animations, and without the timeline/animations visible.  Causes an object to be garbage collected.  We now report the object reference to prevent GC.

diff --git a/Engine/Source/Editor/UMGEditor/Private/Animation/UMGSequencerObjectBindingManager.cpp b/Engine/Source/Editor/UMGEditor/Private/Animation/UMGSequencerObjectBindingManager.cpp
--- a/Engine/Source/Editor/UMGEditor/Private/Animation/UMGSequencerObjectBindingManager.cpp
+++ b/Engine/Source/Editor/UMGEditor/Private/Animation/UMGSequencerObjectBindingManager.cpp
@@ -25,6 +25,11 @@ FUMGSequencerObjectBindingManager::~FUMGSequencerObjectBindingManager()
 	WidgetBlueprintEditor.GetOnWidgetPreviewUpdated().RemoveAll( this );
 }
 
+void FUMGSequencerObjectBindingManager::AddReferencedObjects(FReferenceCollector& Collector)
+{
+	Collector.AddReferencedObject(WidgetAnimation);
+}
+
 FGuid FUMGSequencerObjectBindingManager::FindGuidForObject( const UMovieScene& MovieScene, UObject& Object ) const
 {
 	return PreviewObjectToGuidMap.FindRef( &Object );
@@ -97,7 +102,7 @@ void FUMGSequencerObjectBindingManager::GetRuntimeObjects( const TSharedRef<FMov
 bool FUMGSequencerObjectBindingManager::HasValidWidgetAnimation() const
 {
 	UWidgetBlueprint* WidgetBlueprint = WidgetBlueprintEditor.GetWidgetBlueprintObj();
-	return WidgetAnimation.IsValid() && WidgetBlueprint->Animations.Contains( WidgetAnimation.Get() );
+	return WidgetAnimation != nullptr && WidgetBlueprint->Animations.Contains( WidgetAnimation );
 }
 
 void FUMGSequencerObjectBindingManager::InitPreviewObjects()
diff --git a/Engine/Source/Editor/UMGEditor/Private/Animation/UMGSequencerObjectBindingManager.h b/Engine/Source/Editor/UMGEditor/Private/Animation/UMGSequencerObjectBindingManager.h
--- a/Engine/Source/Editor/UMGEditor/Private/Animation/UMGSequencerObjectBindingManager.h
+++ b/Engine/Source/Editor/UMGEditor/Private/Animation/UMGSequencerObjectBindingManager.h
@@ -7,12 +7,15 @@
 class FWidgetBlueprintEditor;
 class UWidgetAnimation;
 
-class FUMGSequencerObjectBindingManager : public ISequencerObjectBindingManager
+class FUMGSequencerObjectBindingManager : public FGCObject, public ISequencerObjectBindingManager
 {
 public:
 	FUMGSequencerObjectBindingManager( FWidgetBlueprintEditor& InWidgetBlueprintEditor, UWidgetAnimation& InAnimation );
 	~FUMGSequencerObjectBindingManager();
 
+	/** FGCObject interface */
+	virtual void AddReferencedObjects(FReferenceCollector& Collector);
+
 	/** ISequencerObjectBindingManager interface */
 	virtual bool AllowsSpawnableObjects() const override { return false; }
 	virtual FGuid FindGuidForObject( const UMovieScene& MovieScene, UObject& Object ) const override;
@@ -27,7 +30,7 @@ public:
 	bool HasValidWidgetAnimation() const;
 
 	/** @return the current widget animation */
-	UWidgetAnimation* GetWidgetAnimation() { return WidgetAnimation.Get(); }
+	UWidgetAnimation* GetWidgetAnimation() { return WidgetAnimation; }
 
 	/** Rebuilds mappings to live preview objects */
 	void InitPreviewObjects();
@@ -38,7 +41,7 @@ private:
 	/** Mapping of preview objects to sequencer guids */
 	TMap< TWeakObjectPtr<UObject>, FGuid> PreviewObjectToGuidMap;
 
-	TWeakObjectPtr<UWidgetAnimation> WidgetAnimation;
+	UWidgetAnimation* WidgetAnimation;
 
 	FWidgetBlueprintEditor& WidgetBlueprintEditor;
 };
