UE-7550 - UMG - Clearing the selected objects when the widget selection is about to change.  This will ensure all last minute changes like those coming from the color picker while the eye dropper are finished get a chance to finish, before the migration set of widgets changes out from under it.

diff --git a/Engine/Source/Editor/UMGEditor/Private/Details/SWidgetDetailsView.cpp b/Engine/Source/Editor/UMGEditor/Private/Details/SWidgetDetailsView.cpp
--- a/Engine/Source/Editor/UMGEditor/Private/Details/SWidgetDetailsView.cpp
+++ b/Engine/Source/Editor/UMGEditor/Private/Details/SWidgetDetailsView.cpp
@@ -196,6 +196,11 @@ void SWidgetDetailsView::RegisterCustomizations()
 void SWidgetDetailsView::OnEditorSelectionChanging()
 {
 	ClearFocusIfOwned();
+
+	// We force the destruction of the currently monitored object when selection is about to change, to ensure all migrations occur
+	// immediately.
+	SelectedObjects.Empty();
+	PropertyView->SetObjects(SelectedObjects);
 }
 
 void SWidgetDetailsView::OnEditorSelectionChanged()
