UMG - The Category for a widget is now recorded in tag data parsable without loading the widget into memory.  So now they can be categoried in the palette correctly even when unloaded.

diff --git a/Engine/Source/Editor/UMGEditor/Private/Details/SWidgetDetailsView.cpp b/Engine/Source/Editor/UMGEditor/Private/Details/SWidgetDetailsView.cpp
--- a/Engine/Source/Editor/UMGEditor/Private/Details/SWidgetDetailsView.cpp
+++ b/Engine/Source/Editor/UMGEditor/Private/Details/SWidgetDetailsView.cpp
@@ -292,8 +292,12 @@ void SWidgetDetailsView::HandleCategoryTextCommitted(const FText& Text, ETextCom
 			UUserWidget* WidgetCDO = Widget->GetClass()->GetDefaultObject<UUserWidget>();
 			WidgetCDO->PaletteCategory = Text;
 
+			// Set the new category on the widget blueprint as well so that it's available when the blueprint isn't loaded.
+			UWidgetBlueprint* Blueprint = BlueprintEditor.Pin()->GetWidgetBlueprintObj();
+			Blueprint->PaletteCategory = Text.ToString();
+
 			// Immediately force a rebuild so that all palettes update to show it in a new category.
-			FBlueprintEditorUtils::MarkBlueprintAsStructurallyModified(BlueprintEditor.Pin()->GetBlueprintObj());
+			FBlueprintEditorUtils::MarkBlueprintAsStructurallyModified(Blueprint);
 		}
 	}
 }
diff --git a/Engine/Source/Editor/UMGEditor/Private/Templates/WidgetTemplateBlueprintClass.cpp b/Engine/Source/Editor/UMGEditor/Private/Templates/WidgetTemplateBlueprintClass.cpp
--- a/Engine/Source/Editor/UMGEditor/Private/Templates/WidgetTemplateBlueprintClass.cpp
+++ b/Engine/Source/Editor/UMGEditor/Private/Templates/WidgetTemplateBlueprintClass.cpp
@@ -35,9 +35,18 @@ FText FWidgetTemplateBlueprintClass::GetCategory() const
 	}
 	else
 	{
-		//If the blueprint is unloaded we need to extract it from the asset metadata.
-		auto DefaultUserWidget = UUserWidget::StaticClass()->GetDefaultObject<UUserWidget>();
-		return DefaultUserWidget->GetPaletteCategory();
+		const FString* FoundPaletteCategory = WidgetAssetData.TagsAndValues.Find(GET_MEMBER_NAME_CHECKED(UWidgetBlueprint, PaletteCategory));
+
+		if ( FoundPaletteCategory && !FoundPaletteCategory->IsEmpty() )
+		{
+			return FText::FromString(*FoundPaletteCategory);
+		}
+		else
+		{
+			//If the blueprint is unloaded we need to extract it from the asset metadata.
+			auto DefaultUserWidget = UUserWidget::StaticClass()->GetDefaultObject<UUserWidget>();
+			return DefaultUserWidget->GetPaletteCategory();
+		}
 	}
 }
 
diff --git a/Engine/Source/Editor/UMGEditor/Public/WidgetBlueprint.h b/Engine/Source/Editor/UMGEditor/Public/WidgetBlueprint.h
--- a/Engine/Source/Editor/UMGEditor/Public/WidgetBlueprint.h
+++ b/Engine/Source/Editor/UMGEditor/Public/WidgetBlueprint.h
@@ -194,6 +194,14 @@ public:
 	UPROPERTY()
 	TArray<UWidgetAnimation*> Animations;
 
+	/**
+	 * Don't directly modify this property to change the palette category.  The actual value is stored 
+	 * in the CDO of the UUserWidget, but a copy is stored here so that it's available in the serialized 
+	 * Tag data in the asset header for access in the FAssetData.
+	 */
+	UPROPERTY(AssetRegistrySearchable)
+	FString PaletteCategory;
+
 public:
 
 	/** UObject interface */
