[INTEGRATE] Change 2425674 by [REDACTED_PERSON][REDACTED_HANDLE] on 2015/01/30 09:38:37

	UMG - Removing some dead code from the AssetTypeActions for widget blueprints.

diff --git a/Engine/Source/Editor/UMGEditor/Private/AssetTypeActions_WidgetBlueprint.cpp b/Engine/Source/Editor/UMGEditor/Private/AssetTypeActions_WidgetBlueprint.cpp
--- a/Engine/Source/Editor/UMGEditor/Private/AssetTypeActions_WidgetBlueprint.cpp
+++ b/Engine/Source/Editor/UMGEditor/Private/AssetTypeActions_WidgetBlueprint.cpp
@@ -30,13 +30,10 @@ void FAssetTypeActions_WidgetBlueprint::OpenAssetEditor( const TArray<UObject*>&
 			TArray<UBlueprint*> Blueprints;
 			Blueprints.Add(Blueprint);
 			NewBlueprintEditor->InitWidgetBlueprintEditor(Mode, EditWithinLevelEditor, Blueprints, true);
-
-			//FBlueprintEditorModule& BlueprintEditorModule = FModuleManager::LoadModuleChecked<FBlueprintEditorModule>( "Kismet" );
-			//TSharedRef< IBlueprintEditor > NewKismetEditor = BlueprintEditorModule.CreateBlueprintEditor( Mode, EditWithinLevelEditor, Blueprint, ShouldUseDataOnlyEditor(Blueprint) );
 		}
 		else
 		{
-			FMessageDialog::Open( EAppMsgType::Ok, LOCTEXT("FailedToLoadBlueprint", "Widget Blueprint could not be loaded because it derives from an invalid class.  Check to make sure the parent class for this blueprint hasn't been removed!"));
+			FMessageDialog::Open( EAppMsgType::Ok, LOCTEXT("FailedToLoadWidgetBlueprint", "Widget Blueprint could not be loaded because it derives from an invalid class.\nCheck to make sure the parent class for this blueprint hasn't been removed!"));
 		}
 	}
 }
@@ -46,15 +43,6 @@ UClass* FAssetTypeActions_WidgetBlueprint::GetSupportedClass() const
 	return UWidgetBlueprint::StaticClass();
 }
 
-bool FAssetTypeActions_WidgetBlueprint::ShouldUseDataOnlyEditor(const UBlueprint* Blueprint) const
-{
-	return FBlueprintEditorUtils::IsDataOnlyBlueprint(Blueprint)
-		&& !FBlueprintEditorUtils::IsLevelScriptBlueprint(Blueprint)
-		&& !FBlueprintEditorUtils::IsInterfaceBlueprint(Blueprint)
-		&& !Blueprint->bForceFullEditor
-		&& !Blueprint->bIsNewlyCreated;
-}
-
 FText FAssetTypeActions_WidgetBlueprint::GetAssetDescription( const FAssetData& AssetData ) const
 {
 	if ( const FString* pDescription = AssetData.TagsAndValues.Find( GET_MEMBER_NAME_CHECKED( UBlueprint, BlueprintDescription ) ) )
