[INTEGRATE] Change 2425763 by [REDACTED_PERSON][REDACTED_HANDLE] on 2015/01/30 10:26:44

	UMG - Improving the visuals of the binding widgets for multicast events.

diff --git a/Engine/Source/Editor/UMGEditor/Private/Customizations/DetailCustomizations.cpp b/Engine/Source/Editor/UMGEditor/Private/Customizations/DetailCustomizations.cpp
--- a/Engine/Source/Editor/UMGEditor/Private/Customizations/DetailCustomizations.cpp
+++ b/Engine/Source/Editor/UMGEditor/Private/Customizations/DetailCustomizations.cpp
@@ -27,6 +27,8 @@ class SGraphSchemaActionButton : public SCompoundWidget, public FGCObject
 public:
 
 	SLATE_BEGIN_ARGS(SGraphSchemaActionButton) {}
+		/** Slot for this designers content (optional) */
+		SLATE_DEFAULT_SLOT(FArguments, Content)
 	SLATE_END_ARGS()
 
 	void Construct(const FArguments& InArgs, TSharedPtr<FWidgetBlueprintEditor> InEditor, TSharedPtr<FEdGraphSchemaAction> InClickAction)
@@ -37,12 +39,14 @@ public:
 		ChildSlot
 		[
 			SNew(SButton)
+			.ButtonStyle(FEditorStyle::Get(), "RoundButton")
+			.TextStyle(FEditorStyle::Get(), "NormalText")
+			.HAlign(HAlign_Center)
+			.ForegroundColor(FSlateColor::UseForeground())
 			.ToolTipText(FText::FromString(Action->TooltipDescription))
 			.OnClicked(this, &SGraphSchemaActionButton::AddOrViewEventBinding)
-			.HAlign(HAlign_Center)
 			[
-				SNew(STextBlock)
-				.Text(Action->MenuDescription)
+				InArgs._Content.Widget
 			]
 		];
 	}
@@ -189,6 +193,7 @@ void FBlueprintWidgetCustomization::CreateMulticastEventCustomization(IDetailLay
 	const UK2Node* EventNode = FKismetEditorUtilities::FindBoundEventForComponent(Blueprint, DelegateProperty->GetFName(), ComponentProperty->GetFName());
 
 	TSharedPtr<FEdGraphSchemaAction> ClickAction;
+	TSharedPtr<SWidget> ButtonContent;
 
 	if ( EventNode )
 	{
@@ -196,6 +201,9 @@ void FBlueprintWidgetCustomization::CreateMulticastEventCustomization(IDetailLay
 		NewDelegateNode->NodePtr = EventNode;
 
 		ClickAction = NewDelegateNode;
+
+		ButtonContent = SNew(STextBlock)
+			.Text(LOCTEXT("ViewEvent", "View"));
 	}
 	else
 	{
@@ -207,6 +215,9 @@ void FBlueprintWidgetCustomization::CreateMulticastEventCustomization(IDetailLay
 		NewDelegateNode->bGotoNode = true;
 
 		ClickAction = NewDelegateNode;
+
+		ButtonContent = SNew(SImage)
+			.Image(FEditorStyle::GetBrush("Plus"));
 	}
 
 	TSharedRef<IPropertyHandle> DelegatePropertyHandle = DetailLayout.GetProperty(DelegateProperty->GetFName(), CastChecked<UClass>(DelegateProperty->GetOuter()));
@@ -232,14 +243,18 @@ void FBlueprintWidgetCustomization::CreateMulticastEventCustomization(IDetailLay
 		.VAlign(VAlign_Center)
 		[
 			SNew(STextBlock)
+			.Font(IDetailLayoutBuilder::GetDetailFont())
 			.Text(DelegatePropertyName)
 		]
 	]
 	.ValueContent()
-	.MinDesiredWidth(200)
-	.MaxDesiredWidth(250)
+	.MinDesiredWidth(150)
+	.MaxDesiredWidth(200)
 	[
 		SNew(SGraphSchemaActionButton, Editor.Pin(), ClickAction)
+		[
+			ButtonContent.ToSharedRef()
+		]
 	];
 }
 
