[INTEGRATE] Change 2426213 by [REDACTED_PERSON][REDACTED_HANDLE] on 2015/01/30 15:03:46

	UMG - Build fix, we no longer use this function.

diff --git a/Engine/Source/Editor/UMGEditor/Private/AssetTypeActions_WidgetBlueprint.h b/Engine/Source/Editor/UMGEditor/Private/AssetTypeActions_WidgetBlueprint.h
--- a/Engine/Source/Editor/UMGEditor/Private/AssetTypeActions_WidgetBlueprint.h
+++ b/Engine/Source/Editor/UMGEditor/Private/AssetTypeActions_WidgetBlueprint.h
@@ -14,8 +14,4 @@ public:
 	virtual uint32 GetCategories() override { return EAssetTypeCategories::UI; }
 	virtual FText GetAssetDescription( const FAssetData& AssetData ) const override;
 	// End IAssetTypeActions Implementation
-
-private:
-	/** Returns true if the blueprint is data only */
-	bool ShouldUseDataOnlyEditor(const UBlueprint* Blueprint) const;
 };
\ No newline at end of file
