[INTEGRATE] Change 2426449 by [REDACTED_PERSON][REDACTED_HANDLE] on 2015/01/30 17:18:02

	UMG - Changing color of widget blueprints to blue.

diff --git a/Engine/Source/Editor/UMGEditor/Private/AssetTypeActions_WidgetBlueprint.h b/Engine/Source/Editor/UMGEditor/Private/AssetTypeActions_WidgetBlueprint.h
--- a/Engine/Source/Editor/UMGEditor/Private/AssetTypeActions_WidgetBlueprint.h
+++ b/Engine/Source/Editor/UMGEditor/Private/AssetTypeActions_WidgetBlueprint.h
@@ -8,7 +8,7 @@ class FAssetTypeActions_WidgetBlueprint : public FAssetTypeActions_Base
 public:
 	// IAssetTypeActions Implementation
 	virtual FText GetName() const override { return NSLOCTEXT("AssetTypeActions", "AssetTypeActions_WidgetBlueprint", "Widget Blueprint"); }
-	virtual FColor GetTypeColor() const override { return FColor(255,255,255); }
+	virtual FColor GetTypeColor() const override { return FColor(44, 89, 180); }
 	virtual UClass* GetSupportedClass() const override;
 	virtual void OpenAssetEditor(const TArray<UObject*>& InObjects, TSharedPtr<class IToolkitHost> EditWithinLevelEditor = TSharedPtr<IToolkitHost>()) override;
 	virtual uint32 GetCategories() override { return EAssetTypeCategories::UI; }
