UMG - Widgets in the designer can now be previewed using the widget's own desired size.

diff --git a/Engine/Source/Editor/UMGEditor/Private/Designer/SDesignerView.cpp b/Engine/Source/Editor/UMGEditor/Private/Designer/SDesignerView.cpp
--- a/Engine/Source/Editor/UMGEditor/Private/Designer/SDesignerView.cpp
+++ b/Engine/Source/Editor/UMGEditor/Private/Designer/SDesignerView.cpp
@@ -176,6 +176,8 @@ void SDesignerView::Construct(const FArguments& InArgs, TSharedPtr<FWidgetBluepr
 
 	SetStartupResolution();
 
+	CachedPreviewDesiredSize = FVector2D(0, 0);
+
 	ResolutionTextFade = FCurveSequence(0.0f, 1.0f);
 	ResolutionTextFade.Play(this->AsShared());
 
@@ -647,6 +649,10 @@ FOptionalSize SDesignerView::GetPreviewWidth() const
 		{
 			return DefaultWidget->DesignTimeSize.X;
 		}
+		else if ( DefaultWidget->bUseDesiredSizeAtDesignTime )
+		{
+			return CachedPreviewDesiredSize.X;
+		}
 	}
 
 	return (float)PreviewWidth;
@@ -660,6 +666,10 @@ FOptionalSize SDesignerView::GetPreviewHeight() const
 		{
 			return DefaultWidget->DesignTimeSize.Y;
 		}
+		else if ( DefaultWidget->bUseDesiredSizeAtDesignTime )
+		{
+			return CachedPreviewDesiredSize.Y;
+		}
 	}
 
 	return (float)PreviewHeight;
@@ -670,7 +680,7 @@ float SDesignerView::GetPreviewDPIScale() const
 	// If the user is using a custom size then we disable the DPI scaling logic.
 	if ( UUserWidget* DefaultWidget = GetDefaultWidget() )
 	{
-		if ( DefaultWidget->bUseDesignTimeSize )
+		if ( DefaultWidget->bUseDesignTimeSize || DefaultWidget->bUseDesiredSizeAtDesignTime )
 		{
 			return 1.0f;
 		}
@@ -1425,6 +1435,18 @@ void SDesignerView::Tick( const FGeometry& AllottedGeometry, const double InCurr
 		TopRuler->SetCursor(TOptional<FVector2D>());
 		SideRuler->SetCursor(TOptional<FVector2D>());
 	}
+
+	if ( UUserWidget* DefaultWidget = GetDefaultWidget() )
+	{
+		if ( DefaultWidget->bUseDesiredSizeAtDesignTime && PreviewWidget )
+		{
+			TSharedPtr<SWidget> CachedWidget = PreviewWidget->GetCachedWidget();
+			if ( CachedWidget.IsValid() )
+			{
+				CachedPreviewDesiredSize = CachedWidget->GetDesiredSize();
+			}
+		}
+	}
 }
 
 FReply SDesignerView::OnDragDetected(const FGeometry& MyGeometry, const FPointerEvent& MouseEvent)
@@ -1909,7 +1931,7 @@ EVisibility SDesignerView::GetResolutionTextVisibility() const
 	// If we're using a custom design time size, don't bother showing the resolution
 	if ( UUserWidget* DefaultWidget = GetDefaultWidget() )
 	{
-		if ( DefaultWidget->bUseDesignTimeSize )
+		if ( DefaultWidget->bUseDesignTimeSize || DefaultWidget->bUseDesiredSizeAtDesignTime )
 		{
 			return EVisibility::Collapsed;
 		}
@@ -1949,6 +1971,7 @@ void SDesignerView::HandleOnCommonResolutionSelected(int32 Width, int32 Height,
 	if ( UUserWidget* DefaultWidget = GetDefaultWidget() )
 	{
 		DefaultWidget->bUseDesignTimeSize = false;
+		DefaultWidget->bUseDesiredSizeAtDesignTime = false;
 		MarkDesignModifed(/*bRequiresRecompile*/ false);
 	}
 
@@ -1960,7 +1983,7 @@ bool SDesignerView::HandleIsCommonResolutionSelected(int32 Width, int32 Height)
 	// If we're using a custom design time size, none of the other resolutions should appear selected, even if they match.
 	if ( UUserWidget* DefaultWidget = GetDefaultWidget() )
 	{
-		if ( DefaultWidget->bUseDesignTimeSize )
+		if ( DefaultWidget->bUseDesignTimeSize || DefaultWidget->bUseDesiredSizeAtDesignTime )
 		{
 			return false;
 		}
@@ -1996,11 +2019,32 @@ bool SDesignerView::HandleIsCustomResolutionSelected() const
 	return false;
 }
 
+bool SDesignerView::HandleIsDesiredSizeSelected() const
+{
+	if ( UUserWidget* DefaultWidget = GetDefaultWidget() )
+	{
+		return DefaultWidget->bUseDesiredSizeAtDesignTime;
+	}
+
+	return false;
+}
+
 void SDesignerView::HandleOnCustomResolutionSelected()
 {
 	if ( UUserWidget* DefaultWidget = GetDefaultWidget() )
 	{
 		DefaultWidget->bUseDesignTimeSize = true;
+		DefaultWidget->bUseDesiredSizeAtDesignTime = false;
+		MarkDesignModifed(/*bRequiresRecompile*/ false);
+	}
+}
+
+void SDesignerView::HandleOnDesiredSizeSelected()
+{
+	if ( UUserWidget* DefaultWidget = GetDefaultWidget() )
+	{
+		DefaultWidget->bUseDesignTimeSize = false;
+		DefaultWidget->bUseDesiredSizeAtDesignTime = true;
 		MarkDesignModifed(/*bRequiresRecompile*/ false);
 	}
 }
@@ -2070,11 +2114,21 @@ TSharedRef<SWidget> SDesignerView::GetAspectMenu()
 	FMenuBuilder MenuBuilder(true, nullptr);
 
 	// Add custom option
-	FExecuteAction OnResolutionSelected = FExecuteAction::CreateRaw(this, &SDesignerView::HandleOnCustomResolutionSelected);
-	FIsActionChecked OnIsResolutionSelected = FIsActionChecked::CreateRaw(this, &SDesignerView::HandleIsCustomResolutionSelected);
-	FUIAction Action(OnResolutionSelected, FCanExecuteAction(), OnIsResolutionSelected);
+	FUIAction CustomSizeAction(
+		FExecuteAction::CreateRaw(this, &SDesignerView::HandleOnCustomResolutionSelected),
+		FCanExecuteAction(),
+		FIsActionChecked::CreateRaw(this, &SDesignerView::HandleIsCustomResolutionSelected));
+
+	MenuBuilder.AddMenuEntry(LOCTEXT("Custom", "Custom"), LOCTEXT("Custom", "Custom"), FSlateIcon(), CustomSizeAction, NAME_None, EUserInterfaceActionType::Check);
+
+	// Add desired size option
+	FUIAction DesiredSizeAction(
+		FExecuteAction::CreateRaw(this, &SDesignerView::HandleOnDesiredSizeSelected),
+		FCanExecuteAction(),
+		FIsActionChecked::CreateRaw(this, &SDesignerView::HandleIsDesiredSizeSelected));
+
 
-	MenuBuilder.AddMenuEntry(LOCTEXT("Custom", "Custom"), LOCTEXT("Custom", "Custom"), FSlateIcon(), Action, NAME_None, EUserInterfaceActionType::Check);
+	MenuBuilder.AddMenuEntry(LOCTEXT("DesiredSize", "Desired Size"), LOCTEXT("DesiredSize", "Desired Size"), FSlateIcon(), DesiredSizeAction, NAME_None, EUserInterfaceActionType::Check);
 
 	// Add the normal set of resultion options.
 	AddScreenResolutionSection(MenuBuilder, PlaySettings->PhoneScreenResolutions, LOCTEXT("CommonPhonesSectionHeader", "Phones"));
diff --git a/Engine/Source/Editor/UMGEditor/Private/Designer/SDesignerView.h b/Engine/Source/Editor/UMGEditor/Private/Designer/SDesignerView.h
--- a/Engine/Source/Editor/UMGEditor/Private/Designer/SDesignerView.h
+++ b/Engine/Source/Editor/UMGEditor/Private/Designer/SDesignerView.h
@@ -142,6 +142,8 @@ private:
 	void AddScreenResolutionSection(FMenuBuilder& MenuBuilder, const TArray<FPlayScreenResolution>& Resolutions, const FText& SectionName);
 	bool HandleIsCustomResolutionSelected() const;
 	void HandleOnCustomResolutionSelected();
+	bool HandleIsDesiredSizeSelected() const;
+	void HandleOnDesiredSizeSelected();
 	TSharedRef<SWidget> GetAspectMenu();
 
 	EVisibility PIENotification() const;
@@ -246,6 +248,9 @@ private:
 	/** The configured Height of the preview area, simulates screen size. */
 	int32 PreviewHeight;
 
+	/** We cache the desired preview desired size to maintain the same size between compiles when it lags a frame behind and no widget is available. */
+	FVector2D CachedPreviewDesiredSize;
+
 	// Resolution Info
 	FString PreviewAspectRatio;
 
diff --git a/Engine/Source/Runtime/UMG/Public/Blueprint/UserWidget.h b/Engine/Source/Runtime/UMG/Public/Blueprint/UserWidget.h
--- a/Engine/Source/Runtime/UMG/Public/Blueprint/UserWidget.h
+++ b/Engine/Source/Runtime/UMG/Public/Blueprint/UserWidget.h
@@ -654,11 +654,11 @@ public:
 	UPROPERTY()
 	FVector2D DesignTimeSize;
 
-	/** Designer flag used to tell when the user wants to provide their own custom size. */
+	/** A flag that determines if the design time size is used for previewing the widget in the designer. */
 	UPROPERTY()
 	bool bUseDesignTimeSize;
 
-	/** Designer flag used to tell when to just use the desired size of the widget instead of an explicit size. */
+	/** A flag that determines if the widget's desired size is used for previewing the widget in the designer. */
 	UPROPERTY()
 	bool bUseDesiredSizeAtDesignTime;
 
