[INTEGRATE] Change 2442422 by [REDACTED_PERSON][REDACTED_HANDLE] on 2015/02/11 17:16:53

	UE-9652 - UMG - Fixing several undo bugs that were not properly restoring the state in the designer after hitting undo for moving widgets and for resizing them.

diff --git a/Engine/Source/Editor/UMGEditor/Private/Designer/SDesignerView.cpp b/Engine/Source/Editor/UMGEditor/Private/Designer/SDesignerView.cpp
--- a/Engine/Source/Editor/UMGEditor/Private/Designer/SDesignerView.cpp
+++ b/Engine/Source/Editor/UMGEditor/Private/Designer/SDesignerView.cpp
@@ -1727,13 +1727,14 @@ UWidget* SDesignerView::ProcessDropAndAddWidget(const FGeometry& MyGeometry, con
 			}
 
 			UWidget* Widget = bIsPreview ? SelectedDragDropOp->Preview : SelectedDragDropOp->Template;
-			if ( !Widget )
-			{
-				Widget = bIsPreview ? SelectedDragDropOp->Preview : SelectedDragDropOp->Template;
-			}
 
 			if ( ensure(Widget) )
 			{
+				if ( !bIsPreview )
+				{
+					Widget->Modify();
+				}
+
 				if ( Widget->GetParent() )
 				{
 					if ( !bIsPreview )
@@ -1833,11 +1834,11 @@ UWidget* SDesignerView::ProcessDropAndAddWidget(const FGeometry& MyGeometry, con
 					// TODO UMG ERROR Slot can not be created because maybe the max children has been reached.
 					//          Maybe we can traverse the hierarchy and add it to the first parent that will accept it?
 				}
+			}
 
-				if ( bIsPreview )
-				{
-					Transaction.Cancel();
-				}
+			if ( bIsPreview )
+			{
+				Transaction.Cancel();
 			}
 		}
 		else
diff --git a/Engine/Source/Editor/UMGEditor/Private/Designer/STransformHandle.cpp b/Engine/Source/Editor/UMGEditor/Private/Designer/STransformHandle.cpp
--- a/Engine/Source/Editor/UMGEditor/Private/Designer/STransformHandle.cpp
+++ b/Engine/Source/Editor/UMGEditor/Private/Designer/STransformHandle.cpp
@@ -19,6 +19,7 @@ void STransformHandle::Construct(const FArguments& InArgs, IUMGDesigner* InDesig
 	Designer = InDesigner;
 
 	Action = ETransformAction::None;
+	ScopedTransaction = nullptr;
 
 	DragDirection = ComputeDragDirection(InTransformDirection);
 	DragOrigin = ComputeOrigin(InTransformDirection);
@@ -76,6 +77,7 @@ FReply STransformHandle::OnMouseButtonDown(const FGeometry& MyGeometry, const FP
 
 		FWidgetReference SelectedWidget = Designer->GetSelectedWidget();
 		UWidget* Preview = SelectedWidget.GetPreview();
+		UWidget* Template = SelectedWidget.GetTemplate();
 
 		if ( UCanvasPanelSlot* Slot = Cast<UCanvasPanelSlot>(Preview->Slot) )
 		{
@@ -84,6 +86,9 @@ FReply STransformHandle::OnMouseButtonDown(const FGeometry& MyGeometry, const FP
 
 		MouseDownPosition = MouseEvent.GetScreenSpacePosition();
 
+		ScopedTransaction = new FScopedTransaction(LOCTEXT("ResizeWidget", "Resize Widget"));
+		Template->Modify();
+
 		return FReply::Handled().CaptureMouse(SharedThis(this));
 	}
 
@@ -97,6 +102,12 @@ FReply STransformHandle::OnMouseButtonUp(const FGeometry& MyGeometry, const FPoi
 		const bool bRequiresRecompile = false;
 		Designer->MarkDesignModifed(bRequiresRecompile);
 
+		if ( ScopedTransaction )
+		{
+			delete ScopedTransaction;
+			ScopedTransaction = nullptr;
+		}
+
 		Action = ETransformAction::None;
 		return FReply::Handled().ReleaseMouseCapture();
 	}
diff --git a/Engine/Source/Editor/UMGEditor/Private/Designer/STransformHandle.h b/Engine/Source/Editor/UMGEditor/Private/Designer/STransformHandle.h
--- a/Engine/Source/Editor/UMGEditor/Private/Designer/STransformHandle.h
+++ b/Engine/Source/Editor/UMGEditor/Private/Designer/STransformHandle.h
@@ -69,4 +69,6 @@ protected:
 
 	FVector2D MouseDownPosition;
 	FMargin StartingOffsets;
+
+	FScopedTransaction* ScopedTransaction;
 };
