[INTEGRATE] Change 2442897 by [REDACTED_PERSON][REDACTED_HANDLE] on 2015/02/11 23:25:28

	Fix non-unity build

diff --git a/Engine/Source/Editor/UMGEditor/Private/Designer/STransformHandle.cpp b/Engine/Source/Editor/UMGEditor/Private/Designer/STransformHandle.cpp
--- a/Engine/Source/Editor/UMGEditor/Private/Designer/STransformHandle.cpp
+++ b/Engine/Source/Editor/UMGEditor/Private/Designer/STransformHandle.cpp
@@ -7,6 +7,7 @@
 #include "WidgetReference.h"
 #include "Widget.h"
 #include "Components/PanelSlot.h"
+#include "ScopedTransaction.h"
 
 #define LOCTEXT_NAMESPACE "UMG"
 
diff --git a/Engine/Source/Editor/UMGEditor/Private/Designer/STransformHandle.h b/Engine/Source/Editor/UMGEditor/Private/Designer/STransformHandle.h
--- a/Engine/Source/Editor/UMGEditor/Private/Designer/STransformHandle.h
+++ b/Engine/Source/Editor/UMGEditor/Private/Designer/STransformHandle.h
@@ -70,5 +70,5 @@ protected:
 	FVector2D MouseDownPosition;
 	FMargin StartingOffsets;
 
-	FScopedTransaction* ScopedTransaction;
+	class FScopedTransaction* ScopedTransaction;
 };
