UE-10040 - UMG - Fixing a paste crash that occurs if you attempt to paste widgets into a container that has no more room for them, you now get a notification of the failure.

diff --git a/Engine/Source/Editor/UMGEditor/Private/WidgetBlueprintEditorUtils.cpp b/Engine/Source/Editor/UMGEditor/Private/WidgetBlueprintEditorUtils.cpp
--- a/Engine/Source/Editor/UMGEditor/Private/WidgetBlueprintEditorUtils.cpp
+++ b/Engine/Source/Editor/UMGEditor/Private/WidgetBlueprintEditorUtils.cpp
@@ -18,6 +18,8 @@
 #include "K2Node_ComponentBoundEvent.h"
 #include "CanvasPanel.h"
 #include "WidgetSlotPair.h"
+#include "SNotificationList.h"
+#include "NotificationManager.h"
 
 #define LOCTEXT_NAMESPACE "UMG"
 
@@ -554,10 +556,10 @@ void FWidgetBlueprintEditorUtils::PasteWidgets(TSharedRef<FWidgetBlueprintEditor
 		{
 			if ( ParentWidget->GetChildrenCount() > 0 || RootPasteWidgets.Num() > 1 )
 			{
-				UCanvasPanel* PasteContainer = BP->WidgetTree->ConstructWidget<UCanvasPanel>(UCanvasPanel::StaticClass());
-				//TODO UMG The new container could be tiny, unless filling the space.
-				UPanelSlot* Slot = ParentWidget->AddChild(PasteContainer);
-				ParentWidget = PasteContainer;
+				FNotificationInfo Info(LOCTEXT("NotEnoughSlots", "Can't paste contents, not enough available slots in target widget."));
+				FSlateNotificationManager::Get().AddNotification(Info);
+
+				return;
 			}
 		}
 
