UMG - You can now rearrange widgets in the hierarchy.

diff --git a/Engine/Source/Editor/UMGEditor/Private/Hierarchy/SHierarchyViewItem.cpp b/Engine/Source/Editor/UMGEditor/Private/Hierarchy/SHierarchyViewItem.cpp
--- a/Engine/Source/Editor/UMGEditor/Private/Hierarchy/SHierarchyViewItem.cpp
+++ b/Engine/Source/Editor/UMGEditor/Private/Hierarchy/SHierarchyViewItem.cpp
@@ -99,8 +99,37 @@ void FHierarchyWidgetDragDropOp::OnDrop(bool bDropWasHandled, const FPointerEven
 
 //////////////////////////////////////////////////////////////////////////
 
-FReply ProcessHierarchyDragDrop(const FDragDropEvent& DragDropEvent, bool bIsDrop, TSharedPtr<FWidgetBlueprintEditor> BlueprintEditor, UWidgetBlueprint* Blueprint, FWidgetReference TargetItem)
+TOptional<EItemDropZone> ProcessHierarchyDragDrop(const FDragDropEvent& DragDropEvent, EItemDropZone DropZone, bool bIsDrop, TSharedPtr<FWidgetBlueprintEditor> BlueprintEditor, FWidgetReference TargetItem, TOptional<int32> Index = TOptional<int32>())
 {
+	UWidget* TargetTemplate = TargetItem.GetTemplate();
+
+	if ( TargetTemplate && (DropZone == EItemDropZone::AboveItem || DropZone == EItemDropZone::BelowItem) )
+	{
+		if ( UPanelWidget* TargetParentTemplate = Cast<UPanelWidget>(TargetTemplate->GetParent()) )
+		{
+			int32 InsertIndex = TargetParentTemplate->GetChildIndex(TargetTemplate);
+			InsertIndex += ( DropZone == EItemDropZone::AboveItem ) ? 0 : 1;
+			InsertIndex = FMath::Clamp(InsertIndex, 0, TargetParentTemplate->GetChildrenCount());
+
+			FWidgetReference TargetParentTemplateRef = BlueprintEditor->GetReferenceFromTemplate(TargetParentTemplate);
+			TOptional<EItemDropZone> ParentZone = ProcessHierarchyDragDrop(DragDropEvent, EItemDropZone::OntoItem, bIsDrop, BlueprintEditor, TargetParentTemplateRef, InsertIndex);
+			if ( ParentZone.IsSet() )
+			{
+				return DropZone;
+			}
+			else
+			{
+				DropZone = EItemDropZone::OntoItem;
+			}
+		}
+	}
+	else
+	{
+		DropZone = EItemDropZone::OntoItem;
+	}
+
+	UWidgetBlueprint* Blueprint = BlueprintEditor->GetWidgetBlueprintObj();
+
 	// Is this a drag/drop op to create a new widget in the tree?
 	TSharedPtr<FWidgetTemplateDragDropOp> TemplateDragDropOp = DragDropEvent.GetOperationAs<FWidgetTemplateDragDropOp>();
 	if ( TemplateDragDropOp.IsValid() )
@@ -108,7 +137,7 @@ FReply ProcessHierarchyDragDrop(const FDragDropEvent& DragDropEvent, bool bIsDro
 		TemplateDragDropOp->SetCursorOverride(TOptional<EMouseCursor::Type>());
 
 		// Are we adding to the root?
-		if ( !TargetItem.IsValid() && Blueprint->WidgetTree->RootWidget == NULL )
+		if ( !TargetItem.IsValid() && Blueprint->WidgetTree->RootWidget == nullptr )
 		{
 			// TODO UMG Allow showing a preview of this.
 			if ( bIsDrop )
@@ -116,6 +145,8 @@ FReply ProcessHierarchyDragDrop(const FDragDropEvent& DragDropEvent, bool bIsDro
 				Blueprint->WidgetTree->RootWidget = TemplateDragDropOp->Template->Create(Blueprint->WidgetTree);
 				FBlueprintEditorUtils::MarkBlueprintAsStructurallyModified(Blueprint);
 			}
+
+			return EItemDropZone::OntoItem;
 		}
 		// Are we adding to a panel?
 		else if ( UPanelWidget* Parent = Cast<UPanelWidget>(TargetItem.GetTemplate()) )
@@ -125,17 +156,28 @@ FReply ProcessHierarchyDragDrop(const FDragDropEvent& DragDropEvent, bool bIsDro
 			{
 				UWidget* Widget = TemplateDragDropOp->Template->Create(Blueprint->WidgetTree);
 				
-				Parent->AddChild(Widget);
+				UPanelSlot* NewSlot = nullptr;
+				if ( Index.IsSet() )
+				{
+					NewSlot = Parent->InsertChildAt(Index.GetValue(), Widget);
+				}
+				else
+				{
+					NewSlot = Parent->AddChild(Widget);
+				}
+				check(NewSlot);
 
 				FBlueprintEditorUtils::MarkBlueprintAsStructurallyModified(Blueprint);
 			}
+
+			return EItemDropZone::OntoItem;
 		}
 		else
 		{
 			TemplateDragDropOp->SetCursorOverride(EMouseCursor::SlashedCircle);
 		}
 
-		return FReply::Handled();
+		return TOptional<EItemDropZone>();
 	}
 
 	TSharedPtr<FHierarchyWidgetDragDropOp> HierarchyDragDropOp = DragDropEvent.GetOperationAs<FHierarchyWidgetDragDropOp>();
@@ -152,7 +194,7 @@ FReply ProcessHierarchyDragDrop(const FDragDropEvent& DragDropEvent, bool bIsDro
 			if ( bIsDraggedObject )
 			{
 				HierarchyDragDropOp->SetCursorOverride(EMouseCursor::SlashedCircle);
-				return FReply::Handled();
+				return TOptional<EItemDropZone>();
 			}
 
 			UPanelWidget* NewParent = Cast<UPanelWidget>(TargetItem.GetTemplate());
@@ -160,14 +202,14 @@ FReply ProcessHierarchyDragDrop(const FDragDropEvent& DragDropEvent, bool bIsDro
 			{
 				HierarchyDragDropOp->SetCursorOverride(EMouseCursor::SlashedCircle);
 				HierarchyDragDropOp->CurrentHoverText = LOCTEXT("CantHaveChildren", "Widget can't have children.");
-				return FReply::Handled();
+				return TOptional<EItemDropZone>();
 			}
 
 			if ( !NewParent->CanHaveMultipleChildren() && NewParent->GetChildrenCount() > 0 )
 			{
 				HierarchyDragDropOp->SetCursorOverride(EMouseCursor::SlashedCircle);
 				HierarchyDragDropOp->CurrentHoverText = LOCTEXT("NoAdditionalChildren", "Widget can't accept additional children.");
-				return FReply::Handled();
+				return TOptional<EItemDropZone>();
 			}
 
 			UWidget* TemplateWidget = HierarchyDragDropOp->Widget.GetTemplate();
@@ -185,17 +227,39 @@ FReply ProcessHierarchyDragDrop(const FDragDropEvent& DragDropEvent, bool bIsDro
 			{
 				HierarchyDragDropOp->SetCursorOverride(EMouseCursor::SlashedCircle);
 				HierarchyDragDropOp->CurrentHoverText = LOCTEXT("CantMakeWidgetChildOfChildren", "Can't make widget a child of its children.");
-				return FReply::Handled();
+				return TOptional<EItemDropZone>();
 			}
 
 			if ( bIsDrop )
 			{
+				if ( Index.IsSet() )
+				{
+					// If we're inserting at an index, and the widget we're moving is already
+					// in the hierarchy before the point we're moving it to, we need to reduce the index
+					// count by one, because the whole set is about to be shifted when it's removed.
+					const bool bInsertInSameParent = TemplateWidget->GetParent() == NewParent;
+					const bool bNeedToDropIndex = NewParent->GetChildIndex(TemplateWidget) < Index.GetValue();
+
+					if ( bInsertInSameParent && bNeedToDropIndex )
+					{
+						Index = Index.GetValue() - 1;
+					}
+				}
+
 				TemplateWidget->RemoveFromParent();
 
 				NewParent->SetFlags(RF_Transactional);
 				NewParent->Modify();
 
-				UPanelSlot* NewSlot = NewParent->AddChild(TemplateWidget);
+				UPanelSlot* NewSlot = nullptr;
+				if ( Index.IsSet() )
+				{
+					NewSlot = NewParent->InsertChildAt(Index.GetValue(), TemplateWidget);
+				}
+				else
+				{
+					NewSlot = NewParent->AddChild(TemplateWidget);
+				}
 				check(NewSlot);
 
 				// Import the old slot properties
@@ -208,17 +272,17 @@ FReply ProcessHierarchyDragDrop(const FDragDropEvent& DragDropEvent, bool bIsDro
 				BlueprintEditor->SelectWidgets(SelectedTemplates, false);
 			}
 
-			return FReply::Handled();
+			return EItemDropZone::OntoItem;
 		}
 		else
 		{
 			HierarchyDragDropOp->SetCursorOverride(EMouseCursor::SlashedCircle);
 		}
 
-		return FReply::Handled();
+		return TOptional<EItemDropZone>();
 	}
 
-	return FReply::Unhandled();
+	return TOptional<EItemDropZone>();
 }
 
 //////////////////////////////////////////////////////////////////////////
@@ -230,9 +294,9 @@ FHierarchyModel::FHierarchyModel()
 
 }
 
-FReply FHierarchyModel::OnDragOver(const FGeometry& MyGeometry, const FDragDropEvent& DragDropEvent)
+TOptional<EItemDropZone> FHierarchyModel::HandleCanAcceptDrop(const FDragDropEvent& DragDropEvent, EItemDropZone DropZone)
 {
-	return FReply::Unhandled();
+	return TOptional<EItemDropZone>();
 }
 
 FReply FHierarchyModel::HandleDragDetected(const FGeometry& MyGeometry, const FPointerEvent& MouseEvent)
@@ -255,7 +319,7 @@ void FHierarchyModel::HandleDragLeave(const FDragDropEvent& DragDropEvent)
 	}
 }
 
-FReply FHierarchyModel::HandleDrop(FDragDropEvent const& DragDropEvent)
+FReply FHierarchyModel::HandleAcceptDrop(FDragDropEvent const& DragDropEvent, EItemDropZone DropZone)
 {
 	return FReply::Unhandled();
 }
@@ -383,18 +447,17 @@ void FHierarchyRoot::UpdateSelection()
 	}
 }
 
-FReply FHierarchyRoot::OnDragOver(const FGeometry& MyGeometry, const FDragDropEvent& DragDropEvent)
+TOptional<EItemDropZone> FHierarchyRoot::HandleCanAcceptDrop(const FDragDropEvent& DragDropEvent, EItemDropZone DropZone)
 {
-	UWidgetBlueprint* Blueprint = BlueprintEditor.Pin()->GetWidgetBlueprintObj();
 	bool bIsDrop = false;
-	return ProcessHierarchyDragDrop(DragDropEvent, bIsDrop, BlueprintEditor.Pin(), Blueprint, FWidgetReference());
+	return ProcessHierarchyDragDrop(DragDropEvent, DropZone, bIsDrop, BlueprintEditor.Pin(), FWidgetReference());
 }
 
-FReply FHierarchyRoot::HandleDrop(FDragDropEvent const& DragDropEvent)
+FReply FHierarchyRoot::HandleAcceptDrop(FDragDropEvent const& DragDropEvent, EItemDropZone DropZone)
 {
-	UWidgetBlueprint* Blueprint = BlueprintEditor.Pin()->GetWidgetBlueprintObj();
 	bool bIsDrop = true;
-	return ProcessHierarchyDragDrop(DragDropEvent, bIsDrop, BlueprintEditor.Pin(), Blueprint, FWidgetReference());
+	TOptional<EItemDropZone> Zone = ProcessHierarchyDragDrop(DragDropEvent, DropZone, bIsDrop, BlueprintEditor.Pin(), FWidgetReference());
+	return Zone.IsSet() ? FReply::Handled() : FReply::Unhandled();
 }
 
 //////////////////////////////////////////////////////////////////////////
@@ -487,7 +550,7 @@ void FNamedSlotModel::UpdateSelection()
 	//}
 }
 
-FReply FNamedSlotModel::OnDragOver(const FGeometry& MyGeometry, const FDragDropEvent& DragDropEvent)
+TOptional<EItemDropZone> FNamedSlotModel::HandleCanAcceptDrop(const FDragDropEvent& DragDropEvent, EItemDropZone DropZone)
 {
 	TSharedPtr<FWidgetTemplateDragDropOp> TemplateDragDropOp = DragDropEvent.GetOperationAs<FWidgetTemplateDragDropOp>();
 	if ( TemplateDragDropOp.IsValid() )
@@ -499,15 +562,18 @@ FReply FNamedSlotModel::OnDragOver(const FGeometry& MyGeometry, const FDragDropE
 			{
 				TSharedPtr<FDecoratedDragDropOp> DecoratedDragDropOp = DragDropEvent.GetOperationAs<FDecoratedDragDropOp>();
 				TemplateDragDropOp->SetCursorOverride(EMouseCursor::SlashedCircle);
-				return FReply::Handled();
+
+				return TOptional<EItemDropZone>();
 			}
+
+			return EItemDropZone::OntoItem;
 		}
 	}
 
-	return FReply::Unhandled();
+	return TOptional<EItemDropZone>();
 }
 
-FReply FNamedSlotModel::HandleDrop(FDragDropEvent const& DragDropEvent)
+FReply FNamedSlotModel::HandleAcceptDrop(FDragDropEvent const& DragDropEvent, EItemDropZone DropZone)
 {
 	UWidgetBlueprint* Blueprint = BlueprintEditor.Pin()->GetWidgetBlueprintObj();
 
@@ -638,29 +704,15 @@ FSlateFontInfo FHierarchyWidget::GetFont() const
 	return FCoreStyle::Get().GetFontStyle(NormalFont);
 }
 
-FReply FHierarchyWidget::OnDragOver(const FGeometry& MyGeometry, const FDragDropEvent& DragDropEvent)
+TOptional<EItemDropZone> FHierarchyWidget::HandleCanAcceptDrop(const FDragDropEvent& DragDropEvent, EItemDropZone DropZone)
 {
-	TSharedPtr<FWidgetTemplateDragDropOp> TemplateDragDropOp = DragDropEvent.GetOperationAs<FWidgetTemplateDragDropOp>();
-	if ( TemplateDragDropOp.IsValid() )
-	{
-		// TODO UMG - Show feedback
-		return FReply::Handled();
-	}
-
-	TSharedPtr<FHierarchyWidgetDragDropOp> HierarchyDragDropOp = DragDropEvent.GetOperationAs<FHierarchyWidgetDragDropOp>();
-	if ( HierarchyDragDropOp.IsValid() )
-	{
-		bool bIsDrop = false;
-		UWidgetBlueprint* Blueprint = BlueprintEditor.Pin()->GetWidgetBlueprintObj();
-		return ProcessHierarchyDragDrop(DragDropEvent, bIsDrop, BlueprintEditor.Pin(), Blueprint, Item);
-	}
-
-	return FReply::Unhandled();
+	bool bIsDrop = false;
+	return ProcessHierarchyDragDrop(DragDropEvent, DropZone, bIsDrop, BlueprintEditor.Pin(), Item);
 }
 
 FReply FHierarchyWidget::HandleDragDetected(const FGeometry& MyGeometry, const FPointerEvent& MouseEvent)
 {
-	const bool bIsRoot = Item.GetTemplate()->GetParent() == NULL;
+	const bool bIsRoot = Item.GetTemplate()->GetParent() == nullptr;
 
 	if ( !bIsRoot )
 	{
@@ -680,11 +732,11 @@ void FHierarchyWidget::HandleDragLeave(const FDragDropEvent& DragDropEvent)
 	}
 }
 
-FReply FHierarchyWidget::HandleDrop(const FDragDropEvent& DragDropEvent)
+FReply FHierarchyWidget::HandleAcceptDrop(const FDragDropEvent& DragDropEvent, EItemDropZone DropZone)
 {
-	UWidgetBlueprint* Blueprint = BlueprintEditor.Pin()->GetWidgetBlueprintObj();
 	bool bIsDrop = true;
-	return ProcessHierarchyDragDrop(DragDropEvent, bIsDrop, BlueprintEditor.Pin(), Blueprint, Item);
+	TOptional<EItemDropZone> Zone = ProcessHierarchyDragDrop(DragDropEvent, DropZone, bIsDrop, BlueprintEditor.Pin(), Item);
+	return Zone.IsSet() ? FReply::Handled() : FReply::Unhandled();
 }
 
 bool FHierarchyWidget::OnVerifyNameTextChanged(const FText& InText, FText& OutErrorMessage)
@@ -701,7 +753,7 @@ bool FHierarchyWidget::OnVerifyNameTextChanged(const FText& InText, FText& OutEr
 	UWidget* ExistingTemplate = Blueprint->WidgetTree->FindWidget( FName(*NewName) );
 
 	bool bIsSameWidget = false;
-	if ( ExistingTemplate != NULL )
+	if ( ExistingTemplate != nullptr )
 	{
 		if ( Item.GetTemplate() != ExistingTemplate )
 		{
@@ -801,11 +853,12 @@ void SHierarchyViewItem::Construct(const FArguments& InArgs, const TSharedRef< S
 
 	STableRow< TSharedPtr<FHierarchyModel> >::Construct(
 		STableRow< TSharedPtr<FHierarchyModel> >::FArguments()
-		.Padding(2.0f)
+		.OnCanAcceptDrop(this, &SHierarchyViewItem::HandleCanAcceptDrop)
+		.OnAcceptDrop(this, &SHierarchyViewItem::HandleAcceptDrop)
 		.OnDragDetected(this, &SHierarchyViewItem::HandleDragDetected)
 		.OnDragEnter(this, &SHierarchyViewItem::HandleDragEnter)
 		.OnDragLeave(this, &SHierarchyViewItem::HandleDragLeave)
-		.OnDrop(this, &SHierarchyViewItem::HandleDrop)
+		.Padding(0.0f)
 		.Content()
 		[
 			SNew(SHorizontalBox)
@@ -842,17 +895,18 @@ void SHierarchyViewItem::Construct(const FArguments& InArgs, const TSharedRef< S
 			.VAlign(VAlign_Center)
 			[
 				SNew(SButton)
-				.ContentPadding(0)
+				.ContentPadding(FMargin(3, 1))
 				.ButtonStyle(FEditorStyle::Get(), "HoverHintOnly")
+				.ForegroundColor(FCoreStyle::Get().GetSlateColor("Foreground"))
 				.OnClicked(this, &SHierarchyViewItem::OnToggleVisibility)
 				.Visibility(Model->CanControlVisibility() ? EVisibility::Visible : EVisibility::Hidden)
 				.ToolTipText(LOCTEXT("WidgetVisibilityButtonToolTip", "Toggle Widget's Editor Visibility"))
 				.HAlign(HAlign_Center)
 				.VAlign(VAlign_Center)
-				.Content()
 				[
-					SNew(SImage)
-					.Image(this, &SHierarchyViewItem::GetVisibilityBrushForWidget)
+					SNew(STextBlock)
+					.Font(FEditorStyle::Get().GetFontStyle("FontAwesome.10"))
+					.Text(this, &SHierarchyViewItem::GetVisibilityBrushForWidget)
 				]
 			]
 		],
@@ -905,29 +959,29 @@ bool SHierarchyViewItem::IsHovered() const
 	return bIsHovered || Model->IsHovered();
 }
 
-FReply SHierarchyViewItem::OnDragOver(const FGeometry& MyGeometry, const FDragDropEvent& DragDropEvent)
+void SHierarchyViewItem::HandleDragEnter(FDragDropEvent const& DragDropEvent)
 {
-	return Model->OnDragOver(MyGeometry, DragDropEvent);
+	Model->HandleDragEnter(DragDropEvent);
 }
 
-FReply SHierarchyViewItem::HandleDragDetected(const FGeometry& MyGeometry, const FPointerEvent& MouseEvent)
+void SHierarchyViewItem::HandleDragLeave(const FDragDropEvent& DragDropEvent)
 {
-	return Model->HandleDragDetected(MyGeometry, MouseEvent);
+	Model->HandleDragLeave(DragDropEvent);
 }
 
-void SHierarchyViewItem::HandleDragEnter(FDragDropEvent const& DragDropEvent)
+TOptional<EItemDropZone> SHierarchyViewItem::HandleCanAcceptDrop(const FDragDropEvent& DragDropEvent, EItemDropZone DropZone, TSharedPtr<FHierarchyModel> TargetItem)
 {
-	Model->HandleDragEnter(DragDropEvent);
+	return Model->HandleCanAcceptDrop(DragDropEvent, DropZone);
 }
 
-void SHierarchyViewItem::HandleDragLeave(const FDragDropEvent& DragDropEvent)
+FReply SHierarchyViewItem::HandleDragDetected(const FGeometry& MyGeometry, const FPointerEvent& MouseEvent)
 {
-	Model->HandleDragLeave(DragDropEvent);
+	return Model->HandleDragDetected(MyGeometry, MouseEvent);
 }
 
-FReply SHierarchyViewItem::HandleDrop(const FDragDropEvent& DragDropEvent)
+FReply SHierarchyViewItem::HandleAcceptDrop(const FDragDropEvent& DragDropEvent, EItemDropZone DropZone, TSharedPtr<FHierarchyModel> TargetItem)
 {
-	return Model->HandleDrop(DragDropEvent);
+	return Model->HandleAcceptDrop(DragDropEvent, DropZone);
 }
 
 FReply SHierarchyViewItem::OnToggleVisibility()
@@ -937,11 +991,11 @@ FReply SHierarchyViewItem::OnToggleVisibility()
 	return FReply::Handled();
 }
 
-const FSlateBrush* SHierarchyViewItem::GetVisibilityBrushForWidget() const
+FText SHierarchyViewItem::GetVisibilityBrushForWidget() const
 {
 	return Model->IsVisible() ?
-		FEditorStyle::GetBrush("WidgetDesigner.WidgetVisible") :
-		FEditorStyle::GetBrush("WidgetDesigner.WidgetHidden");
+		FText::FromString(FString(TEXT("\xf06e")) /*fa-eye*/) :
+		FText::FromString(FString(TEXT("\xf070")) /*fa-eye-slash*/);
 }
 
 #undef LOCTEXT_NAMESPACE
diff --git a/Engine/Source/Editor/UMGEditor/Private/Hierarchy/SHierarchyViewItem.h b/Engine/Source/Editor/UMGEditor/Private/Hierarchy/SHierarchyViewItem.h
--- a/Engine/Source/Editor/UMGEditor/Private/Hierarchy/SHierarchyViewItem.h
+++ b/Engine/Source/Editor/UMGEditor/Private/Hierarchy/SHierarchyViewItem.h
@@ -32,12 +32,12 @@ public:
 
 	virtual FSlateFontInfo GetFont() const = 0;
 
-	virtual FReply OnDragOver(const FGeometry& MyGeometry, const FDragDropEvent& DragDropEvent);
-
 	virtual FReply HandleDragDetected(const FGeometry& MyGeometry, const FPointerEvent& MouseEvent);
 	virtual void HandleDragEnter(const FDragDropEvent& DragDropEvent);
 	virtual void HandleDragLeave(const FDragDropEvent& DragDropEvent);
-	virtual FReply HandleDrop(FDragDropEvent const& DragDropEvent);
+
+	virtual TOptional<EItemDropZone> HandleCanAcceptDrop(const FDragDropEvent& DragDropEvent, EItemDropZone DropZone);
+	virtual FReply HandleAcceptDrop(FDragDropEvent const& DragDropEvent, EItemDropZone DropZone);
 
 	virtual bool OnVerifyNameTextChanged(const FText& InText, FText& OutErrorMessage);
 
@@ -92,11 +92,10 @@ public:
 
 	virtual FSlateFontInfo GetFont() const override;
 
-	virtual void OnSelection();
+	virtual void OnSelection() override;
 
-	virtual FReply OnDragOver(const FGeometry& MyGeometry, const FDragDropEvent& DragDropEvent);
-
-	virtual FReply HandleDrop(FDragDropEvent const& DragDropEvent);
+	virtual TOptional<EItemDropZone> HandleCanAcceptDrop(const FDragDropEvent& DragDropEvent, EItemDropZone DropZone) override;
+	virtual FReply HandleAcceptDrop(FDragDropEvent const& DragDropEvent, EItemDropZone DropZone) override;
 
 protected:
 	virtual void GetChildren(TArray< TSharedPtr<FHierarchyModel> >& Children) override;
@@ -125,8 +124,8 @@ public:
 
 	virtual void OnSelection();
 	
-	virtual FReply OnDragOver(const FGeometry& MyGeometry, const FDragDropEvent& DragDropEvent) override;
-	virtual FReply HandleDrop(FDragDropEvent const& DragDropEvent) override;
+	virtual TOptional<EItemDropZone> HandleCanAcceptDrop(const FDragDropEvent& DragDropEvent, EItemDropZone DropZone) override;
+	virtual FReply HandleAcceptDrop(FDragDropEvent const& DragDropEvent, EItemDropZone DropZone) override;
 	virtual FReply HandleDragDetected(const FGeometry& MyGeometry, const FPointerEvent& MouseEvent) override;
 
 protected:
@@ -159,11 +158,11 @@ public:
 
 	virtual FSlateFontInfo GetFont() const override;
 
-	virtual FReply OnDragOver(const FGeometry& MyGeometry, const FDragDropEvent& DragDropEvent) override;
+	virtual TOptional<EItemDropZone> HandleCanAcceptDrop(const FDragDropEvent& DragDropEvent, EItemDropZone DropZone) override;
 
 	virtual FReply HandleDragDetected(const FGeometry& MyGeometry, const FPointerEvent& MouseEvent) override;
 	virtual void HandleDragLeave(const FDragDropEvent& DragDropEvent) override;
-	virtual FReply HandleDrop(FDragDropEvent const& DragDropEvent) override;
+	virtual FReply HandleAcceptDrop(FDragDropEvent const& DragDropEvent, EItemDropZone DropZone) override;
 
 	virtual bool OnVerifyNameTextChanged(const FText& InText, FText& OutErrorMessage) override;
 
@@ -238,16 +237,17 @@ public:
 
 	// Begin SWidget
 	virtual bool IsHovered() const;
-	void OnMouseEnter(const FGeometry& MyGeometry, const FPointerEvent& MouseEvent);
-	void OnMouseLeave(const FPointerEvent& MouseEvent);
-	virtual FReply OnDragOver(const FGeometry& MyGeometry, const FDragDropEvent& DragDropEvent) override;
+	virtual void OnMouseEnter(const FGeometry& MyGeometry, const FPointerEvent& MouseEvent) override;
+	virtual void OnMouseLeave(const FPointerEvent& MouseEvent) override;
 	// End SWidget
 
 private:
 	FReply HandleDragDetected(const FGeometry& MyGeometry, const FPointerEvent& MouseEvent);
 	void HandleDragEnter(const FDragDropEvent& DragDropEvent);
 	void HandleDragLeave(const FDragDropEvent& DragDropEvent);
-	FReply HandleDrop(FDragDropEvent const& DragDropEvent);
+
+	TOptional<EItemDropZone> HandleCanAcceptDrop(const FDragDropEvent& DragDropEvent, EItemDropZone DropZone, TSharedPtr<FHierarchyModel> TargetItem);
+	FReply HandleAcceptDrop(const FDragDropEvent& DragDropEvent, EItemDropZone DropZone, TSharedPtr<FHierarchyModel> TargetItem);
 
 	/** Called when text is being committed to check for validity */
 	bool OnVerifyNameTextChanged(const FText& InText, FText& OutErrorMessage);
@@ -265,7 +265,7 @@ private:
 	FReply OnToggleVisibility();
 
 	/** Returns a brush representing the visibility item of the widget's visibility button */
-	const FSlateBrush* GetVisibilityBrushForWidget() const;
+	FText GetVisibilityBrushForWidget() const;
 
 	/* The mode that this tree item represents */
 	TSharedPtr<FHierarchyModel> Model;
