UMG - Fix an issue where deleted animations could have name collisions with new animations.

diff --git a/Engine/Source/Editor/UMGEditor/Private/TabFactory/AnimationTabSummoner.cpp b/Engine/Source/Editor/UMGEditor/Private/TabFactory/AnimationTabSummoner.cpp
--- a/Engine/Source/Editor/UMGEditor/Private/TabFactory/AnimationTabSummoner.cpp
+++ b/Engine/Source/Editor/UMGEditor/Private/TabFactory/AnimationTabSummoner.cpp
@@ -464,7 +464,8 @@ private:
 		{
 			const FScopedTransaction Transaction(LOCTEXT("DeleteAnimationTransaction", "Delete Animation"));
 			WidgetBlueprint->Modify();
-
+			// Rename the animation and move it to the transient package to avoid collisions.
+			SelectedAnimation->Animation->Rename( NULL, GetTransientPackage() );
 			WidgetAnimations.Remove(SelectedAnimation->Animation);
 
 			UpdateAnimationList();
