UMG - Changing the add event style to be green.  Updating the +New button on the animation tab to now be + Animation and it's now green.

diff --git a/Engine/Source/Editor/EditorStyle/Private/SlateEditorStyle.cpp b/Engine/Source/Editor/EditorStyle/Private/SlateEditorStyle.cpp
--- a/Engine/Source/Editor/EditorStyle/Private/SlateEditorStyle.cpp
+++ b/Engine/Source/Editor/EditorStyle/Private/SlateEditorStyle.cpp
@@ -207,6 +207,13 @@ void FSlateEditorStyle::FStyle::SetupGeneralStyles()
 		Set("NormalText.Subdued", FTextBlockStyle(NormalText)
 			.SetColorAndOpacity(FSlateColor::UseSubduedForeground()));
 
+		Set("NormalText.Important", FTextBlockStyle(NormalText)
+			.SetFont(TTF_CORE_FONT("Fonts/Roboto-Bold", 10))
+			.SetColorAndOpacity(FLinearColor(1.0f, 1.0f, 1.0f))
+			.SetHighlightColor(FLinearColor(1.0f, 1.0f, 1.0f))
+			.SetShadowOffset(FVector2D(1, 1))
+			.SetShadowColorAndOpacity(FLinearColor(0, 0, 0, 0.9f)));
+
 		Set("SmallText", FTextBlockStyle(NormalText)
 			.SetFont(TTF_CORE_FONT("Fonts/Roboto-Regular", 9)));
 
diff --git a/Engine/Source/Editor/UMGEditor/Private/Customizations/DetailCustomizations.cpp b/Engine/Source/Editor/UMGEditor/Private/Customizations/DetailCustomizations.cpp
--- a/Engine/Source/Editor/UMGEditor/Private/Customizations/DetailCustomizations.cpp
+++ b/Engine/Source/Editor/UMGEditor/Private/Customizations/DetailCustomizations.cpp
@@ -39,7 +39,7 @@ public:
 		ChildSlot
 		[
 			SNew(SButton)
-			.ButtonStyle(FEditorStyle::Get(), "RoundButton")
+			.ButtonStyle(FEditorStyle::Get(), "FlatButton.Success")
 			.TextStyle(FEditorStyle::Get(), "NormalText")
 			.HAlign(HAlign_Center)
 			.ForegroundColor(FSlateColor::UseForeground())
diff --git a/Engine/Source/Editor/UMGEditor/Private/TabFactory/AnimationTabSummoner.cpp b/Engine/Source/Editor/UMGEditor/Private/TabFactory/AnimationTabSummoner.cpp
--- a/Engine/Source/Editor/UMGEditor/Private/TabFactory/AnimationTabSummoner.cpp
+++ b/Engine/Source/Editor/UMGEditor/Private/TabFactory/AnimationTabSummoner.cpp
@@ -208,6 +208,8 @@ public:
 					.AutoWidth()
 					[
 						SNew(SButton)
+						.ButtonStyle(FEditorStyle::Get(), "FlatButton.Success")
+						.ForegroundColor(FEditorStyle::Get().GetSlateColor("Foreground"))
 						.ContentPadding(FMargin(2.0f, 1.0f))
 						.OnClicked( this, &SUMGAnimationList::OnNewAnimationClicked )
 						.HAlign(HAlign_Center)
@@ -218,14 +220,18 @@ public:
 							.VAlign(VAlign_Center)
 							.AutoWidth()
 							[
-								SNew( SImage )
-								.Image( FEditorStyle::GetBrush("UMGEditor.AddAnimationIcon")) 
+								SNew(STextBlock)
+								.TextStyle(FEditorStyle::Get(), "NormalText.Important")
+								.Font(FEditorStyle::Get().GetFontStyle("FontAwesome.10"))
+								.Text(FString(TEXT("\xf067")) /*fa-plus*/)
 							]
+
 							+ SHorizontalBox::Slot()
 							.Padding( 2.0f, 0.0f )
 							[
 								SNew( STextBlock )
-								.Text( LOCTEXT("NewAnimationButtonText", "New") )
+								.TextStyle(FEditorStyle::Get(), "NormalText.Important")
+								.Text( LOCTEXT("NewAnimationButtonText", "Animation") )
 							]
 						]
 					]
