Fixed crash caused by renaming to something >1024 chars.
Removed some duplicated code and fixed some inefficiencies.

diff --git a/Engine/Source/Editor/UMGEditor/Private/Details/SWidgetDetailsView.cpp b/Engine/Source/Editor/UMGEditor/Private/Details/SWidgetDetailsView.cpp
--- a/Engine/Source/Editor/UMGEditor/Private/Details/SWidgetDetailsView.cpp
+++ b/Engine/Source/Editor/UMGEditor/Private/Details/SWidgetDetailsView.cpp
@@ -366,44 +366,13 @@ bool SWidgetDetailsView::HandleVerifyNameTextChanged(const FText& InText, FText&
 {
 	if ( SelectedObjects.Num() == 1 )
 	{
-		FString NewName = InText.ToString();
-
-		if (NewName.IsEmpty())
-		{
-			OutErrorMessage = LOCTEXT("EmptyWidgetName", "Empty Widget Name");
-			return false;
-		}
-
 		UWidget* PreviewWidget = Cast<UWidget>(SelectedObjects[0].Get());
+		
+		TSharedRef<class FWidgetBlueprintEditor> BlueprintEditorRef = BlueprintEditor.Pin().ToSharedRef();
 
-		UWidgetBlueprint* Blueprint = BlueprintEditor.Pin()->GetWidgetBlueprintObj();
-		UWidget* TemplateWidget = Blueprint->WidgetTree->FindWidget( FName(*NewName) );
-
-		bool bIsSameWidget = false;
-		if ( TemplateWidget != NULL )
-		{
-			if ( BlueprintEditor.Pin()->GetReferenceFromTemplate( TemplateWidget ).GetPreview() != PreviewWidget )
-			{
-				OutErrorMessage = LOCTEXT("ExistingWidgetName", "Existing Widget Name");
-				return false;
-			}
-			else
-			{
-				bIsSameWidget = true;
-			}
-		}
-
-		FKismetNameValidator Validator(Blueprint);
-
-		const bool bUniqueNameForVariable = ( EValidatorResult::Ok == Validator.IsValid(NewName) );
-
-		if ( !bUniqueNameForVariable && !bIsSameWidget )
-		{
-			OutErrorMessage = LOCTEXT("ExistingVariableName", "Existing Variable Name");
-			return false;
-		}
+		FWidgetReference WidgetRef = BlueprintEditorRef->GetReferenceFromPreview(PreviewWidget);
 
-		return true;
+		return FWidgetBlueprintEditorUtils::VerifyWidgetRename(BlueprintEditorRef, WidgetRef, InText, OutErrorMessage);
 	}
 	else
 	{
diff --git a/Engine/Source/Editor/UMGEditor/Private/Hierarchy/SHierarchyViewItem.cpp b/Engine/Source/Editor/UMGEditor/Private/Hierarchy/SHierarchyViewItem.cpp
--- a/Engine/Source/Editor/UMGEditor/Private/Hierarchy/SHierarchyViewItem.cpp
+++ b/Engine/Source/Editor/UMGEditor/Private/Hierarchy/SHierarchyViewItem.cpp
@@ -741,42 +741,7 @@ FReply FHierarchyWidget::HandleAcceptDrop(const FDragDropEvent& DragDropEvent, E
 
 bool FHierarchyWidget::OnVerifyNameTextChanged(const FText& InText, FText& OutErrorMessage)
 {
-	FString NewName = InText.ToString();
-
-	if (NewName.IsEmpty())
-	{
-		OutErrorMessage = LOCTEXT("EmptyWidgetName", "Empty Widget Name");
-		return false;
-	}
-
-	UWidgetBlueprint* Blueprint = BlueprintEditor.Pin()->GetWidgetBlueprintObj();
-	UWidget* ExistingTemplate = Blueprint->WidgetTree->FindWidget( FName(*NewName) );
-
-	bool bIsSameWidget = false;
-	if ( ExistingTemplate != nullptr )
-	{
-		if ( Item.GetTemplate() != ExistingTemplate )
-		{
-			OutErrorMessage = LOCTEXT("ExistingWidgetName", "Existing Widget Name");
-			return false;
-		}
-		else
-		{
-			bIsSameWidget = true;
-		}
-	}
-
-	FKismetNameValidator Validator(Blueprint);
-
-	const bool bUniqueNameForVariable = ( EValidatorResult::Ok == Validator.IsValid(NewName) );
-
-	if ( !bUniqueNameForVariable && !bIsSameWidget )
-	{
-		OutErrorMessage = LOCTEXT("ExistingVariableName", "Existing Variable Name");
-		return false;
-	}
-
-	return true;
+	return FWidgetBlueprintEditorUtils::VerifyWidgetRename(BlueprintEditor.Pin().ToSharedRef(), Item, InText, OutErrorMessage);
 }
 
 void FHierarchyWidget::OnNameTextCommited(const FText& InText, ETextCommit::Type CommitInfo)
diff --git a/Engine/Source/Editor/UMGEditor/Private/WidgetBlueprintEditorUtils.cpp b/Engine/Source/Editor/UMGEditor/Private/WidgetBlueprintEditorUtils.cpp
--- a/Engine/Source/Editor/UMGEditor/Private/WidgetBlueprintEditorUtils.cpp
+++ b/Engine/Source/Editor/UMGEditor/Private/WidgetBlueprintEditorUtils.cpp
@@ -67,6 +67,54 @@ public:
 	TMap<FName, UWidgetSlotPair*> MissingSlotData;
 };
 
+bool FWidgetBlueprintEditorUtils::VerifyWidgetRename(TSharedRef<class FWidgetBlueprintEditor> BlueprintEditor, FWidgetReference Widget, const FText& NewName, FText& OutErrorMessage)
+{
+	if (NewName.IsEmptyOrWhitespace())
+	{
+		OutErrorMessage = LOCTEXT("EmptyWidgetName", "Empty Widget Name");
+		return false;
+	}
+
+	const FString& NewNameString = NewName.ToString();
+
+	if (NewNameString.Len() >= NAME_SIZE)
+	{
+		OutErrorMessage = LOCTEXT("WidgetNameTooLong", "Widget Name is Too Long");
+		return false;
+	}
+
+	const FName NewFName(*NewNameString);
+	
+	UWidgetBlueprint* Blueprint = BlueprintEditor->GetWidgetBlueprintObj();
+	UWidget* ExistingTemplate = Blueprint->WidgetTree->FindWidget(NewFName);
+
+	bool bIsSameWidget = false;
+	if (ExistingTemplate != nullptr)
+	{
+		if (Widget.GetTemplate() != ExistingTemplate)
+		{
+			OutErrorMessage = LOCTEXT("ExistingWidgetName", "Existing Widget Name");
+			return false;
+		}
+		else
+		{
+			bIsSameWidget = true;
+		}
+	}
+
+	FKismetNameValidator Validator(Blueprint);
+
+	const bool bUniqueNameForVariable = (EValidatorResult::Ok == Validator.IsValid(NewFName));
+
+	if (!bUniqueNameForVariable && !bIsSameWidget)
+	{
+		OutErrorMessage = LOCTEXT("ExistingVariableName", "Existing Variable Name");
+		return false;
+	}
+
+	return true;
+}
+
 bool FWidgetBlueprintEditorUtils::RenameWidget(TSharedRef<FWidgetBlueprintEditor> BlueprintEditor, const FName& OldName, const FName& NewName)
 {
 	UWidgetBlueprint* Blueprint = BlueprintEditor->GetWidgetBlueprintObj();
diff --git a/Engine/Source/Editor/UMGEditor/Private/WidgetBlueprintEditorUtils.h b/Engine/Source/Editor/UMGEditor/Private/WidgetBlueprintEditorUtils.h
--- a/Engine/Source/Editor/UMGEditor/Private/WidgetBlueprintEditorUtils.h
+++ b/Engine/Source/Editor/UMGEditor/Private/WidgetBlueprintEditorUtils.h
@@ -13,6 +13,8 @@ struct FWidgetReference;
 class FWidgetBlueprintEditorUtils
 {
 public:
+	static bool VerifyWidgetRename(TSharedRef<class FWidgetBlueprintEditor> BlueprintEditor, FWidgetReference Widget, const FText& NewName, FText& OutErrorMessage);
+
 	static bool RenameWidget(TSharedRef<class FWidgetBlueprintEditor> BlueprintEditor, const FName& OldName, const FName& NewName);
 
 	static void CreateWidgetContextMenu(FMenuBuilder& MenuBuilder, TSharedRef<FWidgetBlueprintEditor> BlueprintEditor, FVector2D TargetLocation);
