UMG - The SRuler now has a check for a 0 pixels per input, to prevent an infinite loop.

diff --git a/Engine/Source/Editor/UMGEditor/Private/Designer/SRuler.cpp b/Engine/Source/Editor/UMGEditor/Private/Designer/SRuler.cpp
--- a/Engine/Source/Editor/UMGEditor/Private/Designer/SRuler.cpp
+++ b/Engine/Source/Editor/UMGEditor/Private/Designer/SRuler.cpp
@@ -103,6 +103,9 @@ static float GetNextSpacing( uint32 CurrentStep )
  */
 float DetermineOptimalSpacing( float InPixelsPerInput, uint32 MinTick, float MinTickSpacing )
 {
+	if (InPixelsPerInput == 0.0f)
+		return MinTickSpacing;
+	
 	uint32 CurStep = 0;
 
 	// Start with the smallest spacing
