UMG - Fixing a bug that was preventing GoTo function on function bindings from working for the newer style binding system.

diff --git a/Engine/Source/Editor/UMGEditor/Private/Details/SPropertyBinding.cpp b/Engine/Source/Editor/UMGEditor/Private/Details/SPropertyBinding.cpp
--- a/Engine/Source/Editor/UMGEditor/Private/Details/SPropertyBinding.cpp
+++ b/Engine/Source/Editor/UMGEditor/Private/Details/SPropertyBinding.cpp
@@ -730,15 +730,20 @@ FReply SPropertyBinding::HandleGotoBindingClicked(TSharedRef<IPropertyHandle> Pr
 					TArray<UEdGraph*> AllGraphs;
 					Blueprint->GetAllGraphs(AllGraphs);
 
+					FGuid SearchForGuid = Binding.MemberGuid;
+					if ( !Binding.SourcePath.IsEmpty() )
+					{
+						SearchForGuid = Binding.SourcePath.Segments.Last().GetMemberGuid();
+					}
+
 					for ( UEdGraph* Graph : AllGraphs )
 					{
-						if ( Graph->GraphGuid == Binding.MemberGuid )
+						if ( Graph->GraphGuid == SearchForGuid )
 						{
 							GotoFunction(Graph);
 						}
 					}
 
-
 					// Either way return
 					return FReply::Handled();
 				}
diff --git a/Engine/Source/Editor/UMGEditor/Private/WidgetBlueprint.cpp b/Engine/Source/Editor/UMGEditor/Private/WidgetBlueprint.cpp
--- a/Engine/Source/Editor/UMGEditor/Private/WidgetBlueprint.cpp
+++ b/Engine/Source/Editor/UMGEditor/Private/WidgetBlueprint.cpp
@@ -266,6 +266,11 @@ FText FEditorPropertyPathSegment::GetMemberDisplayText() const
 	return FText::FromName(MemberName);
 }
 
+FGuid FEditorPropertyPathSegment::GetMemberGuid() const
+{
+	return MemberGuid;
+}
+
 FEditorPropertyPath::FEditorPropertyPath()
 {
 }
diff --git a/Engine/Source/Editor/UMGEditor/Public/WidgetBlueprint.h b/Engine/Source/Editor/UMGEditor/Public/WidgetBlueprint.h
--- a/Engine/Source/Editor/UMGEditor/Public/WidgetBlueprint.h
+++ b/Engine/Source/Editor/UMGEditor/Public/WidgetBlueprint.h
@@ -31,6 +31,7 @@ public:
 
 	FName GetMemberName() const;
 	FText GetMemberDisplayText() const;
+	FGuid GetMemberGuid() const;
 
 private:
 
