UMG - Make sequencer animation selection more consistent.
+ Update the preview widgets whenever the selected animation changes.
+ Start with no animation selected so that the initial widget states are shown.
+ Clear the sequencer when no animation is selected in the animation tab.

Jira: UE-12583

diff --git a/Engine/Source/Editor/Sequencer/Private/Sequencer.cpp b/Engine/Source/Editor/Sequencer/Private/Sequencer.cpp
--- a/Engine/Source/Editor/Sequencer/Private/Sequencer.cpp
+++ b/Engine/Source/Editor/Sequencer/Private/Sequencer.cpp
@@ -689,6 +689,7 @@ void FSequencer::UpdateRuntimeInstances()
 
 	// Refresh the current root instance
 	RootMovieSceneInstance->RefreshInstance( *this );
+	RootMovieSceneInstance->Update(ScrubPosition, ScrubPosition, *this);
 }
 
 void FSequencer::DestroySpawnablesForAllMovieScenes()
diff --git a/Engine/Source/Editor/UMGEditor/Private/TabFactory/AnimationTabSummoner.cpp b/Engine/Source/Editor/UMGEditor/Private/TabFactory/AnimationTabSummoner.cpp
--- a/Engine/Source/Editor/UMGEditor/Private/TabFactory/AnimationTabSummoner.cpp
+++ b/Engine/Source/Editor/UMGEditor/Private/TabFactory/AnimationTabSummoner.cpp
@@ -369,11 +369,17 @@ private:
 
 	void OnSelectionChanged( TSharedPtr<FWidgetAnimationListItem> InSelectedItem, ESelectInfo::Type SelectionInfo )
 	{
+		UWidgetBlueprint* WidgetBlueprint = BlueprintEditor.Pin()->GetWidgetBlueprintObj();
+		UWidgetAnimation* WidgetAnimation;
 		if( InSelectedItem.IsValid() )
 		{
-			UWidgetBlueprint* WidgetBlueprint = BlueprintEditor.Pin()->GetWidgetBlueprintObj();
-			BlueprintEditor.Pin()->ChangeViewedAnimation( *InSelectedItem->Animation );
+			WidgetAnimation = InSelectedItem->Animation;
 		}
+		else
+		{
+			WidgetAnimation = UWidgetAnimation::GetNullAnimation();
+		}
+		BlueprintEditor.Pin()->ChangeViewedAnimation( *WidgetAnimation );
 	}
 
 	TSharedPtr<SWidget> OnContextMenuOpening() const
@@ -479,16 +485,7 @@ private:
 
 		FBlueprintEditorUtils::MarkBlueprintAsStructurallyModified(WidgetBlueprint);
 
-		if(Animations.Num() > 0)
-		{
-			// Ensure we are viewing a valid animation
-			WidgetBlueprintEditorPin->ChangeViewedAnimation(*Animations[0]->Animation);
-			AnimationListView->SetSelection( Animations[0] );
-		}
-		else
-		{
-			WidgetBlueprintEditorPin->ChangeViewedAnimation(*UWidgetAnimation::GetNullAnimation());
-		}
+		WidgetBlueprintEditorPin->ChangeViewedAnimation(*UWidgetAnimation::GetNullAnimation());
 	}
 
 	void OnRenameAnimation()
diff --git a/Engine/Source/Editor/UMGEditor/Private/WidgetBlueprintEditor.cpp b/Engine/Source/Editor/UMGEditor/Private/WidgetBlueprintEditor.cpp
--- a/Engine/Source/Editor/UMGEditor/Private/WidgetBlueprintEditor.cpp
+++ b/Engine/Source/Editor/UMGEditor/Private/WidgetBlueprintEditor.cpp
@@ -299,18 +299,8 @@ const UWidgetAnimation* FWidgetBlueprintEditor::RefreshCurrentAnimation()
 {
 	if( !SequencerObjectBindingManager->HasValidWidgetAnimation() )
 	{
-		const TArray<UWidgetAnimation*>& Animations = GetWidgetBlueprintObj()->Animations;
-		if( Animations.Num() > 0 )
-		{
-			// Ensure we are viewing a valid animation
-			ChangeViewedAnimation( *Animations[0] );
-			return Animations[0];
-		}
-		else
-		{
-			ChangeViewedAnimation( *UWidgetAnimation::GetNullAnimation() );
-			return nullptr;
-		}
+		ChangeViewedAnimation(*UWidgetAnimation::GetNullAnimation());
+		return nullptr;
 	}
 
 	return SequencerObjectBindingManager->GetWidgetAnimation();
@@ -466,15 +456,13 @@ TSharedRef<SWidget> FWidgetBlueprintEditor::CreateSequencerWidget()
 {
 	TSharedRef<SOverlay> SequencerOverlayRef =
 		SNew(SOverlay)
-		.AddMetaData<FTagMetaData>(FTagMetaData(TEXT("Sequencer")))
-		+ SOverlay::Slot()
-		[
-			GetSequencer()->GetSequencerWidget()
-		];
-
+		.AddMetaData<FTagMetaData>(FTagMetaData(TEXT("Sequencer")));
 	SequencerOverlay = SequencerOverlayRef;
 
-	RefreshCurrentAnimation();
+	SequencerOverlayRef->AddSlot(0)
+	[
+		GetSequencer()->GetSequencerWidget()
+	];
 
 	return SequencerOverlayRef;
 }
@@ -538,32 +526,14 @@ TSharedPtr<ISequencer>& FWidgetBlueprintEditor::GetSequencer()
 {
 	if(!Sequencer.IsValid())
 	{
-		UWidgetBlueprint* Blueprint = GetWidgetBlueprintObj();
-
-		UWidgetAnimation* WidgetAnimation = nullptr;
-
-		if( Blueprint->Animations.Num() )
-		{
-			WidgetAnimation = Blueprint->Animations[0];
-		}
-		else
-		{
-			UWidgetAnimation* NewAnimation = NewObject<UWidgetAnimation>(Blueprint, MakeUniqueObjectName(Blueprint, UWidgetAnimation::StaticClass(), "NewAnimation"), RF_Transactional);
-			NewAnimation->MovieScene = NewObject<UMovieScene>(NewAnimation, NAME_None, RF_Transactional);
-
-			WidgetAnimation = NewAnimation;
-		}
-
-		TSharedRef<FUMGSequencerObjectBindingManager> ObjectBindingManager = MakeShareable(new FUMGSequencerObjectBindingManager(*this, *WidgetAnimation ) );
-		check( !SequencerObjectBindingManager.IsValid() );
-
-		SequencerObjectBindingManager = ObjectBindingManager;
-
 		FSequencerViewParams ViewParams;
 		ViewParams.InitalViewRange = TRange<float>(-0.02f, 3.2f);
 		ViewParams.InitialScrubPosition = 0;
 
-		Sequencer = FModuleManager::LoadModuleChecked< ISequencerModule >("Sequencer").CreateSequencer( WidgetAnimation->MovieScene, ViewParams, ObjectBindingManager );
+		SequencerObjectBindingManager = MakeShareable(new FUMGSequencerObjectBindingManager(*this, *UWidgetAnimation::GetNullAnimation()));
+
+		Sequencer = FModuleManager::LoadModuleChecked< ISequencerModule >("Sequencer").CreateSequencer(UWidgetAnimation::GetNullAnimation()->MovieScene, ViewParams, SequencerObjectBindingManager.ToSharedRef());
+		ChangeViewedAnimation(*UWidgetAnimation::GetNullAnimation());
 	}
 
 	return Sequencer;
@@ -571,43 +541,39 @@ TSharedPtr<ISequencer>& FWidgetBlueprintEditor::GetSequencer()
 
 void FWidgetBlueprintEditor::ChangeViewedAnimation( UWidgetAnimation& InAnimationToView )
 {
-	if( InAnimationToView.MovieScene != Sequencer->GetRootMovieScene() )
-	{
-		TSharedRef<FUMGSequencerObjectBindingManager> NewObjectBindingManager = MakeShareable(new FUMGSequencerObjectBindingManager(*this, InAnimationToView));
+	TSharedRef<FUMGSequencerObjectBindingManager> NewObjectBindingManager = MakeShareable(new FUMGSequencerObjectBindingManager(*this, InAnimationToView));
+	Sequencer->ResetToNewRootMovieScene(*InAnimationToView.MovieScene, NewObjectBindingManager);
 
-		Sequencer->ResetToNewRootMovieScene(*InAnimationToView.MovieScene, NewObjectBindingManager);
+	check(SequencerObjectBindingManager.IsUnique());
+	SequencerObjectBindingManager = NewObjectBindingManager;
+	SequencerObjectBindingManager->InitPreviewObjects();
 
-		check(SequencerObjectBindingManager.IsUnique());
-
-		SequencerObjectBindingManager = NewObjectBindingManager;
-
-		SequencerObjectBindingManager->InitPreviewObjects();
+	TSharedPtr<SOverlay> SequencerOverlayPin = SequencerOverlay.Pin();
+	if( &InAnimationToView == UWidgetAnimation::GetNullAnimation() && SequencerOverlayPin.IsValid() )
+	{
+		Sequencer->GetSequencerWidget()->SetEnabled(false);
+		// Disable sequencer from interaction
+		SequencerOverlayPin->AddSlot(1)
+		.HAlign(HAlign_Center)
+		.VAlign(VAlign_Center)
+		[
+			SNew( STextBlock )
+			.TextStyle( FEditorStyle::Get(), "UMGEditor.NoAnimationFont" )
+			.Text( LOCTEXT("NoAnimationSelected","No Animation Selected") )
+		];
 
-		TSharedPtr<SOverlay> SequencerOverlayPin = SequencerOverlay.Pin();
-		if( &InAnimationToView == UWidgetAnimation::GetNullAnimation() && SequencerOverlayPin.IsValid() )
-		{
-			Sequencer->GetSequencerWidget()->SetEnabled(false);
-			// Disable sequencer from interaction
-			SequencerOverlayPin->AddSlot(1)
-			.HAlign(HAlign_Center)
-			.VAlign(VAlign_Center)
-			[
-				SNew( STextBlock )
-				.TextStyle( FEditorStyle::Get(), "UMGEditor.NoAnimationFont" )
-				.Text( LOCTEXT("NoAnimationSelected","No Animation Selected") )
-			];
-
-			SequencerOverlayPin->SetVisibility( EVisibility::HitTestInvisible );
-		}
-		else if( SequencerOverlayPin.IsValid() && SequencerOverlayPin->GetNumWidgets() > 1 )
-		{
-			Sequencer->GetSequencerWidget()->SetEnabled(true);
+		SequencerOverlayPin->SetVisibility( EVisibility::HitTestInvisible );
+	}
+	else if( SequencerOverlayPin.IsValid() && SequencerOverlayPin->GetNumWidgets() > 1 )
+	{
+		Sequencer->GetSequencerWidget()->SetEnabled(true);
 
-			SequencerOverlayPin->RemoveSlot(1);
-			// Allow sequencer to be interacted with
-			SequencerOverlayPin->SetVisibility( EVisibility::SelfHitTestInvisible );
-		}
+		SequencerOverlayPin->RemoveSlot(1);
+		// Allow sequencer to be interacted with
+		SequencerOverlayPin->SetVisibility( EVisibility::SelfHitTestInvisible );
 	}
+
+	InvalidatePreview();
 }
 
 void FWidgetBlueprintEditor::RefreshPreview()
