UMG - Adding new settings to the editor for the widget designer.  Allows users to configure a few options for the widget designer.  Options are under Editor Preferences > Content Editors > Widget Designer

diff --git a/Engine/Source/Editor/UMGEditor/Private/Designer/SDesignerView.cpp b/Engine/Source/Editor/UMGEditor/Private/Designer/SDesignerView.cpp
--- a/Engine/Source/Editor/UMGEditor/Private/Designer/SDesignerView.cpp
+++ b/Engine/Source/Editor/UMGEditor/Private/Designer/SDesignerView.cpp
@@ -87,7 +87,7 @@ TSharedRef<FSelectedWidgetDragDropOp> FSelectedWidgetDragDropOp::New(TSharedPtr<
 	bool bStayInParent = false;
 	if ( UPanelWidget* PanelTemplate = InWidget.GetTemplate()->GetParent() )
 	{
-		bStayInParent = PanelTemplate->LockToPanelOnDrag();
+		bStayInParent = PanelTemplate->LockToPanelOnDrag() || GetDefault<UWidgetDesignerSettings>()->bLockToPanelOnDragByDefault;
 	}
 
 	TSharedRef<FSelectedWidgetDragDropOp> Operation = MakeShareable(new FSelectedWidgetDragDropOp());
diff --git a/Engine/Source/Editor/UMGEditor/Private/UMGEditorModule.cpp b/Engine/Source/Editor/UMGEditor/Private/UMGEditorModule.cpp
--- a/Engine/Source/Editor/UMGEditor/Private/UMGEditorModule.cpp
+++ b/Engine/Source/Editor/UMGEditor/Private/UMGEditorModule.cpp
@@ -19,6 +19,11 @@
 #include "DesignerCommands.h"
 #include "WidgetBlueprint.h"
 
+#include "Settings/WidgetDesignerSettings.h"
+#include "ISettingsModule.h"
+
+#define LOCTEXT_NAMESPACE "UMG"
+
 const FName UMGEditorAppIdentifier = FName(TEXT("UMGEditorApp"));
 
 class FUMGEditorModule : public IUMGEditorModule, public IBlueprintCompiler
@@ -56,6 +61,17 @@ public:
 		ISequencerModule& SequencerModule = FModuleManager::Get().LoadModuleChecked<ISequencerModule>("Sequencer");
 		MarginTrackEditorCreateTrackEditorHandle    = SequencerModule.RegisterTrackEditor_Handle(FOnCreateTrackEditor::CreateStatic(&FMarginTrackEditor::CreateTrackEditor));
 		TransformTrackEditorCreateTrackEditorHandle = SequencerModule.RegisterTrackEditor_Handle(FOnCreateTrackEditor::CreateStatic(&F2DTransformTrackEditor::CreateTrackEditor));
+
+		ISettingsModule* SettingsModule = FModuleManager::GetModulePtr<ISettingsModule>("Settings");
+		if ( SettingsModule != nullptr )
+		{
+			// Designer settings
+			SettingsModule->RegisterSettings("Editor", "ContentEditors", "WidgetDesigner",
+				LOCTEXT("WidgetDesignerSettingsName", "Widget Designer"),
+				LOCTEXT("WidgetDesignerSettingsDescription", "Configure options for the Widget Designer."),
+				GetMutableDefault<UWidgetDesignerSettings>()
+				);
+		}
 	}
 
 	/** Called before the module is unloaded, right before the module object is destroyed. */
@@ -74,6 +90,14 @@ public:
 			}
 		}
 		CreatedAssetTypeActions.Empty();
+
+		// Unregister the setting
+		ISettingsModule* SettingsModule = FModuleManager::GetModulePtr<ISettingsModule>("Settings");
+
+		if ( SettingsModule != nullptr )
+		{
+			SettingsModule->UnregisterSettings("Editor", "ContentEditors", "WidgetDesigner");
+		}
 	}
 
 	bool CanCompile(const UBlueprint* Blueprint)
@@ -149,3 +173,5 @@ private:
 };
 
 IMPLEMENT_MODULE(FUMGEditorModule, UMGEditor);
+
+#undef LOCTEXT_NAMESPACE
diff --git a/Engine/Source/Editor/UMGEditor/Public/Settings/WidgetDesignerSettings.h b/Engine/Source/Editor/UMGEditor/Public/Settings/WidgetDesignerSettings.h
--- a/Engine/Source/Editor/UMGEditor/Public/Settings/WidgetDesignerSettings.h
+++ b/Engine/Source/Editor/UMGEditor/Public/Settings/WidgetDesignerSettings.h
@@ -8,8 +8,7 @@
  * Implements the settings for the Widget Blueprint Designer.
  */
 UCLASS(config=EditorUserSettings)
-class UMGEDITOR_API UWidgetDesignerSettings
-	: public UObject
+class UMGEDITOR_API UWidgetDesignerSettings : public UObject
 {
 	GENERATED_UCLASS_BODY()
 
@@ -19,4 +18,7 @@ class UMGEDITOR_API UWidgetDesignerSettings
 
 	UPROPERTY(config)
 	int32 GridSnapSize;
+
+	UPROPERTY(EditAnywhere, config, Category=Dragging)
+	bool bLockToPanelOnDragByDefault;
 };
