UMG - Adding handlers to the design surface for touch gestures so that it moves more like what's expected on Mac.

diff --git a/Engine/Source/Editor/UMGEditor/Private/Designer/SDesignSurface.cpp b/Engine/Source/Editor/UMGEditor/Private/Designer/SDesignSurface.cpp
--- a/Engine/Source/Editor/UMGEditor/Private/Designer/SDesignSurface.cpp
+++ b/Engine/Source/Editor/UMGEditor/Private/Designer/SDesignSurface.cpp
@@ -135,6 +135,7 @@ void SDesignSurface::Construct(const FArguments& InArgs)
 	ZoomTargetBottomRight = FVector2D::ZeroVector;
 
 	ZoomToFitPadding = FVector2D(100, 100);
+	TotalGestureMagnify = 0.0f;
 
 	ChildSlot
 	[
@@ -250,6 +251,38 @@ FReply SDesignSurface::OnMouseWheel(const FGeometry& MyGeometry, const FPointerE
 	return FReply::Handled();
 }
 
+FReply SDesignSurface::OnTouchGesture(const FGeometry& MyGeometry, const FPointerEvent& GestureEvent)
+{
+	const EGestureEvent::Type GestureType = GestureEvent.GetGestureType();
+	const FVector2D& GestureDelta = GestureEvent.GetGestureDelta();
+	if ( GestureType == EGestureEvent::Magnify )
+	{
+		TotalGestureMagnify += GestureDelta.X;
+		if ( FMath::Abs(TotalGestureMagnify) > 0.07f )
+		{
+			// We want to zoom into this point; i.e. keep it the same fraction offset into the panel
+			const FVector2D WidgetSpaceCursorPos = MyGeometry.AbsoluteToLocal(GestureEvent.GetScreenSpacePosition());
+			const int32 ZoomLevelDelta = TotalGestureMagnify > 0.0f ? 1 : -1;
+			ChangeZoomLevel(ZoomLevelDelta, WidgetSpaceCursorPos, GestureEvent.IsControlDown());
+			TotalGestureMagnify = 0.0f;
+		}
+		return FReply::Handled();
+	}
+	else if ( GestureType == EGestureEvent::Scroll )
+	{
+		this->bIsPanning = true;
+		ViewOffset -= GestureDelta / GetZoomAmount();
+		return FReply::Handled();
+	}
+	return FReply::Unhandled();
+}
+
+FReply SDesignSurface::OnTouchEnded(const FGeometry& MyGeometry, const FPointerEvent& InTouchEvent)
+{
+	TotalGestureMagnify = 0.0f;
+	return FReply::Unhandled();
+}
+
 inline float FancyMod(float Value, float Size)
 {
 	return ( ( Value >= 0 ) ? 0.0f : Size ) + FMath::Fmod(Value, Size);
diff --git a/Engine/Source/Editor/UMGEditor/Private/Designer/SDesignSurface.h b/Engine/Source/Editor/UMGEditor/Private/Designer/SDesignSurface.h
--- a/Engine/Source/Editor/UMGEditor/Private/Designer/SDesignSurface.h
+++ b/Engine/Source/Editor/UMGEditor/Private/Designer/SDesignSurface.h
@@ -27,6 +27,8 @@ public:
 	virtual FReply OnMouseButtonUp(const FGeometry& MyGeometry, const FPointerEvent& MouseEvent) override;
 	virtual FReply OnMouseMove(const FGeometry& MyGeometry, const FPointerEvent& MouseEvent) override;
 	virtual FReply OnMouseWheel(const FGeometry& MyGeometry, const FPointerEvent& MouseEvent) override;
+	virtual FReply OnTouchGesture(const FGeometry& MyGeometry, const FPointerEvent& GestureEvent) override;
+	virtual FReply OnTouchEnded(const FGeometry& MyGeometry, const FPointerEvent& InTouchEvent) override;
 	virtual int32 OnPaint(const FPaintArgs& Args, const FGeometry& AllottedGeometry, const FSlateRect& MyClippingRect, FSlateWindowElementList& OutDrawElements, int32 LayerId, const FWidgetStyle& InWidgetStyle, bool bParentEnabled) const override;
 	virtual bool SupportsKeyboardFocus() const override { return true; }
 	// End of Swidget interface
@@ -96,6 +98,9 @@ protected:
 	FVector2D ZoomTargetBottomRight;
 	FVector2D ZoomToFitPadding;
 
+	/** Cumulative magnify delta from trackpad gesture */
+	float TotalGestureMagnify;
+
 	/** Does the user need to press Control in order to over-zoom. */
 	bool bRequireControlToOverZoom;
 
