UMG - Using "Wrap With" in the designer now properly wraps multiple selected widgets if they are siblings, otherwise groups all possible groups of selected siblings.

diff --git a/Engine/Source/Editor/UMGEditor/Private/Hierarchy/SHierarchyViewItem.cpp b/Engine/Source/Editor/UMGEditor/Private/Hierarchy/SHierarchyViewItem.cpp
--- a/Engine/Source/Editor/UMGEditor/Private/Hierarchy/SHierarchyViewItem.cpp
+++ b/Engine/Source/Editor/UMGEditor/Private/Hierarchy/SHierarchyViewItem.cpp
@@ -205,7 +205,7 @@ TOptional<EItemDropZone> ProcessHierarchyDragDrop(const FDragDropEvent& DragDrop
 				return TOptional<EItemDropZone>();
 			}
 
-			if ( !NewParent->CanHaveMultipleChildren() && NewParent->GetChildrenCount() > 0 )
+			if ( !NewParent->CanAddMoreChildren() )
 			{
 				HierarchyDragDropOp->SetCursorOverride(EMouseCursor::SlashedCircle);
 				HierarchyDragDropOp->CurrentHoverText = LOCTEXT("NoAdditionalChildren", "Widget can't accept additional children.");
diff --git a/Engine/Source/Editor/UMGEditor/Private/WidgetBlueprintEditorUtils.cpp b/Engine/Source/Editor/UMGEditor/Private/WidgetBlueprintEditorUtils.cpp
--- a/Engine/Source/Editor/UMGEditor/Private/WidgetBlueprintEditorUtils.cpp
+++ b/Engine/Source/Editor/UMGEditor/Private/WidgetBlueprintEditorUtils.cpp
@@ -363,17 +363,16 @@ void FWidgetBlueprintEditorUtils::BuildWrapWithMenu(FMenuBuilder& Menu, UWidgetB
 			UClass* WidgetClass = *ClassIt;
 			if ( FWidgetBlueprintEditorUtils::IsUsableWidgetClass(WidgetClass) )
 			{
-				if ( WidgetClass->IsChildOf(UPanelWidget::StaticClass()) && WidgetClass->HasAnyClassFlags(CLASS_Abstract) == false )
+				if ( WidgetClass->IsChildOf(UPanelWidget::StaticClass()) )
 				{
 					Menu.AddMenuEntry(
 						WidgetClass->GetDisplayNameText(),
 						FText::GetEmpty(),
 						FSlateIcon(),
 						FUIAction(
-						FExecuteAction::CreateStatic(&FWidgetBlueprintEditorUtils::WrapWidgets, BP, Widgets, WidgetClass),
-						FCanExecuteAction()
-						)
-						);
+							FExecuteAction::CreateStatic(&FWidgetBlueprintEditorUtils::WrapWidgets, BP, Widgets, WidgetClass),
+							FCanExecuteAction()
+						));
 				}
 			}
 		}
@@ -387,22 +386,35 @@ void FWidgetBlueprintEditorUtils::WrapWidgets(UWidgetBlueprint* BP, TSet<FWidget
 
 	TSharedPtr<FWidgetTemplateClass> Template = MakeShareable(new FWidgetTemplateClass(WidgetClass));
 
-	UPanelWidget* NewWrapperWidget = CastChecked<UPanelWidget>(Template->Create(BP->WidgetTree));
+	// Old Parent -> New Parent Map
+	TMap<UPanelWidget*, UPanelWidget*> OldParentToNewParent;
 
 	for ( FWidgetReference& Item : Widgets )
 	{
 		int32 OutIndex;
 		UPanelWidget* CurrentParent = BP->WidgetTree->FindWidgetParent(Item.GetTemplate(), OutIndex);
-		if ( CurrentParent )
+
+		// If the widget doesn't currently have a parent, and isn't the root, ignore it.
+		if ( CurrentParent == nullptr && Item.GetTemplate() != BP->WidgetTree->RootWidget )
 		{
-			CurrentParent->Modify();
-			CurrentParent->ReplaceChildAt(OutIndex, NewWrapperWidget);
+			continue;
 		}
-		else if ( Item.GetTemplate() == BP->WidgetTree->RootWidget )
+
+		UPanelWidget*& NewWrapperWidget = OldParentToNewParent.FindOrAdd(CurrentParent);
+		if ( NewWrapperWidget == nullptr || !NewWrapperWidget->CanAddMoreChildren() )
 		{
-			BP->WidgetTree->Modify();
+			NewWrapperWidget = CastChecked<UPanelWidget>(Template->Create(BP->WidgetTree));
 
-			BP->WidgetTree->RootWidget = NewWrapperWidget;
+			if ( CurrentParent )
+			{
+				CurrentParent->Modify();
+				CurrentParent->ReplaceChildAt(OutIndex, NewWrapperWidget);
+			}
+			else // Root Widget
+			{
+				BP->WidgetTree->Modify();
+				BP->WidgetTree->RootWidget = NewWrapperWidget;
+			}
 		}
 
 		NewWrapperWidget->AddChild(Item.GetTemplate());
@@ -420,7 +432,7 @@ void FWidgetBlueprintEditorUtils::BuildReplaceWithMenu(FMenuBuilder& Menu, UWidg
 			UClass* WidgetClass = *ClassIt;
 			if ( FWidgetBlueprintEditorUtils::IsUsableWidgetClass(WidgetClass) )
 			{
-				if ( WidgetClass->IsChildOf(UPanelWidget::StaticClass()) && WidgetClass->HasAnyClassFlags(CLASS_Abstract) == false )
+				if ( WidgetClass->IsChildOf(UPanelWidget::StaticClass()) )
 				{
 					// Only allow replacement with panels that accept multiple children
 					if ( WidgetClass->GetDefaultObject<UPanelWidget>()->CanHaveMultipleChildren() )
@@ -430,10 +442,9 @@ void FWidgetBlueprintEditorUtils::BuildReplaceWithMenu(FMenuBuilder& Menu, UWidg
 							FText::GetEmpty(),
 							FSlateIcon(),
 							FUIAction(
-							FExecuteAction::CreateStatic(&FWidgetBlueprintEditorUtils::ReplaceWidgets, BP, Widgets, WidgetClass),
-							FCanExecuteAction()
-							)
-							);
+								FExecuteAction::CreateStatic(&FWidgetBlueprintEditorUtils::ReplaceWidgets, BP, Widgets, WidgetClass),
+								FCanExecuteAction()
+							));
 					}
 				}
 			}
@@ -448,12 +459,12 @@ void FWidgetBlueprintEditorUtils::ReplaceWidgets(UWidgetBlueprint* BP, TSet<FWid
 
 	TSharedPtr<FWidgetTemplateClass> Template = MakeShareable(new FWidgetTemplateClass(WidgetClass));
 
-	UPanelWidget* NewReplacementWidget = CastChecked<UPanelWidget>(Template->Create(BP->WidgetTree));
-
 	for ( FWidgetReference& Item : Widgets )
 	{
 		if ( UPanelWidget* ExistingPanel = Cast<UPanelWidget>(Item.GetTemplate()) )
 		{
+			UPanelWidget* NewReplacementWidget = CastChecked<UPanelWidget>(Template->Create(BP->WidgetTree));
+
 			ExistingPanel->Modify();
 
 			if ( UPanelWidget* CurrentParent = ExistingPanel->GetParent() )
diff --git a/Engine/Source/Runtime/UMG/Public/Components/PanelWidget.h b/Engine/Source/Runtime/UMG/Public/Components/PanelWidget.h
--- a/Engine/Source/Runtime/UMG/Public/Components/PanelWidget.h
+++ b/Engine/Source/Runtime/UMG/Public/Components/PanelWidget.h
@@ -106,6 +106,12 @@ public:
 		return bCanHaveMultipleChildren;
 	}
 
+	/** @returns true if the panel can accept another child widget. */
+	bool CanAddMoreChildren() const
+	{
+		return CanHaveMultipleChildren() || GetChildrenCount() == 0;
+	}
+
 	/** Sets that this widget is being designed sets it on all children as well. */
 	virtual void SetIsDesignTime(bool bInDesignTime) override;
 
