UMG - The design surface now allows overzooming without control for touch input just like it does for the scrollwheel on mouse.

diff --git a/Engine/Source/Editor/UMGEditor/Private/Designer/SDesignSurface.cpp b/Engine/Source/Editor/UMGEditor/Private/Designer/SDesignSurface.cpp
--- a/Engine/Source/Editor/UMGEditor/Private/Designer/SDesignSurface.cpp
+++ b/Engine/Source/Editor/UMGEditor/Private/Designer/SDesignSurface.cpp
@@ -263,7 +263,7 @@ FReply SDesignSurface::OnTouchGesture(const FGeometry& MyGeometry, const FPointe
 			// We want to zoom into this point; i.e. keep it the same fraction offset into the panel
 			const FVector2D WidgetSpaceCursorPos = MyGeometry.AbsoluteToLocal(GestureEvent.GetScreenSpacePosition());
 			const int32 ZoomLevelDelta = TotalGestureMagnify > 0.0f ? 1 : -1;
-			ChangeZoomLevel(ZoomLevelDelta, WidgetSpaceCursorPos, GestureEvent.IsControlDown());
+			ChangeZoomLevel(ZoomLevelDelta, WidgetSpaceCursorPos, !bRequireControlToOverZoom || GestureEvent.IsControlDown());
 			TotalGestureMagnify = 0.0f;
 		}
 		return FReply::Handled();
