UMG - Some initial work getting explicit widget navigation working.  Currently uses the widget's name for controlling, which isn't super stable.  Considering introducing GUIDs to every widget.

diff --git a/Engine/Source/Editor/UMGEditor/Private/Customizations/WidgetNavigationCustomization.cpp b/Engine/Source/Editor/UMGEditor/Private/Customizations/WidgetNavigationCustomization.cpp
--- a/Engine/Source/Editor/UMGEditor/Private/Customizations/WidgetNavigationCustomization.cpp
+++ b/Engine/Source/Editor/UMGEditor/Private/Customizations/WidgetNavigationCustomization.cpp
@@ -76,32 +76,39 @@ FReply FWidgetNavigationCustomization::OnCustomizeNavigation(TWeakPtr<IPropertyH
 	return FReply::Handled();
 }
 
-FText FWidgetNavigationCustomization::HandleNavigationText(TWeakPtr<IPropertyHandle> PropertyHandle, EUINavigation Nav) const
+EUINavigationRule FWidgetNavigationCustomization::GetNavigationRule(TWeakPtr<IPropertyHandle> PropertyHandle, EUINavigation Nav) const
 {
 	TArray<UObject*> OuterObjects;
 	TSharedPtr<IPropertyHandle> PropertyHandlePtr = PropertyHandle.Pin();
 	PropertyHandlePtr->GetOuterObjects(OuterObjects);
 
 	EUINavigationRule Rule = EUINavigationRule::Invalid;
-	for (UObject* OuterObject : OuterObjects)
+	for ( UObject* OuterObject : OuterObjects )
 	{
-		if (UWidget* Widget = Cast<UWidget>(OuterObject))
+		if ( UWidget* Widget = Cast<UWidget>(OuterObject) )
 		{
 			EUINavigationRule CurRule = EUINavigationRule::Escape;
 			UWidgetNavigation* WidgetNavigation = Widget->Navigation;
-			if (Widget->Navigation != nullptr)
+			if ( Widget->Navigation != nullptr )
 			{
 				CurRule = WidgetNavigation->GetNavigationRule(Nav);
 			}
 
-			if (Rule != EUINavigationRule::Invalid && CurRule != Rule)
+			if ( Rule != EUINavigationRule::Invalid && CurRule != Rule )
 			{
-				return LOCTEXT("NavigationMultipleValues", "Multiple Values");
+				return EUINavigationRule::Invalid;
 			}
 			Rule = CurRule;
 		}
 	}
 
+	return Rule;
+}
+
+FText FWidgetNavigationCustomization::GetNavigationText(TWeakPtr<IPropertyHandle> PropertyHandle, EUINavigation Nav) const
+{
+	EUINavigationRule Rule = GetNavigationRule(PropertyHandle, Nav);
+
 	switch (Rule)
 	{
 	case EUINavigationRule::Escape:
@@ -110,11 +117,79 @@ FText FWidgetNavigationCustomization::HandleNavigationText(TWeakPtr<IPropertyHan
 		return LOCTEXT("NavigationStop", "Stop");
 	case EUINavigationRule::Wrap:
 		return LOCTEXT("NavigationWrap", "Wrap");
+	case EUINavigationRule::Explicit:
+		return LOCTEXT("NavigationExplicit", "Explicit");
+	case EUINavigationRule::Invalid:
+		return LOCTEXT("NavigationMultipleValues", "Multiple Values");
 	}
 
 	return FText::GetEmpty();
 }
 
+FText FWidgetNavigationCustomization::GetExplictWidget(TWeakPtr<IPropertyHandle> PropertyHandle, EUINavigation Nav) const
+{
+	TArray<UObject*> OuterObjects;
+	TSharedPtr<IPropertyHandle> PropertyHandlePtr = PropertyHandle.Pin();
+	PropertyHandlePtr->GetOuterObjects(OuterObjects);
+
+	bool bFirst = true;
+	FName Rule = NAME_None;
+	for ( UObject* OuterObject : OuterObjects )
+	{
+		if ( UWidget* Widget = Cast<UWidget>(OuterObject) )
+		{
+			FName CurRule = NAME_None;
+			UWidgetNavigation* WidgetNavigation = Widget->Navigation;
+			if ( Widget->Navigation != nullptr )
+			{
+				CurRule = WidgetNavigation->GetNavigationData(Nav).WidgetToFocus;
+				if ( bFirst )
+				{
+					Rule = CurRule;
+					bFirst = false;
+				}
+			}
+
+			if ( CurRule != Rule )
+			{
+				return LOCTEXT("NavigationMultipleValues", "Multiple Values");
+			}
+
+			Rule = CurRule;
+		}
+	}
+
+	return FText::FromName(Rule);
+}
+
+void FWidgetNavigationCustomization::OnCommitExplictWidgetText(const FText& ItemFText, ETextCommit::Type CommitInfo, TWeakPtr<IPropertyHandle> PropertyHandle, EUINavigation Nav)
+{
+	TArray<UObject*> OuterObjects;
+	TSharedPtr<IPropertyHandle> PropertyHandlePtr = PropertyHandle.Pin();
+	PropertyHandlePtr->GetOuterObjects(OuterObjects);
+
+	const FScopedTransaction Transaction(LOCTEXT("InitializeNavigation", "Edit Widget Navigation"));
+
+	FName GotoWidgetName = FName(*ItemFText.ToString());
+
+	for ( UObject* OuterObject : OuterObjects )
+	{
+		if ( UWidget* Widget = Cast<UWidget>(OuterObject) )
+		{
+			FWidgetReference WidgetReference = Editor.Pin()->GetReferenceFromPreview(Widget);
+
+			SetNav(WidgetReference.GetPreview(), Nav, TOptional<EUINavigationRule>(), GotoWidgetName);
+			SetNav(WidgetReference.GetTemplate(), Nav, TOptional<EUINavigationRule>(), GotoWidgetName);
+		}
+	}
+}
+
+EVisibility FWidgetNavigationCustomization::GetExplictWidgetFieldVisibility(TWeakPtr<IPropertyHandle> PropertyHandle, EUINavigation Nav) const
+{
+	EUINavigationRule Rule = GetNavigationRule(PropertyHandle, Nav);
+	return Rule == EUINavigationRule::Explicit ? EVisibility::Visible : EVisibility::Collapsed;
+}
+
 void FWidgetNavigationCustomization::MakeNavRow(TWeakPtr<IPropertyHandle> PropertyHandle, IDetailChildrenBuilder& ChildBuilder, EUINavigation Nav, FText NavName)
 {
 	ChildBuilder.AddChildContent(NavName)
@@ -125,18 +200,37 @@ void FWidgetNavigationCustomization::MakeNavRow(TWeakPtr<IPropertyHandle> Proper
 			.Text(NavName)
 		]
 		.ValueContent()
+		.MaxDesiredWidth(300)
 		[
-			SNew(SComboButton)
-			.HAlign(HAlign_Center)
-			.ButtonContent()
+			SNew(SHorizontalBox)
+
+			+ SHorizontalBox::Slot()
+			.AutoWidth()
 			[
-				SNew(STextBlock)
-				.Text(this, &FWidgetNavigationCustomization::HandleNavigationText, PropertyHandle, Nav)
+				SNew(SComboButton)
+				.HAlign(HAlign_Center)
+				.ButtonContent()
+				[
+					SNew(STextBlock)
+					.Text(this, &FWidgetNavigationCustomization::GetNavigationText, PropertyHandle, Nav)
+				]
+				.ContentPadding(FMargin(2.0f, 1.0f))
+				.MenuContent()
+				[
+					MakeNavMenu(PropertyHandle, Nav)
+				]
 			]
-			.ContentPadding(FMargin(2.0f, 1.0f))
-			.MenuContent()
+
+			// Explicit Navigation Widget
+			+ SHorizontalBox::Slot()
+			.FillWidth(1.0f)
 			[
-				MakeNavMenu(PropertyHandle, Nav)
+				SNew(SEditableTextBox)
+				.HintText(LOCTEXT("WidgetName", "Widget Name?"))
+				.Text(this, &FWidgetNavigationCustomization::GetExplictWidget, PropertyHandle, Nav)
+				.OnTextCommitted(this, &FWidgetNavigationCustomization::OnCommitExplictWidgetText, PropertyHandle, Nav)
+				.Font(IDetailLayoutBuilder::GetDetailFont())
+				.Visibility(this, &FWidgetNavigationCustomization::GetExplictWidgetFieldVisibility, PropertyHandle, Nav)
 			]
 		];
 }
@@ -155,8 +249,8 @@ TSharedRef<class SWidget> FWidgetNavigationCustomization::MakeNavMenu(TWeakPtr<I
 		FUIAction WrapAction(FExecuteAction::CreateSP(this, &FWidgetNavigationCustomization::HandleNavMenuEntryClicked, PropertyHandle, Nav, EUINavigationRule::Wrap));
 		MenuBuilder.AddMenuEntry(LOCTEXT("NavigationRuleWrap", "Wrap"), LOCTEXT("NavigationRuleWrapHint", "Navigation will wrap to the opposite bound of this object."), FSlateIcon(), WrapAction);
 
-		//FUIAction ExplicitAction(FExecuteAction::CreateSP(this, &FWidgetNavigationCustomization::HandleNavMenuEntryClicked, PropertyHandle, Nav, EUINavigationRule::Explicit));
-		//MenuBuilder.AddMenuEntry(LOCTEXT("NavigationRuleExplicit", "Explicit"), LOCTEXT("NavigationRuleExplicitHint", "Navigation will go to a specified widget."), FSlateIcon(), ExplicitAction);
+		FUIAction ExplicitAction(FExecuteAction::CreateSP(this, &FWidgetNavigationCustomization::HandleNavMenuEntryClicked, PropertyHandle, Nav, EUINavigationRule::Explicit));
+		MenuBuilder.AddMenuEntry(LOCTEXT("NavigationRuleExplicit", "Explicit"), LOCTEXT("NavigationRuleExplicitHint", "Navigation will go to a specified widget."), FSlateIcon(), ExplicitAction);
 
 		//FUIAction CustomAction(FExecuteAction::CreateSP(this, &FWidgetNavigationCustomization::HandleNavMenuEntryClicked, PropertyHandle, Nav, EUINavigationRule::Custom));
 		//MenuBuilder.AddMenuEntry(LOCTEXT("NavigationRuleCustom", "Custom"), LOCTEXT("NavigationRuleCustomHint", "Custom function can determine what widget is navigated to."), FSlateIcon(), CustomAction);
@@ -172,7 +266,7 @@ void FWidgetNavigationCustomization::HandleNavMenuEntryClicked(TWeakPtr<IPropert
 	TSharedPtr<IPropertyHandle> PropertyHandlePtr = PropertyHandle.Pin();
 	PropertyHandlePtr->GetOuterObjects(OuterObjects);
 
-	const FScopedTransaction Transaction(LOCTEXT("InitializeNavigation", "Initialize Navigation"));
+	const FScopedTransaction Transaction(LOCTEXT("InitializeNavigation", "Edit Widget Navigation"));
 
 	for (UObject* OuterObject : OuterObjects)
 	{
@@ -180,13 +274,13 @@ void FWidgetNavigationCustomization::HandleNavMenuEntryClicked(TWeakPtr<IPropert
 		{
 			FWidgetReference WidgetReference = Editor.Pin()->GetReferenceFromPreview(Widget);
 
-			SetNav(WidgetReference.GetPreview(), Nav, Rule);
-			SetNav(WidgetReference.GetTemplate(), Nav, Rule);
+			SetNav(WidgetReference.GetPreview(), Nav, Rule, TOptional<FName>());
+			SetNav(WidgetReference.GetTemplate(), Nav, Rule, TOptional<FName>());
 		}
 	}
 }
 
-void FWidgetNavigationCustomization::SetNav(UWidget* Widget, EUINavigation Nav, EUINavigationRule Rule)
+void FWidgetNavigationCustomization::SetNav(UWidget* Widget, EUINavigation Nav, TOptional<EUINavigationRule> Rule, TOptional<FName> WidgetToFocus)
 {
 	Widget->Modify();
 
@@ -196,20 +290,36 @@ void FWidgetNavigationCustomization::SetNav(UWidget* Widget, EUINavigation Nav,
 		WidgetNavigation = NewObject<UWidgetNavigation>(Widget);
 	}
 
-	switch (Nav)
+	FWidgetNavigationData* DirectionNavigation = nullptr;
+
+	switch ( Nav )
 	{
 	case EUINavigation::Left:
-		WidgetNavigation->Left.Rule = Rule;
+		DirectionNavigation = &WidgetNavigation->Left;
 		break;
 	case EUINavigation::Right:
-		WidgetNavigation->Right.Rule = Rule;
+		DirectionNavigation = &WidgetNavigation->Right;
 		break;
 	case EUINavigation::Up:
-		WidgetNavigation->Up.Rule = Rule;
+		DirectionNavigation = &WidgetNavigation->Up;
 		break;
 	case EUINavigation::Down:
-		WidgetNavigation->Down.Rule = Rule;
+		DirectionNavigation = &WidgetNavigation->Down;
 		break;
+	default:
+		// Should not be possible.
+		check(false);
+		return;
+	}
+
+	if ( Rule.IsSet() )
+	{
+		DirectionNavigation->Rule = Rule.GetValue();
+	}
+
+	if ( WidgetToFocus.IsSet() )
+	{
+		DirectionNavigation->WidgetToFocus = WidgetToFocus.GetValue();
 	}
 
 	if ( WidgetNavigation->IsDefault() )
diff --git a/Engine/Source/Editor/UMGEditor/Private/Customizations/WidgetNavigationCustomization.h b/Engine/Source/Editor/UMGEditor/Private/Customizations/WidgetNavigationCustomization.h
--- a/Engine/Source/Editor/UMGEditor/Private/Customizations/WidgetNavigationCustomization.h
+++ b/Engine/Source/Editor/UMGEditor/Private/Customizations/WidgetNavigationCustomization.h
@@ -28,7 +28,8 @@ private:
 
 	FReply OnCustomizeNavigation(TWeakPtr<IPropertyHandle> PropertyHandle);
 	
-	FText HandleNavigationText(TWeakPtr<IPropertyHandle> PropertyHandle, EUINavigation Nav) const;
+	EUINavigationRule GetNavigationRule(TWeakPtr<IPropertyHandle> PropertyHandle, EUINavigation Nav) const;
+	FText GetNavigationText(TWeakPtr<IPropertyHandle> PropertyHandle, EUINavigation Nav) const;
 
 	void MakeNavRow(TWeakPtr<IPropertyHandle> PropertyHandle, IDetailChildrenBuilder& ChildBuilder, EUINavigation Nav, FText NavName);
 
@@ -36,7 +37,11 @@ private:
 
 	void HandleNavMenuEntryClicked(TWeakPtr<IPropertyHandle> PropertyHandle, EUINavigation Nav, EUINavigationRule Rule);
 
-	void SetNav(UWidget* Widget, EUINavigation Nav, EUINavigationRule Rule);
+	FText GetExplictWidget(TWeakPtr<IPropertyHandle> PropertyHandle, EUINavigation Nav) const;
+	void OnCommitExplictWidgetText(const FText& ItemFText, ETextCommit::Type CommitInfo, TWeakPtr<IPropertyHandle> PropertyHandle, EUINavigation Nav);
+	EVisibility GetExplictWidgetFieldVisibility(TWeakPtr<IPropertyHandle> PropertyHandle, EUINavigation Nav) const;
+
+	void SetNav(UWidget* Widget, EUINavigation Nav, TOptional<EUINavigationRule> Rule, TOptional<FName> WidgetToFocus);
 
 private:
 	TWeakPtr<FWidgetBlueprintEditor> Editor;
diff --git a/Engine/Source/Runtime/UMG/Private/Components/Widget.cpp b/Engine/Source/Runtime/UMG/Private/Components/Widget.cpp
--- a/Engine/Source/Runtime/UMG/Private/Components/Widget.cpp
+++ b/Engine/Source/Runtime/UMG/Private/Components/Widget.cpp
@@ -668,10 +668,22 @@ void UWidget::SynchronizeProperties()
 		SafeWidget->SetToolTipText(GAME_SAFE_OPTIONAL_BINDING(FText, ToolTipText));
 	}
 
-	if (Navigation != nullptr)
+#if WITH_EDITOR
+	// In editor builds we add metadata to the widget so that once hit with the widget reflector it can report
+	// where it comes from, what blueprint, what the name of the widget was...etc.
+	SafeWidget->AddMetadata<FReflectionMetaData>(MakeShareable(new FReflectionMetaData(GetFName(), GetClass(), WidgetGeneratedBy)));
+#endif
+}
+
+void UWidget::BuildNavigation()
+{
+	TSharedPtr<SWidget> SafeWidget = GetCachedWidget();
+	check(SafeWidget.IsValid());
+
+	if ( Navigation != nullptr )
 	{
 		TSharedPtr<FNavigationMetaData> MetaData = SafeWidget->GetMetaData<FNavigationMetaData>();
-		if (!MetaData.IsValid())
+		if ( !MetaData.IsValid() )
 		{
 			MetaData = MakeShareable(new FNavigationMetaData());
 			SafeWidget->AddMetadata(MetaData.ToSharedRef());
@@ -679,12 +691,6 @@ void UWidget::SynchronizeProperties()
 
 		Navigation->UpdateMetaData(MetaData.ToSharedRef());
 	}
-
-#if WITH_EDITOR
-	// In editor builds we add metadata to the widget so that once hit with the widget reflector it can report
-	// where it comes from, what blueprint, what the name of the widget was...etc.
-	SafeWidget->AddMetadata<FReflectionMetaData>(MakeShareable(new FReflectionMetaData(GetFName(), GetClass(), WidgetGeneratedBy)));
-#endif
 }
 
 #if WITH_EDITOR
diff --git a/Engine/Source/Runtime/UMG/Private/UWidgetNavigation.cpp b/Engine/Source/Runtime/UMG/Private/UWidgetNavigation.cpp
--- a/Engine/Source/Runtime/UMG/Private/UWidgetNavigation.cpp
+++ b/Engine/Source/Runtime/UMG/Private/UWidgetNavigation.cpp
@@ -12,6 +12,26 @@ UWidgetNavigation::UWidgetNavigation(const FObjectInitializer& ObjectInitializer
 
 #if WITH_EDITOR
 
+FWidgetNavigationData& UWidgetNavigation::GetNavigationData(EUINavigation Nav)
+{
+	switch ( Nav )
+	{
+	case EUINavigation::Up:
+		return Up;
+	case EUINavigation::Down:
+		return Down;
+	case EUINavigation::Left:
+		return Left;
+	case EUINavigation::Right:
+		return Right;
+	}
+
+	// Should never happen
+	check(false);
+
+	return Up;
+}
+
 EUINavigationRule UWidgetNavigation::GetNavigationRule(EUINavigation Nav)
 {
 	switch ( Nav )
@@ -30,6 +50,26 @@ EUINavigationRule UWidgetNavigation::GetNavigationRule(EUINavigation Nav)
 
 #endif
 
+void UWidgetNavigation::ResolveExplictRules(UWidgetTree* WidgetTree)
+{
+	if ( Up.Rule == EUINavigationRule::Explicit )
+	{
+		Up.Widget = WidgetTree->FindWidget(Up.WidgetToFocus);
+	}
+	if ( Down.Rule == EUINavigationRule::Explicit )
+	{
+		Down.Widget = WidgetTree->FindWidget(Down.WidgetToFocus);
+	}
+	if ( Left.Rule == EUINavigationRule::Explicit )
+	{
+		Left.Widget = WidgetTree->FindWidget(Left.WidgetToFocus);
+	}
+	if ( Right.Rule == EUINavigationRule::Explicit )
+	{
+		Right.Widget = WidgetTree->FindWidget(Right.WidgetToFocus);
+	}
+}
+
 void UWidgetNavigation::UpdateMetaData(TSharedRef<FNavigationMetaData> MetaData)
 {
 	UpdateMetaDataEntry(MetaData, Up, EUINavigation::Up);
@@ -59,5 +99,11 @@ void UWidgetNavigation::UpdateMetaDataEntry(TSharedRef<FNavigationMetaData> Meta
 	case EUINavigationRule::Wrap:
 		MetaData->SetNavigationWrap(Nav);
 		break;
+	case EUINavigationRule::Explicit:
+		if ( NavData.Widget.IsValid() )
+		{
+			MetaData->SetNavigationExplicit(Nav, NavData.Widget.Get()->GetCachedWidget());
+		}
+		break;
 	}
 }
diff --git a/Engine/Source/Runtime/UMG/Private/UserWidget.cpp b/Engine/Source/Runtime/UMG/Private/UserWidget.cpp
--- a/Engine/Source/Runtime/UMG/Private/UserWidget.cpp
+++ b/Engine/Source/Runtime/UMG/Private/UserWidget.cpp
@@ -61,10 +61,6 @@ void UUserWidget::Initialize()
 				}
 			}
 		});
-
-		// TODO Find a way to remove the bindings from the table.  Currently they're still needed.
-		// Clear the named slot bindings table.
-		//NamedSlotBindings.Reset();
 	}
 }
 
@@ -292,6 +288,12 @@ TSharedRef<SWidget> UUserWidget::RebuildWidget()
 
 void UUserWidget::OnWidgetRebuilt()
 {
+	// When a user widget is rebuilt we can safely initialize the navigation now since all the slate
+	// widgets should be held onto by a smart pointer at this point.
+	WidgetTree->ForEachWidget([&] (UWidget* Widget) {
+		Widget->BuildNavigation();
+	});
+
 	if (!IsDesignTime())
 	{
 		// Notify the widget that it has been constructed.
diff --git a/Engine/Source/Runtime/UMG/Private/WidgetBlueprintGeneratedClass.cpp b/Engine/Source/Runtime/UMG/Private/WidgetBlueprintGeneratedClass.cpp
--- a/Engine/Source/Runtime/UMG/Private/WidgetBlueprintGeneratedClass.cpp
+++ b/Engine/Source/Runtime/UMG/Private/WidgetBlueprintGeneratedClass.cpp
@@ -107,6 +107,12 @@ void UWidgetBlueprintGeneratedClass::InitializeWidget(UUserWidget* UserWidget) c
 				}
 			}
 
+			// Initialize Navigation Data
+			if (Widget->Navigation)
+			{
+				Widget->Navigation->ResolveExplictRules(ClonedTree);
+			}
+
 	#if WITH_EDITOR
 			Widget->ConnectEditorData();
 	#endif
diff --git a/Engine/Source/Runtime/UMG/Public/Blueprint/WidgetNavigation.h b/Engine/Source/Runtime/UMG/Public/Blueprint/WidgetNavigation.h
--- a/Engine/Source/Runtime/UMG/Public/Blueprint/WidgetNavigation.h
+++ b/Engine/Source/Runtime/UMG/Public/Blueprint/WidgetNavigation.h
@@ -21,6 +21,8 @@ public:
 
 	UPROPERTY(EditAnywhere, BlueprintReadOnly, Category="Navigation")
 	FName WidgetToFocus;
+
+	TWeakObjectPtr<UWidget> Widget;
 };
 
 /**
@@ -52,11 +54,16 @@ public:
 
 #if WITH_EDITOR
 
+	/**  */
+	FWidgetNavigationData& GetNavigationData(EUINavigation Nav);
+
 	/**  */
 	EUINavigationRule GetNavigationRule(EUINavigation Nav);
 
 #endif
 
+	void ResolveExplictRules(class UWidgetTree* WidgetTree);
+
 	/** Updates a slate metadata object to match this configured navigation ruleset. */
 	void UpdateMetaData(TSharedRef<FNavigationMetaData> MetaData);
 
diff --git a/Engine/Source/Runtime/UMG/Public/Components/Widget.h b/Engine/Source/Runtime/UMG/Public/Components/Widget.h
--- a/Engine/Source/Runtime/UMG/Public/Components/Widget.h
+++ b/Engine/Source/Runtime/UMG/Public/Components/Widget.h
@@ -334,6 +334,13 @@ public:
 	 */
 	virtual void SynchronizeProperties();
 
+	/**
+	 * Called by the owning user widget after the slate widget has been created.  After the entire widget tree
+	 * has been initialized, any widget reference that was needed to support navigating to another widget will
+	 * now be initialized and ready for usage.
+	 */
+	void BuildNavigation();
+
 #if WITH_EDITOR
 	/** Returns if the widget is currently being displayed in the designer, it may want to display different data. */
 	bool IsDesignTime() const;
