UMG - Adding more + Zoom levels into the designer.

diff --git a/Engine/Source/Editor/UMGEditor/Private/Designer/SDesignSurface.cpp b/Engine/Source/Editor/UMGEditor/Private/Designer/SDesignSurface.cpp
--- a/Engine/Source/Editor/UMGEditor/Private/Designer/SDesignSurface.cpp
+++ b/Engine/Source/Editor/UMGEditor/Private/Designer/SDesignSurface.cpp
@@ -29,28 +29,33 @@ struct FFixedZoomLevelsContainer : public FZoomLevelsContainer
 		// Initialize zoom levels if not done already
 		if ( ZoomLevels.Num() == 0 )
 		{
-			ZoomLevels.Reserve(21);
-			ZoomLevels.Add(FZoomLevelEntry(0.150f, LOCTEXT("ZoomLevel", "-10"), EGraphRenderingLOD::LowestDetail));
-			ZoomLevels.Add(FZoomLevelEntry(0.175f, LOCTEXT("ZoomLevel", "-9"), EGraphRenderingLOD::LowestDetail));
-			ZoomLevels.Add(FZoomLevelEntry(0.200f, LOCTEXT("ZoomLevel", "-8"), EGraphRenderingLOD::LowestDetail));
-			ZoomLevels.Add(FZoomLevelEntry(0.225f, LOCTEXT("ZoomLevel", "-7"), EGraphRenderingLOD::LowDetail));
-			ZoomLevels.Add(FZoomLevelEntry(0.250f, LOCTEXT("ZoomLevel", "-6"), EGraphRenderingLOD::LowDetail));
-			ZoomLevels.Add(FZoomLevelEntry(0.375f, LOCTEXT("ZoomLevel", "-5"), EGraphRenderingLOD::MediumDetail));
-			ZoomLevels.Add(FZoomLevelEntry(0.500f, LOCTEXT("ZoomLevel", "-4"), EGraphRenderingLOD::MediumDetail));
-			ZoomLevels.Add(FZoomLevelEntry(0.675f, LOCTEXT("ZoomLevel", "-3"), EGraphRenderingLOD::MediumDetail));
-			ZoomLevels.Add(FZoomLevelEntry(0.750f, LOCTEXT("ZoomLevel", "-2"), EGraphRenderingLOD::DefaultDetail));
-			ZoomLevels.Add(FZoomLevelEntry(0.875f, LOCTEXT("ZoomLevel", "-1"), EGraphRenderingLOD::DefaultDetail));
-			ZoomLevels.Add(FZoomLevelEntry(1.000f, LOCTEXT("ZoomLevel", "1:1"), EGraphRenderingLOD::DefaultDetail));
-			ZoomLevels.Add(FZoomLevelEntry(1.250f, LOCTEXT("ZoomLevel", "+1"), EGraphRenderingLOD::DefaultDetail));
-			ZoomLevels.Add(FZoomLevelEntry(1.500f, LOCTEXT("ZoomLevel", "+2"), EGraphRenderingLOD::DefaultDetail));
-			ZoomLevels.Add(FZoomLevelEntry(1.750f, LOCTEXT("ZoomLevel", "+3"), EGraphRenderingLOD::FullyZoomedIn));
-			ZoomLevels.Add(FZoomLevelEntry(2.000f, LOCTEXT("ZoomLevel", "+4"), EGraphRenderingLOD::FullyZoomedIn));
-			ZoomLevels.Add(FZoomLevelEntry(2.250f, LOCTEXT("ZoomLevel", "+5"), EGraphRenderingLOD::FullyZoomedIn));
-			ZoomLevels.Add(FZoomLevelEntry(2.500f, LOCTEXT("ZoomLevel", "+6"), EGraphRenderingLOD::FullyZoomedIn));
-			ZoomLevels.Add(FZoomLevelEntry(2.750f, LOCTEXT("ZoomLevel", "+7"), EGraphRenderingLOD::FullyZoomedIn));
-			ZoomLevels.Add(FZoomLevelEntry(3.000f, LOCTEXT("ZoomLevel", "+8"), EGraphRenderingLOD::FullyZoomedIn));
-			ZoomLevels.Add(FZoomLevelEntry(3.250f, LOCTEXT("ZoomLevel", "+9"), EGraphRenderingLOD::FullyZoomedIn));
-			ZoomLevels.Add(FZoomLevelEntry(3.500f, LOCTEXT("ZoomLevel", "+10"), EGraphRenderingLOD::FullyZoomedIn));
+			ZoomLevels.Reserve(26);
+			ZoomLevels.Add(FZoomLevelEntry(0.150f, FText::FromString("-10"), EGraphRenderingLOD::LowestDetail));
+			ZoomLevels.Add(FZoomLevelEntry(0.175f, FText::FromString("-9"), EGraphRenderingLOD::LowestDetail));
+			ZoomLevels.Add(FZoomLevelEntry(0.200f, FText::FromString("-8"), EGraphRenderingLOD::LowestDetail));
+			ZoomLevels.Add(FZoomLevelEntry(0.225f, FText::FromString("-7"), EGraphRenderingLOD::LowDetail));
+			ZoomLevels.Add(FZoomLevelEntry(0.250f, FText::FromString("-6"), EGraphRenderingLOD::LowDetail));
+			ZoomLevels.Add(FZoomLevelEntry(0.375f, FText::FromString("-5"), EGraphRenderingLOD::MediumDetail));
+			ZoomLevels.Add(FZoomLevelEntry(0.500f, FText::FromString("-4"), EGraphRenderingLOD::MediumDetail));
+			ZoomLevels.Add(FZoomLevelEntry(0.675f, FText::FromString("-3"), EGraphRenderingLOD::MediumDetail));
+			ZoomLevels.Add(FZoomLevelEntry(0.750f, FText::FromString("-2"), EGraphRenderingLOD::DefaultDetail));
+			ZoomLevels.Add(FZoomLevelEntry(0.875f, FText::FromString("-1"), EGraphRenderingLOD::DefaultDetail));
+			ZoomLevels.Add(FZoomLevelEntry(1.000f, FText::FromString("1:1"), EGraphRenderingLOD::DefaultDetail));
+			ZoomLevels.Add(FZoomLevelEntry(1.250f, FText::FromString("+1"), EGraphRenderingLOD::DefaultDetail));
+			ZoomLevels.Add(FZoomLevelEntry(1.500f, FText::FromString("+2"), EGraphRenderingLOD::DefaultDetail));
+			ZoomLevels.Add(FZoomLevelEntry(1.750f, FText::FromString("+3"), EGraphRenderingLOD::FullyZoomedIn));
+			ZoomLevels.Add(FZoomLevelEntry(2.000f, FText::FromString("+4"), EGraphRenderingLOD::FullyZoomedIn));
+			ZoomLevels.Add(FZoomLevelEntry(2.250f, FText::FromString("+5"), EGraphRenderingLOD::FullyZoomedIn));
+			ZoomLevels.Add(FZoomLevelEntry(2.500f, FText::FromString("+6"), EGraphRenderingLOD::FullyZoomedIn));
+			ZoomLevels.Add(FZoomLevelEntry(2.750f, FText::FromString("+7"), EGraphRenderingLOD::FullyZoomedIn));
+			ZoomLevels.Add(FZoomLevelEntry(3.000f, FText::FromString("+8"), EGraphRenderingLOD::FullyZoomedIn));
+			ZoomLevels.Add(FZoomLevelEntry(3.250f, FText::FromString("+9"), EGraphRenderingLOD::FullyZoomedIn));
+			ZoomLevels.Add(FZoomLevelEntry(3.500f, FText::FromString("+10"), EGraphRenderingLOD::FullyZoomedIn));
+			ZoomLevels.Add(FZoomLevelEntry(4.000f, FText::FromString("+11"), EGraphRenderingLOD::FullyZoomedIn));
+			ZoomLevels.Add(FZoomLevelEntry(5.000f, FText::FromString("+12"), EGraphRenderingLOD::FullyZoomedIn));
+			ZoomLevels.Add(FZoomLevelEntry(6.000f, FText::FromString("+13"), EGraphRenderingLOD::FullyZoomedIn));
+			ZoomLevels.Add(FZoomLevelEntry(7.000f, FText::FromString("+14"), EGraphRenderingLOD::FullyZoomedIn));
+			ZoomLevels.Add(FZoomLevelEntry(8.000f, FText::FromString("+15"), EGraphRenderingLOD::FullyZoomedIn));
 		}
 	}
 
