UMG - Reworking Custom and Desired sizes in the designer.  You can now preview any custom or desired size in the context of the screen as well.  So there's now a seperate dropdown for choosing screen size, vs. how the widget preview should fill that screen.

diff --git a/Engine/Source/Editor/UMGEditor/Private/Designer/SDesignerView.cpp b/Engine/Source/Editor/UMGEditor/Private/Designer/SDesignerView.cpp
--- a/Engine/Source/Editor/UMGEditor/Private/Designer/SDesignerView.cpp
+++ b/Engine/Source/Editor/UMGEditor/Private/Designer/SDesignerView.cpp
@@ -225,16 +225,12 @@ void SDesignerView::Construct(const FArguments& InArgs, TSharedPtr<FWidgetBluepr
 
 					+ SOverlay::Slot()
 					[
-						SNew(SBox)
-						.WidthOverride(this, &SDesignerView::GetPreviewWidth)
-						.HeightOverride(this, &SDesignerView::GetPreviewHeight)
-						.HAlign(HAlign_Fill)
-						.VAlign(VAlign_Fill)
-						.Visibility(EVisibility::SelfHitTestInvisible)
+						SAssignNew(PreviewAreaConstraint, SBox)
+						.WidthOverride(this, &SDesignerView::GetPreviewAreaWidth)
+						.HeightOverride(this, &SDesignerView::GetPreviewAreaHeight)
 						[
 							SAssignNew(PreviewSurface, SDPIScaler)
 							.DPIScale(this, &SDesignerView::GetPreviewDPIScale)
-							.Visibility(EVisibility::SelfHitTestInvisible)
 						]
 					]
 				]
@@ -341,6 +337,7 @@ TSharedRef<SWidget> SDesignerView::CreateOverlayUI()
 
 		+ SHorizontalBox::Slot()
 		.AutoWidth()
+		.VAlign(VAlign_Center)
 		.Padding(6, 2, 0, 0)
 		[
 			SNew(STextBlock)
@@ -359,7 +356,7 @@ TSharedRef<SWidget> SDesignerView::CreateOverlayUI()
 
 		+ SHorizontalBox::Slot()
 		.AutoWidth()
-		.Padding(5.0f)
+		.VAlign(VAlign_Center)
 		[
 			SNew(SDesignerToolBar)
 			.CommandList(CommandList)
@@ -367,39 +364,59 @@ TSharedRef<SWidget> SDesignerView::CreateOverlayUI()
 
 		+ SHorizontalBox::Slot()
 		.AutoWidth()
-		.Padding(5.0f)
+		.VAlign(VAlign_Center)
 		[
 			SNew(SButton)
 			.ButtonStyle(FEditorStyle::Get(), "ViewportMenu.Button")
 			.ToolTipText(LOCTEXT("ZoomToFit_ToolTip", "Zoom To Fit"))
 			.OnClicked(this, &SDesignerView::HandleZoomToFitClicked)
+			.ContentPadding(1.0f)
 			[
 				SNew(SImage)
 				.Image(FEditorStyle::GetBrush("UMGEditor.ZoomToFit"))
 			]
 		]
-				
+		
+		// Preview Screen Size
+		+ SHorizontalBox::Slot()
+		.AutoWidth()
+		.VAlign(VAlign_Center)
+		[
+			SNew(SComboButton)
+			.ButtonStyle(FEditorStyle::Get(), "ViewportMenu.Button")
+			.ForegroundColor(FLinearColor::Black)
+			.OnGetMenuContent(this, &SDesignerView::GetResolutionsMenu)
+			.ContentPadding(1.0f)
+			.ButtonContent()
+			[
+				SNew(STextBlock)
+				.Text(LOCTEXT("ScreenSize", "Screen Size"))
+				.TextStyle(FEditorStyle::Get(), "ViewportMenu.Label")
+			]
+		]
+
+		// Screen Fill Size Rule
 		+ SHorizontalBox::Slot()
 		.AutoWidth()
-		.Padding(5.0f)
+		.VAlign(VAlign_Center)
 		[
 			SNew(SComboButton)
 			.ButtonStyle(FEditorStyle::Get(), "ViewportMenu.Button")
 			.ForegroundColor(FLinearColor::Black)
-			.OnGetMenuContent(this, &SDesignerView::GetAspectMenu)
-			.ContentPadding(2.0f)
+			.OnGetMenuContent(this, &SDesignerView::GetScreenSizingFillMenu)
+			.ContentPadding(1.0f)
 			.ButtonContent()
 			[
 				SNew(STextBlock)
-				.Text(LOCTEXT("PreviewSize", "Preview Size"))
+				.Text(this, &SDesignerView::GetScreenSizingFillText)
 				.TextStyle(FEditorStyle::Get(), "ViewportMenu.Label")
 			]
 		]
 
 		+ SHorizontalBox::Slot()
 		.AutoWidth()
-		.Padding(5.0f)
 		.VAlign(VAlign_Center)
+		.Padding(FMargin(2, 0))
 		[
 			SNew(SNumericEntryBox<int32>)
 			.AllowSpin(true)
@@ -420,8 +437,8 @@ TSharedRef<SWidget> SDesignerView::CreateOverlayUI()
 
 		+ SHorizontalBox::Slot()
 		.AutoWidth()
-		.Padding(5.0f)
 		.VAlign(VAlign_Center)
+		.Padding(FMargin(2, 0))
 		[
 			SNew(SNumericEntryBox<int32>)
 			.AllowSpin(true)
@@ -641,38 +658,63 @@ ETransformMode::Type SDesignerView::GetTransformMode() const
 	return TransformMode;
 }
 
-FOptionalSize SDesignerView::GetPreviewWidth() const
+FOptionalSize SDesignerView::GetPreviewAreaWidth() const
 {
-	if ( UUserWidget* DefaultWidget = GetDefaultWidget() )
-	{
-		if ( DefaultWidget->bUseDesignTimeSize )
-		{
-			return DefaultWidget->DesignTimeSize.X;
-		}
-		else if ( DefaultWidget->bUseDesiredSizeAtDesignTime )
-		{
-			return CachedPreviewDesiredSize.X;
-		}
-	}
+	FVector2D Area, Size;
+	GetPreviewAreaAndSize(Area, Size);
+
+	return Area.X;
+}
+
+FOptionalSize SDesignerView::GetPreviewAreaHeight() const
+{
+	FVector2D Area, Size;
+	GetPreviewAreaAndSize(Area, Size);
+
+	return Area.Y;
+}
+
+FOptionalSize SDesignerView::GetPreviewSizeWidth() const
+{
+	FVector2D Area, Size;
+	GetPreviewAreaAndSize(Area, Size);
 
-	return (float)PreviewWidth;
+	return Size.X;
 }
 
-FOptionalSize SDesignerView::GetPreviewHeight() const
+FOptionalSize SDesignerView::GetPreviewSizeHeight() const
 {
+	FVector2D Area, Size;
+	GetPreviewAreaAndSize(Area, Size);
+
+	return Size.Y;
+}
+
+void SDesignerView::GetPreviewAreaAndSize(FVector2D& Area, FVector2D& Size) const
+{
+	Area = FVector2D(PreviewWidth, PreviewHeight);
+	Size = FVector2D(PreviewWidth, PreviewHeight);
+
 	if ( UUserWidget* DefaultWidget = GetDefaultWidget() )
 	{
-		if ( DefaultWidget->bUseDesignTimeSize )
+		switch ( DefaultWidget->DesignSizeMode )
 		{
-			return DefaultWidget->DesignTimeSize.Y;
-		}
-		else if ( DefaultWidget->bUseDesiredSizeAtDesignTime )
-		{
-			return CachedPreviewDesiredSize.Y;
+		case EDesignPreviewSizeMode::Custom:
+			Area = DefaultWidget->DesignTimeSize;
+			Size = DefaultWidget->DesignTimeSize;
+			return;
+		case EDesignPreviewSizeMode::CustomOnScreen:
+			Size = DefaultWidget->DesignTimeSize;
+			return;
+		case EDesignPreviewSizeMode::Desired:
+			Area = CachedPreviewDesiredSize;
+			Size = CachedPreviewDesiredSize;
+			return;
+		case EDesignPreviewSizeMode::DesiredOnScreen:
+			Size = CachedPreviewDesiredSize;
+			return;
 		}
 	}
-
-	return (float)PreviewHeight;
 }
 
 float SDesignerView::GetPreviewDPIScale() const
@@ -680,7 +722,8 @@ float SDesignerView::GetPreviewDPIScale() const
 	// If the user is using a custom size then we disable the DPI scaling logic.
 	if ( UUserWidget* DefaultWidget = GetDefaultWidget() )
 	{
-		if ( DefaultWidget->bUseDesignTimeSize || DefaultWidget->bUseDesiredSizeAtDesignTime )
+		if ( DefaultWidget->DesignSizeMode == EDesignPreviewSizeMode::Custom || 
+			 DefaultWidget->DesignSizeMode == EDesignPreviewSizeMode::Desired )
 		{
 			return 1.0f;
 		}
@@ -691,7 +734,7 @@ float SDesignerView::GetPreviewDPIScale() const
 
 FSlateRect SDesignerView::ComputeAreaBounds() const
 {
-	return FSlateRect(0, 0, GetPreviewWidth().Get(), GetPreviewHeight().Get());
+	return FSlateRect(0, 0, GetPreviewAreaWidth().Get(), GetPreviewAreaHeight().Get());
 }
 
 EVisibility SDesignerView::GetInfoBarVisibility() const
@@ -1350,7 +1393,17 @@ void SDesignerView::UpdatePreviewWidget(bool bForceUpdate)
 			NewPreviewSlateWidget->SlatePrepass();
 
 			PreviewSlateWidget = NewPreviewSlateWidget;
-			PreviewSurface->SetContent(NewPreviewSlateWidget);
+
+			// The constraint box for the widget size needs to inside the DPI scaler in order to make sure it too
+			// is sized accurately for the size screen it's on.
+			TSharedRef<SBox> NewPreviewSizeConstraintBox = SAssignNew(PreviewSizeConstraint, SBox)
+				.WidthOverride(this, &SDesignerView::GetPreviewSizeWidth)
+				.HeightOverride(this, &SDesignerView::GetPreviewSizeHeight)
+				[
+					NewPreviewSlateWidget
+				];
+
+			PreviewSurface->SetContent(NewPreviewSizeConstraintBox);
 
 			// Notify all selected widgets that they are selected, because there are new preview objects
 			// state may have been lost so this will recreate it if the widget does something special when
@@ -1382,6 +1435,18 @@ void SDesignerView::UpdatePreviewWidget(bool bForceUpdate)
 
 void SDesignerView::Tick( const FGeometry& AllottedGeometry, const double InCurrentTime, const float InDeltaTime )
 {
+	UUserWidget* DefaultWidget = GetDefaultWidget();
+	if ( DefaultWidget && ( DefaultWidget->DesignSizeMode == EDesignPreviewSizeMode::CustomOnScreen || DefaultWidget->DesignSizeMode == EDesignPreviewSizeMode::DesiredOnScreen ) )
+	{
+		PreviewAreaConstraint->SetHAlign(HAlign_Left);
+		PreviewAreaConstraint->SetVAlign(VAlign_Top);
+	}
+	else
+	{
+		PreviewAreaConstraint->SetHAlign(HAlign_Fill);
+		PreviewAreaConstraint->SetVAlign(VAlign_Fill);
+	}
+
 	// Tick the parent first to update CachedGeometry
 	SDesignSurface::Tick(AllottedGeometry, InCurrentTime, InDeltaTime);
 
@@ -1440,7 +1505,11 @@ void SDesignerView::Tick( const FGeometry& AllottedGeometry, const double InCurr
 
 	if ( UUserWidget* DefaultWidget = GetDefaultWidget() )
 	{
-		if ( DefaultWidget->bUseDesiredSizeAtDesignTime && PreviewWidget )
+		const bool bNeedDesiredSize = 
+			DefaultWidget->DesignSizeMode == EDesignPreviewSizeMode::Desired || 
+			DefaultWidget->DesignSizeMode == EDesignPreviewSizeMode::DesiredOnScreen;
+
+		if ( bNeedDesiredSize && PreviewWidget )
 		{
 			TSharedPtr<SWidget> CachedWidget = PreviewWidget->GetCachedWidget();
 			if ( CachedWidget.IsValid() )
@@ -1934,7 +2003,11 @@ EVisibility SDesignerView::GetResolutionTextVisibility() const
 	// If we're using a custom design time size, don't bother showing the resolution
 	if ( UUserWidget* DefaultWidget = GetDefaultWidget() )
 	{
-		if ( DefaultWidget->bUseDesignTimeSize || DefaultWidget->bUseDesiredSizeAtDesignTime )
+		const bool bScreenlessSizing =
+			DefaultWidget->DesignSizeMode == EDesignPreviewSizeMode::Custom ||
+			DefaultWidget->DesignSizeMode == EDesignPreviewSizeMode::Desired;
+
+		if ( bScreenlessSizing )
 		{
 			return EVisibility::Collapsed;
 		}
@@ -1970,11 +2043,20 @@ void SDesignerView::HandleOnCommonResolutionSelected(int32 Width, int32 Height,
 	GConfig->SetInt(*ConfigSectionName, TEXT("PreviewHeight"), Height, GEditorUserSettingsIni);
 	GConfig->SetString(*ConfigSectionName, TEXT("PreviewAspectRatio"), *AspectRatio, GEditorUserSettingsIni);
 
-	// We're no longer using a custom design time size.
 	if ( UUserWidget* DefaultWidget = GetDefaultWidget() )
 	{
-		DefaultWidget->bUseDesignTimeSize = false;
-		DefaultWidget->bUseDesiredSizeAtDesignTime = false;
+		// If we using custom or desired design time sizes and the user picks a screen size, he must
+		// want to also change his visualization to be custom on screen or desired on screen, doesn't
+		// make sense to change it otherwise as it would have no effect.
+		if ( DefaultWidget->DesignSizeMode == EDesignPreviewSizeMode::Custom )
+		{
+			DefaultWidget->DesignSizeMode = EDesignPreviewSizeMode::CustomOnScreen;
+		}
+		else if ( DefaultWidget->DesignSizeMode == EDesignPreviewSizeMode::Desired )
+		{
+			DefaultWidget->DesignSizeMode = EDesignPreviewSizeMode::DesiredOnScreen;
+		}
+
 		MarkDesignModifed(/*bRequiresRecompile*/ false);
 	}
 
@@ -1986,7 +2068,7 @@ bool SDesignerView::HandleIsCommonResolutionSelected(int32 Width, int32 Height)
 	// If we're using a custom design time size, none of the other resolutions should appear selected, even if they match.
 	if ( UUserWidget* DefaultWidget = GetDefaultWidget() )
 	{
-		if ( DefaultWidget->bUseDesignTimeSize || DefaultWidget->bUseDesiredSizeAtDesignTime )
+		if ( DefaultWidget->DesignSizeMode == EDesignPreviewSizeMode::Custom || DefaultWidget->DesignSizeMode == EDesignPreviewSizeMode::Desired )
 		{
 			return false;
 		}
@@ -2012,46 +2094,6 @@ void SDesignerView::AddScreenResolutionSection(FMenuBuilder& MenuBuilder, const
 	MenuBuilder.EndSection();
 }
 
-bool SDesignerView::HandleIsCustomResolutionSelected() const
-{
-	if ( UUserWidget* DefaultWidget = GetDefaultWidget() )
-	{
-		return DefaultWidget->bUseDesignTimeSize;
-	}
-
-	return false;
-}
-
-bool SDesignerView::HandleIsDesiredSizeSelected() const
-{
-	if ( UUserWidget* DefaultWidget = GetDefaultWidget() )
-	{
-		return DefaultWidget->bUseDesiredSizeAtDesignTime;
-	}
-
-	return false;
-}
-
-void SDesignerView::HandleOnCustomResolutionSelected()
-{
-	if ( UUserWidget* DefaultWidget = GetDefaultWidget() )
-	{
-		DefaultWidget->bUseDesignTimeSize = true;
-		DefaultWidget->bUseDesiredSizeAtDesignTime = false;
-		MarkDesignModifed(/*bRequiresRecompile*/ false);
-	}
-}
-
-void SDesignerView::HandleOnDesiredSizeSelected()
-{
-	if ( UUserWidget* DefaultWidget = GetDefaultWidget() )
-	{
-		DefaultWidget->bUseDesignTimeSize = false;
-		DefaultWidget->bUseDesiredSizeAtDesignTime = true;
-		MarkDesignModifed(/*bRequiresRecompile*/ false);
-	}
-}
-
 TOptional<int32> SDesignerView::GetCustomResolutionWidth() const
 {
 	if ( UUserWidget* DefaultWidget = GetDefaultWidget() )
@@ -2094,7 +2136,11 @@ EVisibility SDesignerView::GetCustomResolutionEntryVisibility() const
 {
 	if ( UUserWidget* DefaultWidget = GetDefaultWidget() )
 	{
-		return DefaultWidget->bUseDesignTimeSize ? EVisibility::Visible : EVisibility::Collapsed;
+		const bool bCustomSizing =
+			DefaultWidget->DesignSizeMode == EDesignPreviewSizeMode::Custom ||
+			DefaultWidget->DesignSizeMode == EDesignPreviewSizeMode::CustomOnScreen;
+
+		return bCustomSizing ? EVisibility::Visible : EVisibility::Collapsed;
 	}
 
 	return EVisibility::Collapsed;
@@ -2111,36 +2157,78 @@ UUserWidget* SDesignerView::GetDefaultWidget() const
 	return nullptr;
 }
 
-TSharedRef<SWidget> SDesignerView::GetAspectMenu()
+TSharedRef<SWidget> SDesignerView::GetResolutionsMenu()
 {
 	const ULevelEditorPlaySettings* PlaySettings = GetDefault<ULevelEditorPlaySettings>();
 	FMenuBuilder MenuBuilder(true, nullptr);
 
-	// Add custom option
-	FUIAction CustomSizeAction(
-		FExecuteAction::CreateRaw(this, &SDesignerView::HandleOnCustomResolutionSelected),
-		FCanExecuteAction(),
-		FIsActionChecked::CreateRaw(this, &SDesignerView::HandleIsCustomResolutionSelected));
+	// Add the normal set of resolution options.
+	AddScreenResolutionSection(MenuBuilder, PlaySettings->PhoneScreenResolutions, LOCTEXT("CommonPhonesSectionHeader", "Phones"));
+	AddScreenResolutionSection(MenuBuilder, PlaySettings->TabletScreenResolutions, LOCTEXT("CommonTabletsSectionHeader", "Tablets"));
+	AddScreenResolutionSection(MenuBuilder, PlaySettings->LaptopScreenResolutions, LOCTEXT("CommonLaptopsSectionHeader", "Laptops"));
+	AddScreenResolutionSection(MenuBuilder, PlaySettings->MonitorScreenResolutions, LOCTEXT("CommoMonitorsSectionHeader", "Monitors"));
+	AddScreenResolutionSection(MenuBuilder, PlaySettings->TelevisionScreenResolutions, LOCTEXT("CommonTelevesionsSectionHeader", "Televisions"));
+
+	return MenuBuilder.MakeWidget();
+}
 
-	MenuBuilder.AddMenuEntry(LOCTEXT("Custom", "Custom"), LOCTEXT("Custom", "Custom"), FSlateIcon(), CustomSizeAction, NAME_None, EUserInterfaceActionType::Check);
+TSharedRef<SWidget> SDesignerView::GetScreenSizingFillMenu()
+{
+	const ULevelEditorPlaySettings* PlaySettings = GetDefault<ULevelEditorPlaySettings>();
+	FMenuBuilder MenuBuilder(true, nullptr);
+
+	CreateScreenFillEntry(MenuBuilder, EDesignPreviewSizeMode::FillScreen);
+	CreateScreenFillEntry(MenuBuilder, EDesignPreviewSizeMode::Custom);
+	CreateScreenFillEntry(MenuBuilder, EDesignPreviewSizeMode::CustomOnScreen);
+	CreateScreenFillEntry(MenuBuilder, EDesignPreviewSizeMode::Desired);
+	CreateScreenFillEntry(MenuBuilder, EDesignPreviewSizeMode::DesiredOnScreen);
+
+	return MenuBuilder.MakeWidget();
+}
+
+void SDesignerView::CreateScreenFillEntry(FMenuBuilder& MenuBuilder, EDesignPreviewSizeMode SizeMode)
+{
+	const static UEnum* PreviewSizeEnum = FindObject<UEnum>(ANY_PACKAGE, TEXT("EDesignPreviewSizeMode"), true);
 
 	// Add desired size option
 	FUIAction DesiredSizeAction(
-		FExecuteAction::CreateRaw(this, &SDesignerView::HandleOnDesiredSizeSelected),
+		FExecuteAction::CreateRaw(this, &SDesignerView::OnScreenFillRuleSelected, SizeMode),
 		FCanExecuteAction(),
-		FIsActionChecked::CreateRaw(this, &SDesignerView::HandleIsDesiredSizeSelected));
+		FIsActionChecked::CreateRaw(this, &SDesignerView::GetIsScreenFillRuleSelected, SizeMode));
 
+	FText EntryText = PreviewSizeEnum->GetEnumText((int32)SizeMode);
+	MenuBuilder.AddMenuEntry(EntryText, FText::GetEmpty(), FSlateIcon(), DesiredSizeAction, NAME_None, EUserInterfaceActionType::Check);
+}
 
-	MenuBuilder.AddMenuEntry(LOCTEXT("DesiredSize", "Desired Size"), LOCTEXT("DesiredSize", "Desired Size"), FSlateIcon(), DesiredSizeAction, NAME_None, EUserInterfaceActionType::Check);
+FText SDesignerView::GetScreenSizingFillText() const
+{
+	const static UEnum* PreviewSizeEnum = FindObject<UEnum>(ANY_PACKAGE, TEXT("EDesignPreviewSizeMode"), true);
 
-	// Add the normal set of resultion options.
-	AddScreenResolutionSection(MenuBuilder, PlaySettings->PhoneScreenResolutions, LOCTEXT("CommonPhonesSectionHeader", "Phones"));
-	AddScreenResolutionSection(MenuBuilder, PlaySettings->TabletScreenResolutions, LOCTEXT("CommonTabletsSectionHeader", "Tablets"));
-	AddScreenResolutionSection(MenuBuilder, PlaySettings->LaptopScreenResolutions, LOCTEXT("CommonLaptopsSectionHeader", "Laptops"));
-	AddScreenResolutionSection(MenuBuilder, PlaySettings->MonitorScreenResolutions, LOCTEXT("CommoMonitorsSectionHeader", "Monitors"));
-	AddScreenResolutionSection(MenuBuilder, PlaySettings->TelevisionScreenResolutions, LOCTEXT("CommonTelevesionsSectionHeader", "Televisions"));
+	if ( UUserWidget* DefaultWidget = GetDefaultWidget() )
+	{
+		return PreviewSizeEnum->GetEnumText((int32)DefaultWidget->DesignSizeMode);
+	}
 
-	return MenuBuilder.MakeWidget();
+	return FText::GetEmpty();
+}
+
+bool SDesignerView::GetIsScreenFillRuleSelected(EDesignPreviewSizeMode SizeMode) const
+{
+	if ( UUserWidget* DefaultWidget = GetDefaultWidget() )
+	{
+		return DefaultWidget->DesignSizeMode == SizeMode;
+	}
+
+	return false;
+}
+
+void SDesignerView::OnScreenFillRuleSelected(EDesignPreviewSizeMode SizeMode)
+{
+	if ( UUserWidget* DefaultWidget = GetDefaultWidget() )
+	{
+		DefaultWidget->DesignSizeMode = SizeMode;
+		MarkDesignModifed(/*bRequiresRecompile*/ false);
+	}
 }
 
 void SDesignerView::BeginTransaction(const FText& SessionName)
diff --git a/Engine/Source/Editor/UMGEditor/Private/Designer/SDesignerView.h b/Engine/Source/Editor/UMGEditor/Private/Designer/SDesignerView.h
--- a/Engine/Source/Editor/UMGEditor/Private/Designer/SDesignerView.h
+++ b/Engine/Source/Editor/UMGEditor/Private/Designer/SDesignerView.h
@@ -81,10 +81,18 @@ private:
 	void SetStartupResolution();
 
 	/** The width of the preview screen for the UI */
-	FOptionalSize GetPreviewWidth() const;
+	FOptionalSize GetPreviewAreaWidth() const;
 
 	/** The height of the preview screen for the UI */
-	FOptionalSize GetPreviewHeight() const;
+	FOptionalSize GetPreviewAreaHeight() const;
+
+	/** The width of the preview widget for the UI */
+	FOptionalSize GetPreviewSizeWidth() const;
+
+	/** The height of the preview widget for the UI */
+	FOptionalSize GetPreviewSizeHeight() const;
+
+	void GetPreviewAreaAndSize(FVector2D& Area, FVector2D& Size) const;
 
 	/** Gets the DPI scale that would be applied given the current preview width and height */
 	float GetPreviewDPIScale() const;
@@ -140,11 +148,13 @@ private:
 	void HandleOnCommonResolutionSelected(int32 Width, int32 Height, FString AspectRatio);
 	bool HandleIsCommonResolutionSelected(int32 Width, int32 Height) const;
 	void AddScreenResolutionSection(FMenuBuilder& MenuBuilder, const TArray<FPlayScreenResolution>& Resolutions, const FText& SectionName);
-	bool HandleIsCustomResolutionSelected() const;
-	void HandleOnCustomResolutionSelected();
-	bool HandleIsDesiredSizeSelected() const;
-	void HandleOnDesiredSizeSelected();
-	TSharedRef<SWidget> GetAspectMenu();
+	TSharedRef<SWidget> GetResolutionsMenu();
+
+	TSharedRef<SWidget> GetScreenSizingFillMenu();
+	void CreateScreenFillEntry(FMenuBuilder& MenuBuilder, EDesignPreviewSizeMode SizeMode);
+	FText GetScreenSizingFillText() const;
+	bool GetIsScreenFillRuleSelected(EDesignPreviewSizeMode SizeMode) const;
+	void OnScreenFillRuleSelected(EDesignPreviewSizeMode SizeMode);
 
 	EVisibility PIENotification() const;
 
@@ -220,6 +230,8 @@ private:
 	UPanelWidget* DropPreviewParent;
 
 	TSharedPtr<class SZoomPan> PreviewHitTestRoot;
+	TSharedPtr<SBox> PreviewAreaConstraint;
+	TSharedPtr<SBox> PreviewSizeConstraint;
 	TSharedPtr<SDPIScaler> PreviewSurface;
 	TSharedPtr<SCanvas> ExtensionWidgetCanvas;
 	TSharedPtr<SPaintSurface> EffectsLayer;
diff --git a/Engine/Source/Runtime/Core/Public/UObject/ObjectVersion.h b/Engine/Source/Runtime/Core/Public/UObject/ObjectVersion.h
--- a/Engine/Source/Runtime/Core/Public/UObject/ObjectVersion.h
+++ b/Engine/Source/Runtime/Core/Public/UObject/ObjectVersion.h
@@ -488,6 +488,8 @@ enum EUnrealEngineObjectUE4Version
 	VER_UE4_TRACK_UCS_MODIFIED_PROPERTIES,
 	// Allowed landscape spline meshes to be stored into landscape streaming levels rather than the spline's level
 	VER_UE4_LANDSCAPE_SPLINE_CROSS_LEVEL_MESHES,
+	// Deprecate the variables used for sizing in the designer on UUserWidget
+	VER_UE4_DEPRECATE_USER_WIDGET_DESIGN_SIZE,
 
 	// -----<new versions can be added before this line>-------------------------------------------------
 	// - this needs to be the last line (see note below)
diff --git a/Engine/Source/Runtime/UMG/Private/UserWidget.cpp b/Engine/Source/Runtime/UMG/Private/UserWidget.cpp
--- a/Engine/Source/Runtime/UMG/Private/UserWidget.cpp
+++ b/Engine/Source/Runtime/UMG/Private/UserWidget.cpp
@@ -26,9 +26,11 @@ UUserWidget::UUserWidget(const FObjectInitializer& ObjectInitializer)
 	ForegroundColor = FSlateColor::UseForeground();
 
 #if WITH_EDITORONLY_DATA
-	bUseDesignTimeSize = false;
+	bUseDesignTimeSize_DEPRECATED = false;
+	bUseDesiredSizeAtDesignTime_DEPRECATED = false;
 	DesignTimeSize = FVector2D(100, 100);
 	PaletteCategory = LOCTEXT("UserCreated", "User Created");
+	DesignSizeMode = EDesignPreviewSizeMode::FillScreen;
 #endif
 }
 
@@ -684,6 +686,18 @@ void UUserWidget::NativeTick(const FGeometry& MyGeometry, float InDeltaTime)
 {
 	GInitRunaway();
 
+	TickActionsAndAnimation(MyGeometry, InDeltaTime);
+
+	Tick(MyGeometry, InDeltaTime);
+}
+
+void UUserWidget::TickActionsAndAnimation(const FGeometry& MyGeometry, float InDeltaTime)
+{
+	if ( bDesignTime )
+	{
+		return;
+	}
+
 	// Update active movie scenes
 	for ( UUMGSequencePlayer* Player : ActiveSequencePlayers )
 	{
@@ -698,18 +712,13 @@ void UUserWidget::NativeTick(const FGeometry& MyGeometry, float InDeltaTime)
 
 	StoppedSequencePlayers.Empty();
 
-	if ( !bDesignTime )
+	UWorld* World = GetWorld();
+	if ( World )
 	{
-		UWorld* World = GetWorld();
-		if ( World )
-		{
-			// Update any latent actions we have for this actor
-			FLatentActionManager& LatentActionManager = World->GetLatentActionManager();
-			LatentActionManager.ProcessLatentActions(this, InDeltaTime);
-		}
+		// Update any latent actions we have for this actor
+		FLatentActionManager& LatentActionManager = World->GetLatentActionManager();
+		LatentActionManager.ProcessLatentActions(this, InDeltaTime);
 	}
-
-	Tick(MyGeometry, InDeltaTime);
 }
 
 void UUserWidget::NativePaint( FPaintContext& InContext ) const 
@@ -847,6 +856,23 @@ FReply UUserWidget::NativeOnMotionDetected( const FGeometry& InGeometry, const F
 	return OnMotionDetected( InGeometry, InMotionEvent ).NativeReply;
 }
 
+void UUserWidget::PostLoad()
+{
+	Super::PostLoad();
+
+	if ( GetLinkerUE4Version() < VER_UE4_DEPRECATE_USER_WIDGET_DESIGN_SIZE )
+	{
+		if ( bUseDesignTimeSize_DEPRECATED )
+		{
+			DesignSizeMode = EDesignPreviewSizeMode::Custom;
+		}
+		else if ( bUseDesiredSizeAtDesignTime_DEPRECATED )
+		{
+			DesignSizeMode = EDesignPreviewSizeMode::Desired;
+		}
+	}
+}
+
 /////////////////////////////////////////////////////
 //
 // DEPRECATED
diff --git a/Engine/Source/Runtime/UMG/Public/Blueprint/UserWidget.h b/Engine/Source/Runtime/UMG/Public/Blueprint/UserWidget.h
--- a/Engine/Source/Runtime/UMG/Public/Blueprint/UserWidget.h
+++ b/Engine/Source/Runtime/UMG/Public/Blueprint/UserWidget.h
@@ -116,6 +116,20 @@ namespace EUMGSequencePlayMode
 	};
 }
 
+#if WITH_EDITORONLY_DATA
+
+UENUM()
+enum class EDesignPreviewSizeMode : uint8
+{
+	FillScreen,
+	Custom,
+	CustomOnScreen,
+	Desired,
+	DesiredOnScreen,
+};
+
+#endif
+
 //TODO UMG If you want to host a widget that's full screen there may need to be a SWindow equivalent that you spawn it into.
 
 DECLARE_DYNAMIC_MULTICAST_DELEGATE(FOnConstructEvent);
@@ -137,6 +151,7 @@ public:
 	virtual void PostEditImport() override;
 	virtual void PostDuplicate(bool bDuplicateForPIE) override;
 	virtual void BeginDestroy() override;
+	virtual void PostLoad() override;
 	// End of UObject interface
 
 	void Initialize();
@@ -755,7 +770,7 @@ public:
 	UPROPERTY()
 	FGetSlateColor ForegroundColorDelegate;
 
-	/** Called when the visibility changes. */
+	/** Setting this flag to true, allows this widget to accept focus when clicked, or when navigated to. */
 	UPROPERTY(EditAnywhere, BlueprintReadWrite, AdvancedDisplay, Category="Behavior")
 	bool bSupportsKeyboardFocus;
 
@@ -783,11 +798,14 @@ public:
 
 	/** A flag that determines if the design time size is used for previewing the widget in the designer. */
 	UPROPERTY()
-	bool bUseDesignTimeSize;
+	bool bUseDesignTimeSize_DEPRECATED;
 
 	/** A flag that determines if the widget's desired size is used for previewing the widget in the designer. */
 	UPROPERTY()
-	bool bUseDesiredSizeAtDesignTime;
+	bool bUseDesiredSizeAtDesignTime_DEPRECATED;
+
+	UPROPERTY()
+	EDesignPreviewSizeMode DesignSizeMode;
 
 	/** The category this widget appears in the palette. */
 	UPROPERTY()
@@ -846,6 +864,12 @@ protected:
 	virtual FReply NativeOnTouchEnded( const FGeometry& InGeometry, const FPointerEvent& InGestureEvent );
 	virtual FReply NativeOnMotionDetected( const FGeometry& InGeometry, const FMotionEvent& InMotionEvent );
 
+protected:
+	/**
+	 * Ticks the active sequences and latent actions that have been scheduled for this Widget.
+	 */
+	void TickActionsAndAnimation(const FGeometry& MyGeometry, float InDeltaTime);
+
 private:
 	FAnchors ViewportAnchors;
 	FMargin ViewportOffsets;
