UMG - Making a large pass and fixing all the places properties were incorrectly set to EditDefaultsOnly, they should have been set to EditAnywhere.  Removing some old functions from UWidget for gathering children widgets.  They were outdated, and wouldn't work correctly in cases like named slots.  Copying a widget with named slots and pasting it now correctly recreates items that were contained in named slots.

diff --git a/Engine/Source/Editor/UMGEditor/Private/Details/SWidgetDetailsView.cpp b/Engine/Source/Editor/UMGEditor/Private/Details/SWidgetDetailsView.cpp
--- a/Engine/Source/Editor/UMGEditor/Private/Details/SWidgetDetailsView.cpp
+++ b/Engine/Source/Editor/UMGEditor/Private/Details/SWidgetDetailsView.cpp
@@ -38,7 +38,16 @@ void SWidgetDetailsView::Construct(const FArguments& InArgs, TSharedPtr<FWidgetB
 	FPropertyEditorModule& EditModule = FModuleManager::Get().GetModuleChecked<FPropertyEditorModule>("PropertyEditor");
 
 	FNotifyHook* NotifyHook = this;
-	FDetailsViewArgs DetailsViewArgs( /*bUpdateFromSelection=*/ false, /*bLockable=*/ false, /*bAllowSearch=*/ true, FDetailsViewArgs::HideNameArea, /*bHideSelectionTip=*/ true, /*InNotifyHook=*/ NotifyHook, /*InSearchInitialKeyFocus=*/ false, /*InViewIdentifier=*/ NAME_None);
+	FDetailsViewArgs DetailsViewArgs(
+		/*bUpdateFromSelection=*/ false,
+		/*bLockable=*/ false,
+		/*bAllowSearch=*/ true,
+		FDetailsViewArgs::HideNameArea,
+		/*bHideSelectionTip=*/ true,
+		/*InNotifyHook=*/ NotifyHook,
+		/*InSearchInitialKeyFocus=*/ false,
+		/*InViewIdentifier=*/ NAME_None);
+	DetailsViewArgs.DefaultsOnlyVisibility = FDetailsViewArgs::EEditDefaultsOnlyNodeVisibility::Automatic;
 
 	PropertyView = EditModule.CreateDetailView(DetailsViewArgs);
 
@@ -116,13 +125,11 @@ void SWidgetDetailsView::Construct(const FArguments& InArgs, TSharedPtr<FWidgetB
 			.AutoWidth()
 			[
 				SNew(SCheckBox)
-				//.Style(FEditorStyle::Get(), "ToggleButtonCheckbox")
 				.IsChecked(this, &SWidgetDetailsView::GetIsVariable)
 				.OnCheckStateChanged(this, &SWidgetDetailsView::HandleIsVariableChanged)
 				.Padding(FMargin(3,1,3,1))
 				[
 					SNew(STextBlock)
-					//.Font(FEditorStyle::GetFontStyle("BoldFont"))
 					.Text(LOCTEXT("IsVariable", "Is Variable"))
 				]
 			]
@@ -202,9 +209,11 @@ void SWidgetDetailsView::OnEditorSelectionChanging()
 
 void SWidgetDetailsView::OnEditorSelectionChanged()
 {
+	// Clear selection in the property view.
 	SelectedObjects.Empty();
 	PropertyView->SetObjects(SelectedObjects);
 
+	// Add any selected widgets to the list of pending selected objects.
 	TSet< FWidgetReference > SelectedWidgets = BlueprintEditor.Pin()->GetSelectedWidgets();
 	if ( SelectedWidgets.Num() > 0 )
 	{
@@ -214,6 +223,7 @@ void SWidgetDetailsView::OnEditorSelectionChanged()
 		}
 	}
 
+	// Add any selected objects (non-widgets) to the pending selected objects.
 	TSet< TWeakObjectPtr<UObject> > Selection = BlueprintEditor.Pin()->GetSelectedObjects();
 	for ( TWeakObjectPtr<UObject> Selected : Selection )
 	{
@@ -223,6 +233,7 @@ void SWidgetDetailsView::OnEditorSelectionChanged()
 		}
 	}
 
+	// If only 1 valid selected object exists, update the class link to point to the right class.
 	if ( SelectedObjects.Num() == 1 && SelectedObjects[0].IsValid() )
 	{
 		ClassLinkArea->SetContent(FEditorClassUtils::GetSourceLink(SelectedObjects[0]->GetClass(), TWeakObjectPtr<UObject>()));
@@ -232,6 +243,12 @@ void SWidgetDetailsView::OnEditorSelectionChanged()
 		ClassLinkArea->SetContent(SNullWidget::NullWidget);
 	}
 
+	//if ( IsWidgetCDOSelected() )
+	//{
+	//	//PropertyView->Defau
+	//}
+
+	// Update the preview view to look at the current selection set.
 	const bool bForceRefresh = false;
 	PropertyView->SetObjects(SelectedObjects, bForceRefresh);
 }
@@ -252,32 +269,28 @@ void SWidgetDetailsView::ClearFocusIfOwned()
 	}
 }
 
-EVisibility SWidgetDetailsView::GetNameAreaVisibility() const
+bool SWidgetDetailsView::IsWidgetCDOSelected() const
 {
 	if ( SelectedObjects.Num() == 1 )
 	{
 		UWidget* Widget = Cast<UWidget>(SelectedObjects[0].Get());
 		if ( Widget && !Widget->HasAnyFlags(RF_ClassDefaultObject) )
 		{
-			return EVisibility::Visible;
+			return true;
 		}
 	}
-	
-	return EVisibility::Collapsed;
+
+	return false;
 }
 
-EVisibility SWidgetDetailsView::GetCategoryAreaVisibility() const
+EVisibility SWidgetDetailsView::GetNameAreaVisibility() const
 {
-	if ( SelectedObjects.Num() == 1 )
-	{
-		UUserWidget* Widget = Cast<UUserWidget>(SelectedObjects[0].Get());
-		if ( Widget && Widget->HasAnyFlags(RF_ClassDefaultObject) )
-		{
-			return EVisibility::Visible;
-		}
-	}
+	return IsWidgetCDOSelected() ? EVisibility::Visible : EVisibility::Collapsed;
+}
 
-	return EVisibility::Collapsed;
+EVisibility SWidgetDetailsView::GetCategoryAreaVisibility() const
+{
+	return IsWidgetCDOSelected() ? EVisibility::Collapsed : EVisibility::Visible;
 }
 
 void SWidgetDetailsView::HandleCategoryTextCommitted(const FText& Text, ETextCommit::Type CommitType)
diff --git a/Engine/Source/Editor/UMGEditor/Private/Details/SWidgetDetailsView.h b/Engine/Source/Editor/UMGEditor/Private/Details/SWidgetDetailsView.h
--- a/Engine/Source/Editor/UMGEditor/Private/Details/SWidgetDetailsView.h
+++ b/Engine/Source/Editor/UMGEditor/Private/Details/SWidgetDetailsView.h
@@ -41,6 +41,8 @@ private:
 
 	void ClearFocusIfOwned();
 
+	bool IsWidgetCDOSelected() const;
+
 	EVisibility GetNameAreaVisibility() const;
 
 	const FSlateBrush* GetNameIcon() const;
diff --git a/Engine/Source/Editor/UMGEditor/Private/WidgetBlueprintEditorUtils.cpp b/Engine/Source/Editor/UMGEditor/Private/WidgetBlueprintEditorUtils.cpp
--- a/Engine/Source/Editor/UMGEditor/Private/WidgetBlueprintEditorUtils.cpp
+++ b/Engine/Source/Editor/UMGEditor/Private/WidgetBlueprintEditorUtils.cpp
@@ -305,7 +305,7 @@ void FWidgetBlueprintEditorUtils::DeleteWidgets(UWidgetBlueprint* BP, TSet<FWidg
 
 			// Rename all child widgets as well, to the transient package so that they don't conflict with future widgets sharing the same name.
 			TArray<UWidget*> ChildWidgets;
-			BP->WidgetTree->GetChildWidgets(WidgetTemplate, ChildWidgets);
+			UWidgetTree::GetChildWidgets(WidgetTemplate, ChildWidgets);
 			for ( UWidget* Widget : ChildWidgets )
 			{
 				Widget->Rename(nullptr, nullptr);
@@ -503,7 +503,7 @@ void FWidgetBlueprintEditorUtils::CutWidgets(UWidgetBlueprint* BP, TSet<FWidgetR
 
 void FWidgetBlueprintEditorUtils::CopyWidgets(UWidgetBlueprint* BP, TSet<FWidgetReference> Widgets)
 {
-	TSet<UWidget*> CopyableWidets;
+	TArray<UWidget*> CopyableWidets;
 	for ( const FWidgetReference& Widget : Widgets )
 	{
 		UWidget* ParentWidget = Widget.GetTemplate();
@@ -511,11 +511,13 @@ void FWidgetBlueprintEditorUtils::CopyWidgets(UWidgetBlueprint* BP, TSet<FWidget
 
 		// When copying a widget users expect all sub widgets to be copied as well, so we need to ensure that
 		// we gather all the child widgets and copy them as well.
-		UWidget::GatherAllChildren(ParentWidget, CopyableWidets);
+		UWidgetTree::GetChildWidgets(ParentWidget, CopyableWidets);
 	}
 
+	TSet<UWidget*> CopyableWidetsSet(CopyableWidets);
+
 	FString ExportedText;
-	FWidgetBlueprintEditorUtils::ExportWidgetsToText(CopyableWidets, /*out*/ ExportedText);
+	FWidgetBlueprintEditorUtils::ExportWidgetsToText(CopyableWidetsSet, /*out*/ ExportedText);
 	FPlatformMisc::ClipboardCopy(*ExportedText);
 }
 
@@ -538,6 +540,21 @@ void FWidgetBlueprintEditorUtils::ExportWidgetsToText(TSet<UWidget*> WidgetsToEx
 
 		UExporter::ExportToOutputDevice(&Context, Widget, nullptr, Archive, TEXT("copy"), 0, PPF_ExportsNotFullyQualified | PPF_Copy | PPF_Delimited, false, ThisOuter);
 
+		// Check to see if this widget was content of another widget holding it in a named slot.
+		if ( Widget->GetParent() == nullptr )
+		{
+			for ( UWidget* ExportableWidget : WidgetsToExport )
+			{
+				if ( INamedSlotInterface* NamedSlotContainer = Cast<INamedSlotInterface>(ExportableWidget) )
+				{
+					if ( NamedSlotContainer->ContainsContent(Widget) )
+					{
+						continue;
+					}
+				}
+			}
+		}
+
 		if ( Widget->GetParent() == nullptr || !WidgetsToExport.Contains(Widget->GetParent()) )
 		{
 			auto SlotMetaData = NewObject<UWidgetSlotPair>();
diff --git a/Engine/Source/Runtime/UMG/Private/Components/NamedSlotInterface.cpp b/Engine/Source/Runtime/UMG/Private/Components/NamedSlotInterface.cpp
--- a/Engine/Source/Runtime/UMG/Private/Components/NamedSlotInterface.cpp
+++ b/Engine/Source/Runtime/UMG/Private/Components/NamedSlotInterface.cpp
@@ -6,3 +6,19 @@ UNamedSlotInterface::UNamedSlotInterface(const FObjectInitializer& ObjectInitial
 	: Super(ObjectInitializer)
 {
 }
+
+bool INamedSlotInterface::ContainsContent(UWidget* Content) const
+{
+	TArray<FName> SlotNames;
+	GetSlotNames(SlotNames);
+
+	for ( const FName& SlotName : SlotNames )
+	{
+		if ( GetContentForSlot(SlotName) == Content )
+		{
+			return true;
+		}
+	}
+
+	return false;
+}
diff --git a/Engine/Source/Runtime/UMG/Private/Components/Widget.cpp b/Engine/Source/Runtime/UMG/Private/Components/Widget.cpp
--- a/Engine/Source/Runtime/UMG/Private/Components/Widget.cpp
+++ b/Engine/Source/Runtime/UMG/Private/Components/Widget.cpp
@@ -790,34 +790,6 @@ FSizeParam UWidget::ConvertSerializedSizeParamToRuntime(const FSlateChildSize& I
 	return FAuto();
 }
 
-void UWidget::GatherChildren(UWidget* Root, TSet<UWidget*>& Children)
-{
-	UPanelWidget* PanelRoot = Cast<UPanelWidget>(Root);
-	if ( PanelRoot )
-	{
-		for ( int32 ChildIndex = 0; ChildIndex < PanelRoot->GetChildrenCount(); ChildIndex++ )
-		{
-			UWidget* ChildWidget = PanelRoot->GetChildAt(ChildIndex);
-			Children.Add(ChildWidget);
-		}
-	}
-}
-
-void UWidget::GatherAllChildren(UWidget* Root, TSet<UWidget*>& Children)
-{
-	UPanelWidget* PanelRoot = Cast<UPanelWidget>(Root);
-	if ( PanelRoot )
-	{
-		for ( int32 ChildIndex = 0; ChildIndex < PanelRoot->GetChildrenCount(); ChildIndex++ )
-		{
-			UWidget* ChildWidget = PanelRoot->GetChildAt(ChildIndex);
-			Children.Add(ChildWidget);
-
-			GatherAllChildren(ChildWidget, Children);
-		}
-	}
-}
-
 UWidget* UWidget::FindChildContainingDescendant(UWidget* Root, UWidget* Descendant)
 {
 	if ( Root == nullptr )
diff --git a/Engine/Source/Runtime/UMG/Private/WidgetTree.cpp b/Engine/Source/Runtime/UMG/Private/WidgetTree.cpp
--- a/Engine/Source/Runtime/UMG/Private/WidgetTree.cpp
+++ b/Engine/Source/Runtime/UMG/Private/WidgetTree.cpp
@@ -91,7 +91,7 @@ void UWidgetTree::GetAllWidgets(TArray<UWidget*>& Widgets) const
 	ForEachWidget([&] (UWidget* Widget) { Widgets.Add(Widget); });
 }
 
-void UWidgetTree::GetChildWidgets(UWidget* Parent, TArray<UWidget*>& Widgets) const
+void UWidgetTree::GetChildWidgets(UWidget* Parent, TArray<UWidget*>& Widgets)
 {
 	ForWidgetAndChildren(Parent, [&] (UWidget* Widget) { Widgets.Add(Widget); });
 }
diff --git a/Engine/Source/Runtime/UMG/Public/Blueprint/UserWidget.h b/Engine/Source/Runtime/UMG/Public/Blueprint/UserWidget.h
--- a/Engine/Source/Runtime/UMG/Public/Blueprint/UserWidget.h
+++ b/Engine/Source/Runtime/UMG/Public/Blueprint/UserWidget.h
@@ -754,7 +754,7 @@ public:
 
 public:
 	/** The color and opacity of this widget.  Tints all child widgets. */
-	UPROPERTY(EditDefaultsOnly, BlueprintReadOnly, Category="Style")
+	UPROPERTY(EditAnywhere, BlueprintReadOnly, Category="Style")
 	FLinearColor ColorAndOpacity;
 
 	UPROPERTY()
@@ -764,7 +764,7 @@ public:
 	 * The foreground color of the widget, this is inherited by sub widgets.  Any color property
 	 * that is marked as inherit will use this color.
 	 */
-	UPROPERTY(EditDefaultsOnly, BlueprintReadOnly, Category="Style")
+	UPROPERTY(EditAnywhere, BlueprintReadOnly, Category="Style")
 	FSlateColor ForegroundColor;
 
 	UPROPERTY()
@@ -811,6 +811,13 @@ public:
 	UPROPERTY()
 	FText PaletteCategory;
 
+	/**
+	 * A preview background that you can use when designing the UI to get a sense of scale on the screen.  Use
+	 * a texture with a screenshot of your game in it, for example if you were designing a HUD.
+	 */
+	UPROPERTY(EditDefaultsOnly, Category="Designer")
+	UTexture2D* PreviewBackground;
+
 #endif
 
 protected:
diff --git a/Engine/Source/Runtime/UMG/Public/Blueprint/WidgetTree.h b/Engine/Source/Runtime/UMG/Public/Blueprint/WidgetTree.h
--- a/Engine/Source/Runtime/UMG/Public/Blueprint/WidgetTree.h
+++ b/Engine/Source/Runtime/UMG/Public/Blueprint/WidgetTree.h
@@ -42,7 +42,7 @@ public:
 	void GetAllWidgets(TArray<UWidget*>& Widgets) const;
 
 	/** Gathers descendant child widgets of a parent widget. */
-	void GetChildWidgets(UWidget* Parent, TArray<UWidget*>& Widgets) const;
+	static void GetChildWidgets(UWidget* Parent, TArray<UWidget*>& Widgets);
 
 	/**
 	 * Iterates through all widgets including widgets contained in named slots, other than
@@ -66,7 +66,7 @@ public:
 	 * inside another user widget.
 	 */
 	template <typename Predicate>
-	FORCEINLINE void ForWidgetAndChildren(UWidget* Widget, Predicate Pred) const
+	static FORCEINLINE void ForWidgetAndChildren(UWidget* Widget, Predicate Pred)
 	{
 		// Search for any named slot with content that we need to dive into.
 		if ( INamedSlotInterface* NamedSlotHost = Cast<INamedSlotInterface>(Widget) )
diff --git a/Engine/Source/Runtime/UMG/Public/Components/Border.h b/Engine/Source/Runtime/UMG/Public/Components/Border.h
--- a/Engine/Source/Runtime/UMG/Public/Components/Border.h
+++ b/Engine/Source/Runtime/UMG/Public/Components/Border.h
@@ -21,7 +21,7 @@ class UMG_API UBorder : public UContentWidget
 public:
 
 	/** Color and opacity multiplier of content in the border */
-	UPROPERTY(EditDefaultsOnly, BlueprintReadOnly, Category="Content")
+	UPROPERTY(EditAnywhere, BlueprintReadOnly, Category="Content")
 	FLinearColor ContentColorAndOpacity;
 
 	/** A bindable delegate for the ContentColorAndOpacity. */
@@ -29,19 +29,19 @@ public:
 	FGetLinearColor ContentColorAndOpacityDelegate;
 
 	/** The padding area between the slot and the content it contains. */
-	UPROPERTY(EditDefaultsOnly, BlueprintReadOnly, Category="Content")
+	UPROPERTY(EditAnywhere, BlueprintReadOnly, Category="Content")
 	FMargin Padding;
 
 	/** The alignment of the content horizontally. */
-	UPROPERTY(EditDefaultsOnly, BlueprintReadOnly, Category="Content")
+	UPROPERTY(EditAnywhere, BlueprintReadOnly, Category="Content")
 	TEnumAsByte<EHorizontalAlignment> HorizontalAlignment;
 
 	/** The alignment of the content vertically. */
-	UPROPERTY(EditDefaultsOnly, BlueprintReadOnly, Category="Content")
+	UPROPERTY(EditAnywhere, BlueprintReadOnly, Category="Content")
 	TEnumAsByte<EVerticalAlignment> VerticalAlignment;
 
 	/** Brush to drag as the background */
-	UPROPERTY(EditDefaultsOnly, BlueprintReadOnly, Category=Appearance, meta=( DisplayName="Brush" ))
+	UPROPERTY(EditAnywhere, BlueprintReadOnly, Category=Appearance, meta=( DisplayName="Brush" ))
 	FSlateBrush Background;
 
 	/** A bindable delegate for the Brush. */
@@ -49,7 +49,7 @@ public:
 	FGetSlateBrush BackgroundDelegate;
 
 	/** Color and opacity of the actual border image */
-	UPROPERTY(EditDefaultsOnly, BlueprintReadOnly, Category=Appearance)
+	UPROPERTY(EditAnywhere, BlueprintReadOnly, Category=Appearance)
 	FLinearColor BrushColor;
 
 	/** A bindable delegate for the BrushColor. */
@@ -57,21 +57,21 @@ public:
 	FGetLinearColor BrushColorDelegate;
 
 	/** Whether or not to show the disabled effect when this border is disabled */
-	UPROPERTY(EditDefaultsOnly, BlueprintReadOnly, Category=Appearance, AdvancedDisplay)
+	UPROPERTY(EditAnywhere, BlueprintReadOnly, Category=Appearance, AdvancedDisplay)
 	bool bShowEffectWhenDisabled;
 
 public:
 
-	UPROPERTY(EditDefaultsOnly, Category=Events)
+	UPROPERTY(EditAnywhere, Category=Events)
 	FOnPointerEvent OnMouseButtonDownEvent;
 
-	UPROPERTY(EditDefaultsOnly, Category=Events)
+	UPROPERTY(EditAnywhere, Category=Events)
 	FOnPointerEvent OnMouseButtonUpEvent;
 
-	UPROPERTY(EditDefaultsOnly, Category=Events)
+	UPROPERTY(EditAnywhere, Category=Events)
 	FOnPointerEvent OnMouseMoveEvent;
 
-	UPROPERTY(EditDefaultsOnly, Category=Events)
+	UPROPERTY(EditAnywhere, Category=Events)
 	FOnPointerEvent OnMouseDoubleClickEvent;
 
 public:
diff --git a/Engine/Source/Runtime/UMG/Public/Components/BorderSlot.h b/Engine/Source/Runtime/UMG/Public/Components/BorderSlot.h
--- a/Engine/Source/Runtime/UMG/Public/Components/BorderSlot.h
+++ b/Engine/Source/Runtime/UMG/Public/Components/BorderSlot.h
@@ -29,15 +29,15 @@ public:
 
 protected:
 	/** The padding area between the slot and the content it contains. */
-	UPROPERTY(EditDefaultsOnly, Category="Layout|Border Slot")
+	UPROPERTY(EditAnywhere, Category="Layout|Border Slot")
 	FMargin Padding;
 
 	/** The alignment of the object horizontally. */
-	UPROPERTY(EditDefaultsOnly, Category="Layout|Border Slot")
+	UPROPERTY(EditAnywhere, Category="Layout|Border Slot")
 	TEnumAsByte<EHorizontalAlignment> HorizontalAlignment;
 
 	/** The alignment of the object vertically. */
-	UPROPERTY(EditDefaultsOnly, Category="Layout|Border Slot")
+	UPROPERTY(EditAnywhere, Category="Layout|Border Slot")
 	TEnumAsByte<EVerticalAlignment> VerticalAlignment;
 
 public:
diff --git a/Engine/Source/Runtime/UMG/Public/Components/Button.h b/Engine/Source/Runtime/UMG/Public/Components/Button.h
--- a/Engine/Source/Runtime/UMG/Public/Components/Button.h
+++ b/Engine/Source/Runtime/UMG/Public/Components/Button.h
@@ -33,23 +33,23 @@ public:
 	FButtonStyle WidgetStyle;
 	
 	/** The color multiplier for the button content */
-	UPROPERTY(EditDefaultsOnly, BlueprintReadOnly, Category=Appearance)
+	UPROPERTY(EditAnywhere, BlueprintReadOnly, Category=Appearance)
 	FLinearColor ColorAndOpacity;
 	
 	/** The color multiplier for the button background */
-	UPROPERTY(EditDefaultsOnly, BlueprintReadOnly, Category=Appearance)
+	UPROPERTY(EditAnywhere, BlueprintReadOnly, Category=Appearance)
 	FLinearColor BackgroundColor;
 
 	/** The type of mouse action required by the user to trigger the buttons 'Click' */
-	UPROPERTY(EditDefaultsOnly, BlueprintReadOnly, Category="Interaction", AdvancedDisplay)
+	UPROPERTY(EditAnywhere, BlueprintReadOnly, Category="Interaction", AdvancedDisplay)
 	TEnumAsByte<EButtonClickMethod::Type> ClickMethod;
 
 	/** The type of touch action required by the user to trigger the buttons 'Click' */
-	UPROPERTY(EditDefaultsOnly, BlueprintReadOnly, Category="Interaction", AdvancedDisplay)
+	UPROPERTY(EditAnywhere, BlueprintReadOnly, Category="Interaction", AdvancedDisplay)
 	TEnumAsByte<EButtonTouchMethod::Type> TouchMethod;
 
 	/** Sometimes a button should only be mouse-clickable and never keyboard focusable. */
-	UPROPERTY(EditDefaultsOnly, BlueprintReadOnly, Category="Interaction", AdvancedDisplay)
+	UPROPERTY(EditAnywhere, BlueprintReadOnly, Category="Interaction", AdvancedDisplay)
 	bool IsFocusable;
 
 public:
diff --git a/Engine/Source/Runtime/UMG/Public/Components/ButtonSlot.h b/Engine/Source/Runtime/UMG/Public/Components/ButtonSlot.h
--- a/Engine/Source/Runtime/UMG/Public/Components/ButtonSlot.h
+++ b/Engine/Source/Runtime/UMG/Public/Components/ButtonSlot.h
@@ -15,15 +15,15 @@ class UMG_API UButtonSlot : public UPanelSlot
 public:
 	
 	/** The padding area between the slot and the content it contains. */
-	UPROPERTY(EditDefaultsOnly, BlueprintReadOnly, Category="Layout|Button Slot")
+	UPROPERTY(EditAnywhere, BlueprintReadOnly, Category="Layout|Button Slot")
 	FMargin Padding;
 
 	/** The alignment of the object horizontally. */
-	UPROPERTY(EditDefaultsOnly, BlueprintReadOnly, Category="Layout|Button Slot")
+	UPROPERTY(EditAnywhere, BlueprintReadOnly, Category="Layout|Button Slot")
 	TEnumAsByte<EHorizontalAlignment> HorizontalAlignment;
 
 	/** The alignment of the object vertically. */
-	UPROPERTY(EditDefaultsOnly, BlueprintReadOnly, Category="Layout|Button Slot")
+	UPROPERTY(EditAnywhere, BlueprintReadOnly, Category="Layout|Button Slot")
 	TEnumAsByte<EVerticalAlignment> VerticalAlignment;
 
 public:
diff --git a/Engine/Source/Runtime/UMG/Public/Components/CanvasPanelSlot.h b/Engine/Source/Runtime/UMG/Public/Components/CanvasPanelSlot.h
--- a/Engine/Source/Runtime/UMG/Public/Components/CanvasPanelSlot.h
+++ b/Engine/Source/Runtime/UMG/Public/Components/CanvasPanelSlot.h
@@ -23,7 +23,11 @@ public:
 	UPROPERTY(EditAnywhere, BlueprintReadWrite, Category=AnchorData)
 	FAnchors Anchors;
 
-	/** Alignment. */
+	/**
+	 * Alignment is the pivot point of the widget.  Starting in the upper left at (0,0),
+	 * ending in the lower right at (1,1).  Moving the alignment point allows you to move
+	 * the origin of the widget.
+	 */
 	UPROPERTY(EditAnywhere, BlueprintReadWrite, Category=AnchorData)
 	FVector2D Alignment;
 };
@@ -36,15 +40,15 @@ class UMG_API UCanvasPanelSlot : public UPanelSlot
 public:
 
 	/** The anchoring information for the slot */
-	UPROPERTY(EditDefaultsOnly, BlueprintReadOnly, Category="Layout|Canvas Slot")
+	UPROPERTY(EditAnywhere, BlueprintReadOnly, Category="Layout|Canvas Slot")
 	FAnchorData LayoutData;
 
 	/** When AutoSize is true we use the widget's desired size */
-	UPROPERTY(EditDefaultsOnly, BlueprintReadOnly, Category="Layout|Canvas Slot", AdvancedDisplay, meta=( DisplayName="Size To Content" ))
+	UPROPERTY(EditAnywhere, BlueprintReadOnly, Category="Layout|Canvas Slot", AdvancedDisplay, meta=( DisplayName="Size To Content" ))
 	bool bAutoSize;
 
 	/** The order priority this widget is rendered in.  Higher values are rendered last (and so they will appear to be on top). */
-	UPROPERTY(EditDefaultsOnly, BlueprintReadOnly, Category="Layout|Canvas Slot")
+	UPROPERTY(EditAnywhere, BlueprintReadOnly, Category="Layout|Canvas Slot")
 	int32 ZOrder;
 
 public:
diff --git a/Engine/Source/Runtime/UMG/Public/Components/CheckBox.h b/Engine/Source/Runtime/UMG/Public/Components/CheckBox.h
--- a/Engine/Source/Runtime/UMG/Public/Components/CheckBox.h
+++ b/Engine/Source/Runtime/UMG/Public/Components/CheckBox.h
@@ -23,7 +23,7 @@ class UMG_API UCheckBox : public UContentWidget
 
 public:
 	/** Whether the check box is currently in a checked state */
-	UPROPERTY(EditDefaultsOnly, Category=Appearance)
+	UPROPERTY(EditAnywhere, Category=Appearance)
 	ECheckBoxState CheckedState;
 
 	/** A bindable delegate for the IsChecked. */
@@ -76,15 +76,15 @@ public:
 	USlateBrushAsset* UndeterminedPressedImage_DEPRECATED;
 
 	/** How the content of the toggle button should align within the given space */
-	UPROPERTY(EditDefaultsOnly, BlueprintReadOnly, Category=Appearance)
+	UPROPERTY(EditAnywhere, BlueprintReadOnly, Category=Appearance)
 	TEnumAsByte<EHorizontalAlignment> HorizontalAlignment;
 
 	/** Spacing between the check box image and its content */
-	UPROPERTY(EditDefaultsOnly, BlueprintReadOnly, Category=Appearance)
+	UPROPERTY(EditAnywhere, BlueprintReadOnly, Category=Appearance)
 	FMargin Padding;
 
 	/** The color of the background border */
-	UPROPERTY(EditDefaultsOnly, BlueprintReadOnly, Category=Appearance)
+	UPROPERTY(EditAnywhere, BlueprintReadOnly, Category=Appearance)
 	FSlateColor BorderBackgroundColor;
 
 public:
diff --git a/Engine/Source/Runtime/UMG/Public/Components/CircularThrobber.h b/Engine/Source/Runtime/UMG/Public/Components/CircularThrobber.h
--- a/Engine/Source/Runtime/UMG/Public/Components/CircularThrobber.h
+++ b/Engine/Source/Runtime/UMG/Public/Components/CircularThrobber.h
@@ -20,22 +20,22 @@ class UMG_API UCircularThrobber : public UWidget
 public:
 
 	/** How many pieces there are */
-	UPROPERTY(EditDefaultsOnly, BlueprintReadOnly, Category=Appearance, meta=( ClampMin="1", ClampMax="25", UIMin="1", UIMax="25" ))
+	UPROPERTY(EditAnywhere, BlueprintReadOnly, Category=Appearance, meta=( ClampMin="1", ClampMax="25", UIMin="1", UIMax="25" ))
 	int32 NumberOfPieces;
 
 	/** The amount of time for a full circle (in seconds) */
-	UPROPERTY(EditDefaultsOnly, BlueprintReadOnly, Category=Appearance)
+	UPROPERTY(EditAnywhere, BlueprintReadOnly, Category=Appearance)
 	float Period;
 
 	/** The radius of the circle */
-	UPROPERTY(EditDefaultsOnly, BlueprintReadOnly, Category=Appearance)
+	UPROPERTY(EditAnywhere, BlueprintReadOnly, Category=Appearance)
 	float Radius;
 
 	/** Image to use for each segment of the throbber */
 	UPROPERTY()
 	USlateBrushAsset* PieceImage_DEPRECATED;
 
-	UPROPERTY(EditDefaultsOnly, BlueprintReadOnly, Category=Appearance)
+	UPROPERTY(EditAnywhere, BlueprintReadOnly, Category=Appearance)
 	FSlateBrush Image;
 
 public:
diff --git a/Engine/Source/Runtime/UMG/Public/Components/ComboBox.h b/Engine/Source/Runtime/UMG/Public/Components/ComboBox.h
--- a/Engine/Source/Runtime/UMG/Public/Components/ComboBox.h
+++ b/Engine/Source/Runtime/UMG/Public/Components/ComboBox.h
@@ -19,7 +19,7 @@ public:
 	TArray<UObject*> Items;
 
 	/** Called when the widget is needed for the item. */
-	UPROPERTY(EditDefaultsOnly, Category=Events)
+	UPROPERTY(EditAnywhere, Category=Events)
 	FGenerateWidgetForObject OnGenerateWidgetEvent;
 
 	virtual void ReleaseSlateResources(bool bReleaseChildren) override;
diff --git a/Engine/Source/Runtime/UMG/Public/Components/ComboBoxString.h b/Engine/Source/Runtime/UMG/Public/Components/ComboBoxString.h
--- a/Engine/Source/Runtime/UMG/Public/Components/ComboBoxString.h
+++ b/Engine/Source/Runtime/UMG/Public/Components/ComboBoxString.h
@@ -20,11 +20,11 @@ class UMG_API UComboBoxString : public UWidget
 private:
 
 	/** The default list of items to be displayed on the combobox. */
-	UPROPERTY(EditDefaultsOnly, Category=Content)
+	UPROPERTY(EditAnywhere, Category=Content)
 	TArray<FString> DefaultOptions;
 
 	/** The item in the combobox to select by default */
-	UPROPERTY(EditDefaultsOnly, Category=Content)
+	UPROPERTY(EditAnywhere, Category=Content)
 	FString SelectedOption;
 
 public:
@@ -33,24 +33,24 @@ public:
 	UPROPERTY(EditAnywhere, BlueprintReadWrite, Category="Style", meta=( DisplayName="Style" ))
 	FComboBoxStyle WidgetStyle;
 
-	UPROPERTY(EditDefaultsOnly, BlueprintReadOnly, Category=Content)
+	UPROPERTY(EditAnywhere, BlueprintReadOnly, Category=Content)
 	FMargin ContentPadding;
 
 	/** The max height of the combobox list that opens */
-	UPROPERTY(EditDefaultsOnly, BlueprintReadOnly, Category=Content, AdvancedDisplay)
+	UPROPERTY(EditAnywhere, BlueprintReadOnly, Category=Content, AdvancedDisplay)
 	float MaxListHeight;
 
 	/**
 	 * When false, the down arrow is not generated and it is up to the API consumer
 	 * to make their own visual hint that this is a drop down.
 	 */
-	UPROPERTY(EditDefaultsOnly, BlueprintReadOnly, Category=Content, AdvancedDisplay)
+	UPROPERTY(EditAnywhere, BlueprintReadOnly, Category=Content, AdvancedDisplay)
 	bool HasDownArrow;
 
 public: // Events
 
 	/** Called when the widget is needed for the item. */
-	UPROPERTY(EditDefaultsOnly, Category=Events)
+	UPROPERTY(EditAnywhere, Category=Events)
 	FGenerateWidgetForString OnGenerateWidgetEvent;
 
 	/** Called when a new item is selected in the combobox. */
diff --git a/Engine/Source/Runtime/UMG/Public/Components/EditableText.h b/Engine/Source/Runtime/UMG/Public/Components/EditableText.h
--- a/Engine/Source/Runtime/UMG/Public/Components/EditableText.h
+++ b/Engine/Source/Runtime/UMG/Public/Components/EditableText.h
@@ -21,7 +21,7 @@ public:
 
 public:
 	/** The text content for this editable text box widget */
-	UPROPERTY(EditDefaultsOnly, Category=Content)
+	UPROPERTY(EditAnywhere, Category=Content)
 	FText Text;
 
 	/** A bindable delegate to allow logic to drive the text of the widget */
@@ -31,7 +31,7 @@ public:
 public:
 
 	/** Hint text that appears when there is no text in the text box */
-	UPROPERTY(EditDefaultsOnly, BlueprintReadOnly, Category=Content)
+	UPROPERTY(EditAnywhere, BlueprintReadOnly, Category=Content)
 	FText HintText;
 
 	/** A bindable delegate to allow logic to drive the hint text of the widget */
@@ -63,43 +63,43 @@ public:
 	USlateBrushAsset* CaretImage_DEPRECATED;
 
 	/** Font color and opacity (overrides Style) */
-	UPROPERTY(EditDefaultsOnly, Category=Appearance)
+	UPROPERTY(EditAnywhere, Category=Appearance)
 	FSlateFontInfo Font;
 
 	/** Text color and opacity (overrides Style) */
-	UPROPERTY(EditDefaultsOnly, Category=Appearance)
+	UPROPERTY(EditAnywhere, Category=Appearance)
 	FSlateColor ColorAndOpacity;
 
 	/** Sets whether this text box can actually be modified interactively by the user */
-	UPROPERTY(EditDefaultsOnly, Category=Appearance)
+	UPROPERTY(EditAnywhere, Category=Appearance)
 	bool IsReadOnly;
 
 	/** Sets whether this text box is for storing a password */
-	UPROPERTY(EditDefaultsOnly, Category=Appearance)
+	UPROPERTY(EditAnywhere, Category=Appearance)
 	bool IsPassword;
 
 	/** Minimum width that a text block should be */
-	UPROPERTY(EditDefaultsOnly, Category=Appearance)
+	UPROPERTY(EditAnywhere, Category=Appearance)
 	float MinimumDesiredWidth;
 
 	/** Workaround as we lose focus when the auto completion closes. */
-	UPROPERTY(EditDefaultsOnly, Category=Behavior, AdvancedDisplay)
+	UPROPERTY(EditAnywhere, Category=Behavior, AdvancedDisplay)
 	bool IsCaretMovedWhenGainFocus;
 
 	/** Whether to select all text when the user clicks to give focus on the widget */
-	UPROPERTY(EditDefaultsOnly, Category=Behavior, AdvancedDisplay)
+	UPROPERTY(EditAnywhere, Category=Behavior, AdvancedDisplay)
 	bool SelectAllTextWhenFocused;
 
 	/** Whether to allow the user to back out of changes when they press the escape key */
-	UPROPERTY(EditDefaultsOnly, Category=Behavior, AdvancedDisplay)
+	UPROPERTY(EditAnywhere, Category=Behavior, AdvancedDisplay)
 	bool RevertTextOnEscape;
 
 	/** Whether to clear keyboard focus when pressing enter to commit changes */
-	UPROPERTY(EditDefaultsOnly, Category=Behavior, AdvancedDisplay)
+	UPROPERTY(EditAnywhere, Category=Behavior, AdvancedDisplay)
 	bool ClearKeyboardFocusOnCommit;
 
 	/** Whether to select all text when pressing enter to commit changes */
-	UPROPERTY(EditDefaultsOnly, Category=Behavior, AdvancedDisplay)
+	UPROPERTY(EditAnywhere, Category=Behavior, AdvancedDisplay)
 	bool SelectAllTextOnCommit;
 
 public:
diff --git a/Engine/Source/Runtime/UMG/Public/Components/EditableTextBox.h b/Engine/Source/Runtime/UMG/Public/Components/EditableTextBox.h
--- a/Engine/Source/Runtime/UMG/Public/Components/EditableTextBox.h
+++ b/Engine/Source/Runtime/UMG/Public/Components/EditableTextBox.h
@@ -22,7 +22,7 @@ public:
 
 public:
 	/** The text content for this editable text box widget */
-	UPROPERTY(EditDefaultsOnly, Category=Content)
+	UPROPERTY(EditAnywhere, Category=Content)
 	FText Text;
 
 	/** A bindable delegate to allow logic to drive the text of the widget */
@@ -40,7 +40,7 @@ public:
 	USlateWidgetStyleAsset* Style_DEPRECATED;
 
 	/** Hint text that appears when there is no text in the text box */
-	UPROPERTY(EditDefaultsOnly, BlueprintReadOnly, Category=Content)
+	UPROPERTY(EditAnywhere, BlueprintReadOnly, Category=Content)
 	FText HintText;
 
 	/** A bindable delegate to allow logic to drive the hint text of the widget */
@@ -48,55 +48,55 @@ public:
 	FGetText HintTextDelegate;
 
 	/** Font color and opacity (overrides Style) */
-	UPROPERTY(EditDefaultsOnly, BlueprintReadOnly, Category=Appearance)
+	UPROPERTY(EditAnywhere, BlueprintReadOnly, Category=Appearance)
 	FSlateFontInfo Font;
 
 	/** Text color and opacity (overrides Style) */
-	UPROPERTY(EditDefaultsOnly, BlueprintReadOnly, Category=Appearance)
+	UPROPERTY(EditAnywhere, BlueprintReadOnly, Category=Appearance)
 	FLinearColor ForegroundColor;
 
 	/** The color of the background/border around the editable text (overrides Style) */
-	UPROPERTY(EditDefaultsOnly, BlueprintReadOnly, Category=Appearance)
+	UPROPERTY(EditAnywhere, BlueprintReadOnly, Category=Appearance)
 	FLinearColor BackgroundColor;
 
 	/** Text color and opacity when read-only (overrides Style) */
-	UPROPERTY(EditDefaultsOnly, BlueprintReadOnly, Category=Appearance)
+	UPROPERTY(EditAnywhere, BlueprintReadOnly, Category=Appearance)
 	FLinearColor ReadOnlyForegroundColor;
 
 	/** Sets whether this text box can actually be modified interactively by the user */
-	UPROPERTY(EditDefaultsOnly, BlueprintReadOnly, Category=Appearance)
+	UPROPERTY(EditAnywhere, BlueprintReadOnly, Category=Appearance)
 	bool IsReadOnly;
 
 	/** Sets whether this text box is for storing a password */
-	UPROPERTY(EditDefaultsOnly, BlueprintReadOnly, Category=Appearance)
+	UPROPERTY(EditAnywhere, BlueprintReadOnly, Category=Appearance)
 	bool IsPassword;
 
 	/** Minimum width that a text block should be */
-	UPROPERTY(EditDefaultsOnly, BlueprintReadOnly, Category=Appearance)
+	UPROPERTY(EditAnywhere, BlueprintReadOnly, Category=Appearance)
 	float MinimumDesiredWidth;
 
 	/** Padding between the box/border and the text widget inside (overrides Style) */
-	UPROPERTY(EditDefaultsOnly, BlueprintReadOnly, Category=Appearance)
+	UPROPERTY(EditAnywhere, BlueprintReadOnly, Category=Appearance)
 	FMargin Padding;
 
 	/** Workaround as we lose focus when the auto completion closes. */
-	UPROPERTY(EditDefaultsOnly, BlueprintReadOnly, Category=Behavior, AdvancedDisplay)
+	UPROPERTY(EditAnywhere, BlueprintReadOnly, Category=Behavior, AdvancedDisplay)
 	bool IsCaretMovedWhenGainFocus;
 
 	/** Whether to select all text when the user clicks to give focus on the widget */
-	UPROPERTY(EditDefaultsOnly, BlueprintReadOnly, Category=Behavior, AdvancedDisplay)
+	UPROPERTY(EditAnywhere, BlueprintReadOnly, Category=Behavior, AdvancedDisplay)
 	bool SelectAllTextWhenFocused;
 
 	/** Whether to allow the user to back out of changes when they press the escape key */
-	UPROPERTY(EditDefaultsOnly, BlueprintReadOnly, Category=Behavior, AdvancedDisplay)
+	UPROPERTY(EditAnywhere, BlueprintReadOnly, Category=Behavior, AdvancedDisplay)
 	bool RevertTextOnEscape;
 
 	/** Whether to clear keyboard focus when pressing enter to commit changes */
-	UPROPERTY(EditDefaultsOnly, BlueprintReadOnly, Category=Behavior, AdvancedDisplay)
+	UPROPERTY(EditAnywhere, BlueprintReadOnly, Category=Behavior, AdvancedDisplay)
 	bool ClearKeyboardFocusOnCommit;
 
 	/** Whether to select all text when pressing enter to commit changes */
-	UPROPERTY(EditDefaultsOnly, BlueprintReadOnly, Category=Behavior, AdvancedDisplay)
+	UPROPERTY(EditAnywhere, BlueprintReadOnly, Category=Behavior, AdvancedDisplay)
 	bool SelectAllTextOnCommit;
 
 public:
diff --git a/Engine/Source/Runtime/UMG/Public/Components/ExpandableArea.h b/Engine/Source/Runtime/UMG/Public/Components/ExpandableArea.h
--- a/Engine/Source/Runtime/UMG/Public/Components/ExpandableArea.h
+++ b/Engine/Source/Runtime/UMG/Public/Components/ExpandableArea.h
@@ -16,15 +16,15 @@ class UMG_API UExpandableArea : public UWidget, public INamedSlotInterface
 
 public:
 	/**  */
-	UPROPERTY(EditDefaultsOnly, BlueprintReadOnly, Category="Expansion")
+	UPROPERTY(EditAnywhere, BlueprintReadOnly, Category="Expansion")
 	bool bIsExpanded;
 
 	/** The maximum height of the area */
-	UPROPERTY(EditDefaultsOnly, BlueprintReadOnly, Category="Expansion")
+	UPROPERTY(EditAnywhere, BlueprintReadOnly, Category="Expansion")
 	float MaxHeight;
 	
 	/**  */
-	UPROPERTY(EditDefaultsOnly, BlueprintReadOnly, Category="Expansion")
+	UPROPERTY(EditAnywhere, BlueprintReadOnly, Category="Expansion")
 	FMargin AreaPadding;
 
 	/** A bindable delegate for the IsChecked. */
diff --git a/Engine/Source/Runtime/UMG/Public/Components/GridPanel.h b/Engine/Source/Runtime/UMG/Public/Components/GridPanel.h
--- a/Engine/Source/Runtime/UMG/Public/Components/GridPanel.h
+++ b/Engine/Source/Runtime/UMG/Public/Components/GridPanel.h
@@ -17,11 +17,11 @@ class UMG_API UGridPanel : public UPanelWidget
 public:
 
 	/** The column fill rules */
-	UPROPERTY(EditDefaultsOnly, BlueprintReadOnly, Category="Fill Rules")
+	UPROPERTY(EditAnywhere, BlueprintReadOnly, Category="Fill Rules")
 	TArray<float> ColumnFill;
 
 	/** The row fill rules */
-	UPROPERTY(EditDefaultsOnly, BlueprintReadOnly, Category="Fill Rules")
+	UPROPERTY(EditAnywhere, BlueprintReadOnly, Category="Fill Rules")
 	TArray<float> RowFill;
 
 public:
diff --git a/Engine/Source/Runtime/UMG/Public/Components/GridSlot.h b/Engine/Source/Runtime/UMG/Public/Components/GridSlot.h
--- a/Engine/Source/Runtime/UMG/Public/Components/GridSlot.h
+++ b/Engine/Source/Runtime/UMG/Public/Components/GridSlot.h
@@ -18,35 +18,35 @@ class UMG_API UGridSlot : public UPanelSlot
 public:
 
 	/** The alignment of the object horizontally. */
-	UPROPERTY(EditDefaultsOnly, BlueprintReadOnly, Category="Layout|Grid Slot")
+	UPROPERTY(EditAnywhere, BlueprintReadOnly, Category="Layout|Grid Slot")
 	TEnumAsByte<EHorizontalAlignment> HorizontalAlignment;
 
 	/** The alignment of the object vertically. */
-	UPROPERTY(EditDefaultsOnly, BlueprintReadOnly, Category="Layout|Grid Slot")
+	UPROPERTY(EditAnywhere, BlueprintReadOnly, Category="Layout|Grid Slot")
 	TEnumAsByte<EVerticalAlignment> VerticalAlignment;
 	
 	/** The row index of the cell this slot is in */
-	UPROPERTY(EditDefaultsOnly, BlueprintReadOnly, meta=( UIMin = "0" ), Category="Layout|Grid Slot")
+	UPROPERTY(EditAnywhere, BlueprintReadOnly, meta=( UIMin = "0" ), Category="Layout|Grid Slot")
 	int32 Row;
 	
 	/**  */
-	UPROPERTY(EditDefaultsOnly, BlueprintReadOnly, Category="Layout|Grid Slot")
+	UPROPERTY(EditAnywhere, BlueprintReadOnly, Category="Layout|Grid Slot")
 	int32 RowSpan;
 
 	/** The column index of the cell this slot is in */
-	UPROPERTY(EditDefaultsOnly, BlueprintReadOnly, meta=( UIMin = "0" ), Category="Layout|Grid Slot")
+	UPROPERTY(EditAnywhere, BlueprintReadOnly, meta=( UIMin = "0" ), Category="Layout|Grid Slot")
 	int32 Column;
 
 	/**  */
-	UPROPERTY(EditDefaultsOnly, BlueprintReadOnly, Category="Layout|Grid Slot")
+	UPROPERTY(EditAnywhere, BlueprintReadOnly, Category="Layout|Grid Slot")
 	int32 ColumnSpan;
 
 	/** Positive values offset this cell to be hit-tested and drawn on top of others. Default is 0; i.e. no offset. */
-	UPROPERTY(EditDefaultsOnly, BlueprintReadOnly, Category="Layout|Grid Slot")
+	UPROPERTY(EditAnywhere, BlueprintReadOnly, Category="Layout|Grid Slot")
 	int32 Layer;
 
 	/** Offset this slot's content by some amount; positive values offset to lower right*/
-	UPROPERTY(EditDefaultsOnly, BlueprintReadOnly, Category="Layout|Grid Slot")
+	UPROPERTY(EditAnywhere, BlueprintReadOnly, Category="Layout|Grid Slot")
 	FVector2D Nudge;
 
 	/** Sets the row index of the slot, this determines what cell the slot is in the panel */
diff --git a/Engine/Source/Runtime/UMG/Public/Components/HorizontalBoxSlot.h b/Engine/Source/Runtime/UMG/Public/Components/HorizontalBoxSlot.h
--- a/Engine/Source/Runtime/UMG/Public/Components/HorizontalBoxSlot.h
+++ b/Engine/Source/Runtime/UMG/Public/Components/HorizontalBoxSlot.h
@@ -14,17 +14,17 @@ class UMG_API UHorizontalBoxSlot : public UPanelSlot
 public:
 	
 	/** The amount of padding between the slots parent and the content. */
-	UPROPERTY(EditDefaultsOnly, BlueprintReadOnly, Category="Layout|Horizontal Box Slot")
+	UPROPERTY(EditAnywhere, BlueprintReadOnly, Category="Layout|Horizontal Box Slot")
 	FMargin Padding;
 
 	/** How much space this slot should occupy in the direction of the panel. */
-	UPROPERTY(EditDefaultsOnly, BlueprintReadOnly, Category="Layout|Horizontal Box Slot")
+	UPROPERTY(EditAnywhere, BlueprintReadOnly, Category="Layout|Horizontal Box Slot")
 	FSlateChildSize Size;
 
-	UPROPERTY(EditDefaultsOnly, BlueprintReadOnly, Category="Layout|Horizontal Box Slot")
+	UPROPERTY(EditAnywhere, BlueprintReadOnly, Category="Layout|Horizontal Box Slot")
 	TEnumAsByte<EHorizontalAlignment> HorizontalAlignment;
 
-	UPROPERTY(EditDefaultsOnly, BlueprintReadOnly, Category="Layout|Horizontal Box Slot")
+	UPROPERTY(EditAnywhere, BlueprintReadOnly, Category="Layout|Horizontal Box Slot")
 	TEnumAsByte<EVerticalAlignment> VerticalAlignment;
 
 public:
diff --git a/Engine/Source/Runtime/UMG/Public/Components/Image.h b/Engine/Source/Runtime/UMG/Public/Components/Image.h
--- a/Engine/Source/Runtime/UMG/Public/Components/Image.h
+++ b/Engine/Source/Runtime/UMG/Public/Components/Image.h
@@ -23,7 +23,7 @@ public:
 	USlateBrushAsset* Image_DEPRECATED;
 
 	/** Image to draw */
-	UPROPERTY(EditDefaultsOnly, BlueprintReadOnly, Category=Appearance)
+	UPROPERTY(EditAnywhere, BlueprintReadOnly, Category=Appearance)
 	FSlateBrush Brush;
 
 	/** A bindable delegate for the Image. */
@@ -31,7 +31,7 @@ public:
 	FGetSlateBrush BrushDelegate;
 
 	/** Color and opacity */
-	UPROPERTY(EditDefaultsOnly, BlueprintReadOnly, Category=Appearance)
+	UPROPERTY(EditAnywhere, BlueprintReadOnly, Category=Appearance)
 	FLinearColor ColorAndOpacity;
 
 	/** A bindable delegate for the ColorAndOpacity. */
@@ -40,7 +40,7 @@ public:
 
 public:
 
-	UPROPERTY(EditDefaultsOnly, Category=Events)
+	UPROPERTY(EditAnywhere, Category=Events)
 	FOnPointerEvent OnMouseButtonDownEvent;
 
 public:
diff --git a/Engine/Source/Runtime/UMG/Public/Components/ListView.h b/Engine/Source/Runtime/UMG/Public/Components/ListView.h
--- a/Engine/Source/Runtime/UMG/Public/Components/ListView.h
+++ b/Engine/Source/Runtime/UMG/Public/Components/ListView.h
@@ -15,7 +15,7 @@ class UMG_API UListView : public UTableViewBase
 public:
 
 	/** The height of each widget */
-	UPROPERTY(EditDefaultsOnly, Category=Content)
+	UPROPERTY(EditAnywhere, Category=Content)
 	float ItemHeight;
 
 	/** The list of items to generate widgets for */
@@ -23,11 +23,11 @@ public:
 	TArray<UObject*> Items;
 
 	/** The selection method for the list */
-	UPROPERTY(EditDefaultsOnly, Category=Content)
+	UPROPERTY(EditAnywhere, Category=Content)
 	TEnumAsByte<ESelectionMode::Type> SelectionMode;
 
 	/** Called when a widget needs to be generated */
-	UPROPERTY(EditDefaultsOnly, Category=Events)
+	UPROPERTY(EditAnywhere, Category=Events)
 	FOnGenerateRowUObject OnGenerateRowEvent;
 
 #if WITH_EDITOR
diff --git a/Engine/Source/Runtime/UMG/Public/Components/MenuAnchor.h b/Engine/Source/Runtime/UMG/Public/Components/MenuAnchor.h
--- a/Engine/Source/Runtime/UMG/Public/Components/MenuAnchor.h
+++ b/Engine/Source/Runtime/UMG/Public/Components/MenuAnchor.h
@@ -24,15 +24,15 @@ public:
 	 * If you want to customize the creation of the popup, you should bind a function to OnGetMenuContentEvent 
 	 * instead.
 	 */
-	UPROPERTY(EditDefaultsOnly, BlueprintReadOnly, Category="Menu Anchor")
+	UPROPERTY(EditAnywhere, BlueprintReadOnly, Category="Menu Anchor")
 	TSubclassOf<class UUserWidget> MenuClass;
 
 	/** Called when the menu content is requested to allow a more customized handling over what to display */
-	UPROPERTY(EditDefaultsOnly, Category="Events")
+	UPROPERTY(EditAnywhere, Category="Events")
 	FGetWidget OnGetMenuContentEvent;
 	
 	/** The placement location of the summoned widget. */
-	UPROPERTY(EditDefaultsOnly, BlueprintReadOnly, Category="Menu Anchor")
+	UPROPERTY(EditAnywhere, BlueprintReadOnly, Category="Menu Anchor")
 	TEnumAsByte<EMenuPlacement> Placement;
 
 public:
diff --git a/Engine/Source/Runtime/UMG/Public/Components/MultiLineEditableText.h b/Engine/Source/Runtime/UMG/Public/Components/MultiLineEditableText.h
--- a/Engine/Source/Runtime/UMG/Public/Components/MultiLineEditableText.h
+++ b/Engine/Source/Runtime/UMG/Public/Components/MultiLineEditableText.h
@@ -19,7 +19,7 @@ public:
 
 public:
 	/** The text content for this editable text box widget */
-	UPROPERTY(EditDefaultsOnly, Category=Content, meta=(MultiLine="true"))
+	UPROPERTY(EditAnywhere, Category=Content, meta=(MultiLine="true"))
 	FText Text;
 
 public:
@@ -33,11 +33,11 @@ public:
 	TEnumAsByte<ETextJustify::Type> Justification;
 
 	/** Whether to wrap text automatically based on the widget's computed horizontal space.*/
-	UPROPERTY(EditDefaultsOnly, Category=Content)
+	UPROPERTY(EditAnywhere, Category=Content)
 	bool bAutoWrapText;
 
 	/** Whether text wraps onto a new line when it's length exceeds this width; if this value is zero or negative, no wrapping occurs. */
-	UPROPERTY(EditDefaultsOnly, AdvancedDisplay, Category=Content)
+	UPROPERTY(EditAnywhere, AdvancedDisplay, Category=Content)
 	float WrapTextAt;
 
 	/** Font color and opacity (overrides Style) */
diff --git a/Engine/Source/Runtime/UMG/Public/Components/MultiLineEditableTextBox.h b/Engine/Source/Runtime/UMG/Public/Components/MultiLineEditableTextBox.h
--- a/Engine/Source/Runtime/UMG/Public/Components/MultiLineEditableTextBox.h
+++ b/Engine/Source/Runtime/UMG/Public/Components/MultiLineEditableTextBox.h
@@ -20,7 +20,7 @@ public:
 public:
 	
 	/** The text content for this editable text box widget */
-	UPROPERTY(EditDefaultsOnly, Category=Content, meta=(MultiLine="true"))
+	UPROPERTY(EditAnywhere, Category=Content, meta=(MultiLine="true"))
 	FText Text;
 
 public:
@@ -37,31 +37,31 @@ public:
 	USlateWidgetStyleAsset* Style_DEPRECATED;
 
 	/** The justification of the text in the multilinebox */
-	UPROPERTY(EditDefaultsOnly, BlueprintReadOnly, Category=Content)
+	UPROPERTY(EditAnywhere, BlueprintReadOnly, Category=Content)
 	TEnumAsByte<ETextJustify::Type> Justification;
 
 	/** Whether to wrap text automatically based on the widget's computed horizontal space.*/
-	UPROPERTY(EditDefaultsOnly, Category=Content)
+	UPROPERTY(EditAnywhere, Category=Content)
 	bool bAutoWrapText;
 
 	/** Whether text wraps onto a new line when it's length exceeds this width; if this value is zero or negative, no wrapping occurs. */
-	UPROPERTY(EditDefaultsOnly, AdvancedDisplay, Category=Content)
+	UPROPERTY(EditAnywhere, AdvancedDisplay, Category=Content)
 	float WrapTextAt;
 
 	/** Font color and opacity (overrides Style) */
-	UPROPERTY(EditDefaultsOnly, BlueprintReadOnly, Category=Appearance)
+	UPROPERTY(EditAnywhere, BlueprintReadOnly, Category=Appearance)
 	FSlateFontInfo Font;
 
 	/** Text color and opacity (overrides Style) */
-	UPROPERTY(EditDefaultsOnly, BlueprintReadOnly, Category=Appearance)
+	UPROPERTY(EditAnywhere, BlueprintReadOnly, Category=Appearance)
 	FLinearColor ForegroundColor;
 
 	/** The color of the background/border around the editable text (overrides Style) */
-	UPROPERTY(EditDefaultsOnly, BlueprintReadOnly, Category=Appearance)
+	UPROPERTY(EditAnywhere, BlueprintReadOnly, Category=Appearance)
 	FLinearColor BackgroundColor;
 
 	/** Text color and opacity when read-only (overrides Style) */
-	UPROPERTY(EditDefaultsOnly, BlueprintReadOnly, Category=Appearance)
+	UPROPERTY(EditAnywhere, BlueprintReadOnly, Category=Appearance)
 	FLinearColor ReadOnlyForegroundColor;
 
 	/** Called whenever the text is changed interactively by the user */
diff --git a/Engine/Source/Runtime/UMG/Public/Components/NamedSlotInterface.h b/Engine/Source/Runtime/UMG/Public/Components/NamedSlotInterface.h
--- a/Engine/Source/Runtime/UMG/Public/Components/NamedSlotInterface.h
+++ b/Engine/Source/Runtime/UMG/Public/Components/NamedSlotInterface.h
@@ -7,8 +7,8 @@
 class UWidget;
 
 /**  */
-UINTERFACE(MinimalAPI, meta=( CannotImplementInterfaceInBlueprint ))
-class UNamedSlotInterface : public UInterface
+UINTERFACE(meta=( CannotImplementInterfaceInBlueprint ))
+class UMG_API UNamedSlotInterface : public UInterface
 {
 	GENERATED_UINTERFACE_BODY()
 };
@@ -25,4 +25,7 @@ class UMG_API INamedSlotInterface
 
 	/**  */
 	virtual void SetContentForSlot(FName SlotName, UWidget* Content) = 0;
+
+	/**  */
+	bool ContainsContent(UWidget* Content) const;
 };
diff --git a/Engine/Source/Runtime/UMG/Public/Components/OverlaySlot.h b/Engine/Source/Runtime/UMG/Public/Components/OverlaySlot.h
--- a/Engine/Source/Runtime/UMG/Public/Components/OverlaySlot.h
+++ b/Engine/Source/Runtime/UMG/Public/Components/OverlaySlot.h
@@ -17,15 +17,15 @@ class UMG_API UOverlaySlot : public UPanelSlot
 public:
 	
 	/** The padding area between the slot and the content it contains. */
-	UPROPERTY(EditDefaultsOnly, BlueprintReadOnly, Category="Layout|Overlay Slot")
+	UPROPERTY(EditAnywhere, BlueprintReadOnly, Category="Layout|Overlay Slot")
 	FMargin Padding;
 
 	/** The alignment of the object horizontally. */
-	UPROPERTY(EditDefaultsOnly, BlueprintReadOnly, Category="Layout|Overlay Slot")
+	UPROPERTY(EditAnywhere, BlueprintReadOnly, Category="Layout|Overlay Slot")
 	TEnumAsByte<EHorizontalAlignment> HorizontalAlignment;
 
 	/** The alignment of the object vertically. */
-	UPROPERTY(EditDefaultsOnly, BlueprintReadOnly, Category="Layout|Overlay Slot")
+	UPROPERTY(EditAnywhere, BlueprintReadOnly, Category="Layout|Overlay Slot")
 	TEnumAsByte<EVerticalAlignment> VerticalAlignment;
 
 public:
diff --git a/Engine/Source/Runtime/UMG/Public/Components/ProgressBar.h b/Engine/Source/Runtime/UMG/Public/Components/ProgressBar.h
--- a/Engine/Source/Runtime/UMG/Public/Components/ProgressBar.h
+++ b/Engine/Source/Runtime/UMG/Public/Components/ProgressBar.h
@@ -39,14 +39,14 @@ public:
 	USlateBrushAsset* MarqueeImage_DEPRECATED;
 
 	/** Used to determine the fill position of the progress bar ranging 0..1 */
-	UPROPERTY(EditDefaultsOnly, BlueprintReadOnly, Category=Progress, meta=( UIMin = "0", UIMax = "1" ))
+	UPROPERTY(EditAnywhere, BlueprintReadOnly, Category=Progress, meta=( UIMin = "0", UIMax = "1" ))
 	float Percent;
 
 	/** Defines if this progress bar fills Left to right or right to left */
-	UPROPERTY(EditDefaultsOnly, BlueprintReadOnly, Category=Progress)
+	UPROPERTY(EditAnywhere, BlueprintReadOnly, Category=Progress)
 	TEnumAsByte<EProgressBarFillType::Type> BarFillType;
 	
-	UPROPERTY(EditDefaultsOnly, BlueprintReadOnly, Category=Progress)
+	UPROPERTY(EditAnywhere, BlueprintReadOnly, Category=Progress)
 	bool bIsMarquee;
 
 	/** A bindable delegate to allow logic to drive the text of the widget */
@@ -54,7 +54,7 @@ public:
 	FGetFloat PercentDelegate;
 
 	/** Fill Color and Opacity */
-	UPROPERTY(EditDefaultsOnly, BlueprintReadOnly, Category=Appearance)
+	UPROPERTY(EditAnywhere, BlueprintReadOnly, Category=Appearance)
 	FLinearColor FillColorAndOpacity;
 
 	/** */
diff --git a/Engine/Source/Runtime/UMG/Public/Components/RichTextBlock.h b/Engine/Source/Runtime/UMG/Public/Components/RichTextBlock.h
--- a/Engine/Source/Runtime/UMG/Public/Components/RichTextBlock.h
+++ b/Engine/Source/Runtime/UMG/Public/Components/RichTextBlock.h
@@ -35,7 +35,7 @@ public:
 
 protected:
 	/** The text to display */
-	UPROPERTY(EditDefaultsOnly, Category=Content, meta=( MultiLine="true" ))
+	UPROPERTY(EditAnywhere, Category=Content, meta=( MultiLine="true" ))
 	FText Text;
 
 	/** A bindable delegate to allow logic to drive the text of the widget */
@@ -43,34 +43,34 @@ protected:
 	FGetText TextDelegate;
 
 	/** The default font for the text. */
-	UPROPERTY(EditDefaultsOnly, BlueprintReadOnly, Category=Appearance)
+	UPROPERTY(EditAnywhere, BlueprintReadOnly, Category=Appearance)
 	FSlateFontInfo Font;
 
 	/** The default color for the text. */
-	UPROPERTY(EditDefaultsOnly, BlueprintReadOnly, Category=Appearance)
+	UPROPERTY(EditAnywhere, BlueprintReadOnly, Category=Appearance)
 	FLinearColor Color;
 
 	/** How the text should be aligned with the margin. */
-	UPROPERTY(EditDefaultsOnly, BlueprintReadOnly, Category=Appearance)
+	UPROPERTY(EditAnywhere, BlueprintReadOnly, Category=Appearance)
 	TEnumAsByte<ETextJustify::Type> Justification;
 
 	/** True if we're wrapping text automatically based on the computed horizontal space for this widget */
-	UPROPERTY(EditDefaultsOnly, BlueprintReadOnly, Category=Appearance)
+	UPROPERTY(EditAnywhere, BlueprintReadOnly, Category=Appearance)
 	bool AutoWrapText;
 
 	/** Whether text wraps onto a new line when it's length exceeds this width; if this value is zero or negative, no wrapping occurs. */
-	UPROPERTY(EditDefaultsOnly, BlueprintReadOnly, Category=Appearance, AdvancedDisplay)
+	UPROPERTY(EditAnywhere, BlueprintReadOnly, Category=Appearance, AdvancedDisplay)
 	float WrapTextAt;
 
 	/** The amount of blank space left around the edges of text area. */
-	UPROPERTY(EditDefaultsOnly, BlueprintReadOnly, Category=Appearance, AdvancedDisplay)
+	UPROPERTY(EditAnywhere, BlueprintReadOnly, Category=Appearance, AdvancedDisplay)
 	FMargin Margin;
 
 	/** The amount to scale each lines height by. */
-	UPROPERTY(EditDefaultsOnly, BlueprintReadOnly, Category=Appearance, AdvancedDisplay)
+	UPROPERTY(EditAnywhere, BlueprintReadOnly, Category=Appearance, AdvancedDisplay)
 	float LineHeightPercentage;
 
-	UPROPERTY(EditDefaultsOnly, Instanced, Category=Decorators)
+	UPROPERTY(EditAnywhere, Instanced, Category=Decorators)
 	TArray<class URichTextBlockDecorator*> Decorators;
 
 protected:
diff --git a/Engine/Source/Runtime/UMG/Public/Components/ScaleBox.h b/Engine/Source/Runtime/UMG/Public/Components/ScaleBox.h
--- a/Engine/Source/Runtime/UMG/Public/Components/ScaleBox.h
+++ b/Engine/Source/Runtime/UMG/Public/Components/ScaleBox.h
@@ -20,15 +20,15 @@ class UMG_API UScaleBox : public UContentWidget
 public:
 
 	/** Controls in what direction content can be scaled */
-	UPROPERTY(EditDefaultsOnly, BlueprintReadOnly, Category="Stretching")
+	UPROPERTY(EditAnywhere, BlueprintReadOnly, Category="Stretching")
 	TEnumAsByte<EStretchDirection::Type> StretchDirection;
 
 	/** The stretching rule to apply when content is stretched */
-	UPROPERTY(EditDefaultsOnly, BlueprintReadOnly, Category="Stretching")
+	UPROPERTY(EditAnywhere, BlueprintReadOnly, Category="Stretching")
 	TEnumAsByte<EStretch::Type> Stretch;
 
 	/** Optional scale that can be specified by the User. Isn't used if Stretch != EStretch::UserSpecified */
-	UPROPERTY(EditDefaultsOnly, BlueprintReadOnly, Category="Stretching")
+	UPROPERTY(EditAnywhere, BlueprintReadOnly, Category="Stretching")
 	float UserSpecifiedScale;
 
 public:
diff --git a/Engine/Source/Runtime/UMG/Public/Components/ScaleBoxSlot.h b/Engine/Source/Runtime/UMG/Public/Components/ScaleBoxSlot.h
--- a/Engine/Source/Runtime/UMG/Public/Components/ScaleBoxSlot.h
+++ b/Engine/Source/Runtime/UMG/Public/Components/ScaleBoxSlot.h
@@ -15,15 +15,15 @@ class UMG_API UScaleBoxSlot : public UPanelSlot
 public:
 	
 	/** The padding area between the slot and the content it contains. */
-	UPROPERTY(EditDefaultsOnly, BlueprintReadOnly, Category="Layout|ScaleBox Slot")
+	UPROPERTY(EditAnywhere, BlueprintReadOnly, Category="Layout|ScaleBox Slot")
 	FMargin Padding;
 
 	/** The alignment of the object horizontally. */
-	UPROPERTY(EditDefaultsOnly, BlueprintReadOnly, Category="Layout|ScaleBox Slot")
+	UPROPERTY(EditAnywhere, BlueprintReadOnly, Category="Layout|ScaleBox Slot")
 	TEnumAsByte<EHorizontalAlignment> HorizontalAlignment;
 
 	/** The alignment of the object vertically. */
-	UPROPERTY(EditDefaultsOnly, BlueprintReadOnly, Category="Layout|ScaleBox Slot" )
+	UPROPERTY(EditAnywhere, BlueprintReadOnly, Category="Layout|ScaleBox Slot" )
 	TEnumAsByte<EVerticalAlignment> VerticalAlignment;
 
 public:
diff --git a/Engine/Source/Runtime/UMG/Public/Components/ScrollBar.h b/Engine/Source/Runtime/UMG/Public/Components/ScrollBar.h
--- a/Engine/Source/Runtime/UMG/Public/Components/ScrollBar.h
+++ b/Engine/Source/Runtime/UMG/Public/Components/ScrollBar.h
@@ -20,15 +20,15 @@ public:
 	USlateWidgetStyleAsset* Style_DEPRECATED;
 
 	/**  */
-	UPROPERTY(EditDefaultsOnly, Category="Behavior")
+	UPROPERTY(EditAnywhere, Category="Behavior")
 	bool bAlwaysShowScrollbar;
 
 	/**  */
-	UPROPERTY(EditDefaultsOnly, Category="Behavior")
+	UPROPERTY(EditAnywhere, Category="Behavior")
 	TEnumAsByte<EOrientation> Orientation;
 
 	/** The thickness of the scrollbar thumb */
-	UPROPERTY(EditDefaultsOnly, Category="Behavior")
+	UPROPERTY(EditAnywhere, Category="Behavior")
 	FVector2D Thickness;
 
 public:
diff --git a/Engine/Source/Runtime/UMG/Public/Components/ScrollBox.h b/Engine/Source/Runtime/UMG/Public/Components/ScrollBox.h
--- a/Engine/Source/Runtime/UMG/Public/Components/ScrollBox.h
+++ b/Engine/Source/Runtime/UMG/Public/Components/ScrollBox.h
@@ -29,23 +29,23 @@ public:
 	USlateWidgetStyleAsset* BarStyle_DEPRECATED;
 
 	/** The orientation of the scrolling and stacking in the box. */
-	UPROPERTY(EditDefaultsOnly, BlueprintReadOnly, Category="Scroll")
+	UPROPERTY(EditAnywhere, BlueprintReadOnly, Category="Scroll")
 	TEnumAsByte<EOrientation> Orientation;
 
 	/** Visibility */
-	UPROPERTY(EditDefaultsOnly, BlueprintReadOnly, Category="Scroll")
+	UPROPERTY(EditAnywhere, BlueprintReadOnly, Category="Scroll")
 	TEnumAsByte<ESlateVisibility> ScrollBarVisibility;
 
 	/**  Enable to always consume mouse wheel event, even when scrolling is not possible */
-	UPROPERTY(EditDefaultsOnly, BlueprintReadOnly, Category = "Scroll")
+	UPROPERTY(EditAnywhere, BlueprintReadOnly, Category = "Scroll")
 	TEnumAsByte<EConsumeMouseWheel> ConsumeMouseWheel;
 
 	/**  */
-	UPROPERTY(EditDefaultsOnly, BlueprintReadOnly, Category="Scroll")
+	UPROPERTY(EditAnywhere, BlueprintReadOnly, Category="Scroll")
 	FVector2D ScrollbarThickness;
 
 	/**  */
-	UPROPERTY(EditDefaultsOnly, BlueprintReadOnly, Category="Scroll")
+	UPROPERTY(EditAnywhere, BlueprintReadOnly, Category="Scroll")
 	bool AlwaysShowScrollbar;
 	
 	//TODO UMG Add SetOrientation
diff --git a/Engine/Source/Runtime/UMG/Public/Components/ScrollBoxSlot.h b/Engine/Source/Runtime/UMG/Public/Components/ScrollBoxSlot.h
--- a/Engine/Source/Runtime/UMG/Public/Components/ScrollBoxSlot.h
+++ b/Engine/Source/Runtime/UMG/Public/Components/ScrollBoxSlot.h
@@ -15,11 +15,11 @@ class UMG_API UScrollBoxSlot : public UPanelSlot
 public:
 	
 	/** The padding area between the slot and the content it contains. */
-	UPROPERTY(EditDefaultsOnly, BlueprintReadOnly, Category="Layout|ScrollBox Slot")
+	UPROPERTY(EditAnywhere, BlueprintReadOnly, Category="Layout|ScrollBox Slot")
 	FMargin Padding;
 
 	/** The alignment of the object horizontally. */
-	UPROPERTY(EditDefaultsOnly, BlueprintReadOnly, Category="Layout|ScrollBox Slot")
+	UPROPERTY(EditAnywhere, BlueprintReadOnly, Category="Layout|ScrollBox Slot")
 	TEnumAsByte<EHorizontalAlignment> HorizontalAlignment;
 
 public:
diff --git a/Engine/Source/Runtime/UMG/Public/Components/SizeBox.h b/Engine/Source/Runtime/UMG/Public/Components/SizeBox.h
--- a/Engine/Source/Runtime/UMG/Public/Components/SizeBox.h
+++ b/Engine/Source/Runtime/UMG/Public/Components/SizeBox.h
@@ -46,27 +46,27 @@ public:
 
 
 	/** When specified, ignore the content's desired size and report the WidthOverride as the Box's desired width. */
-	UPROPERTY(EditDefaultsOnly, BlueprintReadOnly, Category="Child Layout", meta=( editcondition="bOverride_WidthOverride" ))
+	UPROPERTY(EditAnywhere, BlueprintReadOnly, Category="Child Layout", meta=( editcondition="bOverride_WidthOverride" ))
 	float WidthOverride;
 
 	/** When specified, ignore the content's desired size and report the HeightOverride as the Box's desired height. */
-	UPROPERTY(EditDefaultsOnly, BlueprintReadOnly, Category="Child Layout", meta=( editcondition="bOverride_HeightOverride" ))
+	UPROPERTY(EditAnywhere, BlueprintReadOnly, Category="Child Layout", meta=( editcondition="bOverride_HeightOverride" ))
 	float HeightOverride;
 
 	/** When specified, will report the MinDesiredWidth if larger than the content's desired width. */
-	UPROPERTY(EditDefaultsOnly, BlueprintReadOnly, Category="Child Layout", meta=( editcondition="bOverride_MinDesiredWidth" ))
+	UPROPERTY(EditAnywhere, BlueprintReadOnly, Category="Child Layout", meta=( editcondition="bOverride_MinDesiredWidth" ))
 	float MinDesiredWidth;
 
 	/** When specified, will report the MinDesiredHeight if larger than the content's desired height. */
-	UPROPERTY(EditDefaultsOnly, BlueprintReadOnly, Category="Child Layout", meta=( editcondition="bOverride_MinDesiredHeight" ))
+	UPROPERTY(EditAnywhere, BlueprintReadOnly, Category="Child Layout", meta=( editcondition="bOverride_MinDesiredHeight" ))
 	float MinDesiredHeight;
 
 	/** When specified, will report the MaxDesiredWidth if smaller than the content's desired width. */
-	UPROPERTY(EditDefaultsOnly, BlueprintReadOnly, Category="Child Layout", meta=( editcondition="bOverride_MaxDesiredWidth" ))
+	UPROPERTY(EditAnywhere, BlueprintReadOnly, Category="Child Layout", meta=( editcondition="bOverride_MaxDesiredWidth" ))
 	float MaxDesiredWidth;
 
 	/** When specified, will report the MaxDesiredHeight if smaller than the content's desired height. */
-	UPROPERTY(EditDefaultsOnly, BlueprintReadOnly, Category="Child Layout", meta=( editcondition="bOverride_MaxDesiredHeight" ))
+	UPROPERTY(EditAnywhere, BlueprintReadOnly, Category="Child Layout", meta=( editcondition="bOverride_MaxDesiredHeight" ))
 	float MaxDesiredHeight;
 
 public:
diff --git a/Engine/Source/Runtime/UMG/Public/Components/SizeBoxSlot.h b/Engine/Source/Runtime/UMG/Public/Components/SizeBoxSlot.h
--- a/Engine/Source/Runtime/UMG/Public/Components/SizeBoxSlot.h
+++ b/Engine/Source/Runtime/UMG/Public/Components/SizeBoxSlot.h
@@ -14,15 +14,15 @@ class UMG_API USizeBoxSlot : public UPanelSlot
 
 public:
 	/** The padding area between the slot and the content it contains. */
-	UPROPERTY(EditDefaultsOnly, BlueprintReadOnly, Category="Layout|SizeBox Slot")
+	UPROPERTY(EditAnywhere, BlueprintReadOnly, Category="Layout|SizeBox Slot")
 	FMargin Padding;
 
 	/** The alignment of the object horizontally. */
-	UPROPERTY(EditDefaultsOnly, BlueprintReadOnly, Category="Layout|SizeBox Slot")
+	UPROPERTY(EditAnywhere, BlueprintReadOnly, Category="Layout|SizeBox Slot")
 	TEnumAsByte<EHorizontalAlignment> HorizontalAlignment;
 
 	/** The alignment of the object vertically. */
-	UPROPERTY(EditDefaultsOnly, BlueprintReadOnly, Category="Layout|SizeBox Slot")
+	UPROPERTY(EditAnywhere, BlueprintReadOnly, Category="Layout|SizeBox Slot")
 	TEnumAsByte<EVerticalAlignment> VerticalAlignment;
 
 public:
diff --git a/Engine/Source/Runtime/UMG/Public/Components/Slider.h b/Engine/Source/Runtime/UMG/Public/Components/Slider.h
--- a/Engine/Source/Runtime/UMG/Public/Components/Slider.h
+++ b/Engine/Source/Runtime/UMG/Public/Components/Slider.h
@@ -20,7 +20,7 @@ class UMG_API USlider : public UWidget
 
 public:
 	/** The volume value to display. */
-	UPROPERTY(EditDefaultsOnly, Category=Appearance)
+	UPROPERTY(EditAnywhere, Category=Appearance)
 	float Value;
 
 	/** A bindable delegate to allow logic to drive the value of the widget */
@@ -34,23 +34,23 @@ public:
 	FSliderStyle WidgetStyle;
 
 	/** The slider's orientation. */
-	UPROPERTY(EditDefaultsOnly, BlueprintReadOnly, Category=Appearance)
+	UPROPERTY(EditAnywhere, BlueprintReadOnly, Category=Appearance)
 	TEnumAsByte<EOrientation> Orientation;
 
 	/** The color to draw the slider bar in. */
-	UPROPERTY(EditDefaultsOnly, BlueprintReadOnly, Category=Appearance)
+	UPROPERTY(EditAnywhere, BlueprintReadOnly, Category=Appearance)
 	FLinearColor SliderBarColor;
 
 	/** The color to draw the slider handle in. */
-	UPROPERTY(EditDefaultsOnly, BlueprintReadOnly, Category=Appearance)
+	UPROPERTY(EditAnywhere, BlueprintReadOnly, Category=Appearance)
 	FLinearColor SliderHandleColor;
 
 	/** Whether the slidable area should be indented to fit the handle. */
-	UPROPERTY(EditDefaultsOnly, BlueprintReadOnly, Category=Appearance, AdvancedDisplay)
+	UPROPERTY(EditAnywhere, BlueprintReadOnly, Category=Appearance, AdvancedDisplay)
 	bool IndentHandle;
 
 	/** Whether the handle is interactive or fixed. */
-	UPROPERTY(EditDefaultsOnly, BlueprintReadOnly, Category=Appearance, AdvancedDisplay)
+	UPROPERTY(EditAnywhere, BlueprintReadOnly, Category=Appearance, AdvancedDisplay)
 	bool Locked;
 
 public:
diff --git a/Engine/Source/Runtime/UMG/Public/Components/Spacer.h b/Engine/Source/Runtime/UMG/Public/Components/Spacer.h
--- a/Engine/Source/Runtime/UMG/Public/Components/Spacer.h
+++ b/Engine/Source/Runtime/UMG/Public/Components/Spacer.h
@@ -17,7 +17,7 @@ class UMG_API USpacer : public UWidget
 public:
 
 	/** The size of the spacer */
-	UPROPERTY(EditDefaultsOnly, BlueprintReadOnly, Category=Appearance)
+	UPROPERTY(EditAnywhere, BlueprintReadOnly, Category=Appearance)
 	FVector2D Size;
 
 public:
diff --git a/Engine/Source/Runtime/UMG/Public/Components/SpinBox.h b/Engine/Source/Runtime/UMG/Public/Components/SpinBox.h
--- a/Engine/Source/Runtime/UMG/Public/Components/SpinBox.h
+++ b/Engine/Source/Runtime/UMG/Public/Components/SpinBox.h
@@ -20,7 +20,7 @@ public:
 public:
 
 	/** Value stored in this spin box */
-	UPROPERTY(EditDefaultsOnly, Category=Content)
+	UPROPERTY(EditAnywhere, Category=Content)
 	float Value;
 
 	/** A bindable delegate to allow logic to drive the value of the widget */
@@ -36,30 +36,30 @@ public:
 	USlateWidgetStyleAsset* Style_DEPRECATED;
 
 	/** The amount by which to change the spin box value as the slider moves. */
-	UPROPERTY(EditDefaultsOnly, BlueprintReadOnly, Category="Slider")
+	UPROPERTY(EditAnywhere, BlueprintReadOnly, Category="Slider")
 	float Delta;
 
 	/** The exponent by which to increase the delta as the mouse moves. 1 is constant (never increases the delta). */
-	UPROPERTY(EditDefaultsOnly, BlueprintReadOnly, Category="Slider")
+	UPROPERTY(EditAnywhere, BlueprintReadOnly, Category="Slider")
 	float SliderExponent;
 	
 	/** Font color and opacity (overrides style) */
-	UPROPERTY(EditDefaultsOnly, BlueprintReadOnly, Category="Display")
+	UPROPERTY(EditAnywhere, BlueprintReadOnly, Category="Display")
 	FSlateFontInfo Font;
 
 	/** The minimum width of the spin box */
-	UPROPERTY(EditDefaultsOnly, BlueprintReadOnly, Category="Display", AdvancedDisplay, DisplayName = "Minimum Desired Width")
+	UPROPERTY(EditAnywhere, BlueprintReadOnly, Category="Display", AdvancedDisplay, DisplayName = "Minimum Desired Width")
 	float MinDesiredWidth;
 
 	/** Whether to remove the keyboard focus from the spin box when the value is committed */
-	UPROPERTY(EditDefaultsOnly, BlueprintReadOnly, Category="Input", AdvancedDisplay)
+	UPROPERTY(EditAnywhere, BlueprintReadOnly, Category="Input", AdvancedDisplay)
 	bool ClearKeyboardFocusOnCommit;
 
 	/** Whether to select the text in the spin box when the value is committed */
-	UPROPERTY(EditDefaultsOnly, BlueprintReadOnly, Category="Input", AdvancedDisplay)
+	UPROPERTY(EditAnywhere, BlueprintReadOnly, Category="Input", AdvancedDisplay)
 	bool SelectAllTextOnCommit;
 
-	UPROPERTY(EditDefaultsOnly, BlueprintReadOnly, Category="Style")
+	UPROPERTY(EditAnywhere, BlueprintReadOnly, Category="Style")
 	FSlateColor ForegroundColor;
 
 public:
@@ -191,19 +191,19 @@ protected:
 	uint32 bOverride_MaxSliderValue : 1;
 
 	/** The minimum allowable value that can be manually entered into the spin box */
-	UPROPERTY(EditDefaultsOnly, Category = Content, DisplayName = "Minimum Value", meta = (editcondition = "bOverride_MinValue"))
+	UPROPERTY(EditAnywhere, Category = Content, DisplayName = "Minimum Value", meta = (editcondition = "bOverride_MinValue"))
 	float MinValue;
 
 	/** The maximum allowable value that can be manually entered into the spin box */
-	UPROPERTY(EditDefaultsOnly, Category = Content, DisplayName = "Maximum Value", meta = (editcondition = "bOverride_MaxValue"))
+	UPROPERTY(EditAnywhere, Category = Content, DisplayName = "Maximum Value", meta = (editcondition = "bOverride_MaxValue"))
 	float MaxValue;
 
 	/** The minimum allowable value that can be specified using the slider */
-	UPROPERTY(EditDefaultsOnly, Category = Content, DisplayName = "Minimum Slider Value", meta = (editcondition = "bOverride_MinSliderValue"))
+	UPROPERTY(EditAnywhere, Category = Content, DisplayName = "Minimum Slider Value", meta = (editcondition = "bOverride_MinSliderValue"))
 	float MinSliderValue;
 
 	/** The maximum allowable value that can be specified using the slider */
-	UPROPERTY(EditDefaultsOnly, Category = Content, DisplayName = "Maximum Slider Value", meta = (editcondition = "bOverride_MaxSliderValue"))
+	UPROPERTY(EditAnywhere, Category = Content, DisplayName = "Maximum Slider Value", meta = (editcondition = "bOverride_MaxSliderValue"))
 	float MaxSliderValue;
 
 protected:
diff --git a/Engine/Source/Runtime/UMG/Public/Components/TextBlock.h b/Engine/Source/Runtime/UMG/Public/Components/TextBlock.h
--- a/Engine/Source/Runtime/UMG/Public/Components/TextBlock.h
+++ b/Engine/Source/Runtime/UMG/Public/Components/TextBlock.h
@@ -41,7 +41,6 @@ public:
 	UFUNCTION(BlueprintCallable, Category="Appearance")
 	void SetShadowOffset(FVector2D InShadowOffset);
 
-
 	/**
 	 *  Set the text justification for this text block
 	 *
@@ -55,7 +54,7 @@ public:
 	USlateWidgetStyleAsset* Style_DEPRECATED;
 
 	/** The text to display */
-	UPROPERTY(EditDefaultsOnly, Category=Content, meta=( MultiLine="true" ))
+	UPROPERTY(EditAnywhere, Category=Content, meta=( MultiLine="true" ))
 	FText Text;
 	
 	/** A bindable delegate to allow logic to drive the text of the widget */
@@ -63,7 +62,7 @@ public:
 	FGetText TextDelegate;
 
 	/** The color of the text */
-	UPROPERTY(EditDefaultsOnly, BlueprintReadOnly, Category=Appearance)
+	UPROPERTY(EditAnywhere, BlueprintReadOnly, Category=Appearance)
 	FSlateColor ColorAndOpacity;
 
 	/** A bindable delegate for the ColorAndOpacity. */
@@ -71,15 +70,15 @@ public:
 	FGetSlateColor ColorAndOpacityDelegate;
 	
 	/** The font to render the text with */
-	UPROPERTY(EditDefaultsOnly, BlueprintReadOnly, Category=Appearance)
+	UPROPERTY(EditAnywhere, BlueprintReadOnly, Category=Appearance)
 	FSlateFontInfo Font;
 
 	/** The direction the shadow is cast */
-	UPROPERTY(EditDefaultsOnly, BlueprintReadOnly, Category=Appearance)
+	UPROPERTY(EditAnywhere, BlueprintReadOnly, Category=Appearance)
 	FVector2D ShadowOffset;
 
 	/** The color of the shadow */
-	UPROPERTY(EditDefaultsOnly, BlueprintReadOnly, Category=Appearance, meta=( DisplayName="Shadow Color" ))
+	UPROPERTY(EditAnywhere, BlueprintReadOnly, Category=Appearance, meta=( DisplayName="Shadow Color" ))
 	FLinearColor ShadowColorAndOpacity;
 
 	/** A bindable delegate for the ShadowColorAndOpacity. */
@@ -87,27 +86,27 @@ public:
 	FGetLinearColor ShadowColorAndOpacityDelegate;
 
 	/** How the text should be aligned with the margin. */
-	UPROPERTY(EditDefaultsOnly, BlueprintReadOnly, Category=Appearance)
+	UPROPERTY(EditAnywhere, BlueprintReadOnly, Category=Appearance)
 	TEnumAsByte<ETextJustify::Type> Justification;
 
 	/** True if we're wrapping text automatically based on the computed horizontal space for this widget */
-	UPROPERTY(EditDefaultsOnly, BlueprintReadOnly, Category=Appearance)
+	UPROPERTY(EditAnywhere, BlueprintReadOnly, Category=Appearance)
 	bool AutoWrapText;
 
 	/** Whether text wraps onto a new line when it's length exceeds this width; if this value is zero or negative, no wrapping occurs. */
-	UPROPERTY(EditDefaultsOnly, BlueprintReadOnly, Category=Appearance, AdvancedDisplay)
+	UPROPERTY(EditAnywhere, BlueprintReadOnly, Category=Appearance, AdvancedDisplay)
 	float WrapTextAt;
 
 	/** The minimum desired size for the text */
-	UPROPERTY(EditDefaultsOnly, BlueprintReadOnly, Category=Appearance, AdvancedDisplay)
+	UPROPERTY(EditAnywhere, BlueprintReadOnly, Category=Appearance, AdvancedDisplay)
 	float MinDesiredWidth;
 
 	/** The amount of blank space left around the edges of text area. */
-	UPROPERTY(EditDefaultsOnly, BlueprintReadOnly, Category=Appearance, AdvancedDisplay)
+	UPROPERTY(EditAnywhere, BlueprintReadOnly, Category=Appearance, AdvancedDisplay)
 	FMargin Margin;
 
 	/** The amount to scale each lines height by. */
-	UPROPERTY(EditDefaultsOnly, BlueprintReadOnly, Category=Appearance, AdvancedDisplay)
+	UPROPERTY(EditAnywhere, BlueprintReadOnly, Category=Appearance, AdvancedDisplay)
 	float LineHeightPercentage;
 
 	///** Called when this text is double clicked */
diff --git a/Engine/Source/Runtime/UMG/Public/Components/Throbber.h b/Engine/Source/Runtime/UMG/Public/Components/Throbber.h
--- a/Engine/Source/Runtime/UMG/Public/Components/Throbber.h
+++ b/Engine/Source/Runtime/UMG/Public/Components/Throbber.h
@@ -18,26 +18,26 @@ class UMG_API UThrobber : public UWidget
 public:
 	
 	/** How many pieces there are */
-	UPROPERTY(EditDefaultsOnly, BlueprintReadOnly, Category=Appearance, meta=( ClampMin="1", ClampMax="25", UIMin="1", UIMax="25" ))
+	UPROPERTY(EditAnywhere, BlueprintReadOnly, Category=Appearance, meta=( ClampMin="1", ClampMax="25", UIMin="1", UIMax="25" ))
 	int32 NumberOfPieces;
 
 	/** Should the pieces animate horizontally? */
-	UPROPERTY(EditDefaultsOnly, BlueprintReadOnly, Category=Appearance)
+	UPROPERTY(EditAnywhere, BlueprintReadOnly, Category=Appearance)
 	bool bAnimateHorizontally;
 
 	/** Should the pieces animate vertically? */
-	UPROPERTY(EditDefaultsOnly, BlueprintReadOnly, Category=Appearance)
+	UPROPERTY(EditAnywhere, BlueprintReadOnly, Category=Appearance)
 	bool bAnimateVertically;
 
 	/** Should the pieces animate their opacity? */
-	UPROPERTY(EditDefaultsOnly, BlueprintReadOnly, Category=Appearance)
+	UPROPERTY(EditAnywhere, BlueprintReadOnly, Category=Appearance)
 	bool bAnimateOpacity;
 
 	/** Image to use for each segment of the throbber */
 	UPROPERTY()
 	USlateBrushAsset* PieceImage_DEPRECATED;
 
-	UPROPERTY(EditDefaultsOnly, BlueprintReadOnly, Category=Appearance)
+	UPROPERTY(EditAnywhere, BlueprintReadOnly, Category=Appearance)
 	FSlateBrush Image;
 
 public:
diff --git a/Engine/Source/Runtime/UMG/Public/Components/TileView.h b/Engine/Source/Runtime/UMG/Public/Components/TileView.h
--- a/Engine/Source/Runtime/UMG/Public/Components/TileView.h
+++ b/Engine/Source/Runtime/UMG/Public/Components/TileView.h
@@ -15,11 +15,11 @@ class UMG_API UTileView : public UTableViewBase
 public:
 
 	/**  */
-	UPROPERTY(EditDefaultsOnly, Category=Appearance)
+	UPROPERTY(EditAnywhere, Category=Appearance)
 	float ItemWidth;
 
 	/**  */
-	UPROPERTY(EditDefaultsOnly, Category=Appearance)
+	UPROPERTY(EditAnywhere, Category=Appearance)
 	float ItemHeight;
 
 	/**  */
@@ -27,11 +27,11 @@ public:
 	TArray<UObject*> Items;
 
 	/**  */
-	UPROPERTY(EditDefaultsOnly, Category=Content)
+	UPROPERTY(EditAnywhere, Category=Content)
 	TEnumAsByte<ESelectionMode::Type> SelectionMode;
 
 	/**  */
-	UPROPERTY(EditDefaultsOnly, Category=Events)
+	UPROPERTY(EditAnywhere, Category=Events)
 	FOnGenerateRowUObject OnGenerateTileEvent;
 
 public:
diff --git a/Engine/Source/Runtime/UMG/Public/Components/UniformGridPanel.h b/Engine/Source/Runtime/UMG/Public/Components/UniformGridPanel.h
--- a/Engine/Source/Runtime/UMG/Public/Components/UniformGridPanel.h
+++ b/Engine/Source/Runtime/UMG/Public/Components/UniformGridPanel.h
@@ -15,15 +15,15 @@ class UMG_API UUniformGridPanel : public UPanelWidget
 public:
 
 	/** Padding given to each slot */
-	UPROPERTY(EditDefaultsOnly, BlueprintReadOnly, Category="Child Layout")
+	UPROPERTY(EditAnywhere, BlueprintReadOnly, Category="Child Layout")
 	FMargin SlotPadding;
 
 	/** The minimum desired width of the slots */
-	UPROPERTY(EditDefaultsOnly, BlueprintReadOnly, Category="Child Layout")
+	UPROPERTY(EditAnywhere, BlueprintReadOnly, Category="Child Layout")
 	float MinDesiredSlotWidth;
 
 	/** The minimum desired height of the slots */
-	UPROPERTY(EditDefaultsOnly, BlueprintReadOnly, Category="Child Layout")
+	UPROPERTY(EditAnywhere, BlueprintReadOnly, Category="Child Layout")
 	float MinDesiredSlotHeight;
 
 public:
diff --git a/Engine/Source/Runtime/UMG/Public/Components/UniformGridSlot.h b/Engine/Source/Runtime/UMG/Public/Components/UniformGridSlot.h
--- a/Engine/Source/Runtime/UMG/Public/Components/UniformGridSlot.h
+++ b/Engine/Source/Runtime/UMG/Public/Components/UniformGridSlot.h
@@ -18,19 +18,19 @@ class UMG_API UUniformGridSlot : public UPanelSlot
 public:
 
 	/** The alignment of the object horizontally. */
-	UPROPERTY(EditDefaultsOnly, BlueprintReadOnly, Category="Layout|Uniform Grid Slot")
+	UPROPERTY(EditAnywhere, BlueprintReadOnly, Category="Layout|Uniform Grid Slot")
 	TEnumAsByte<EHorizontalAlignment> HorizontalAlignment;
 
 	/** The alignment of the object vertically. */
-	UPROPERTY(EditDefaultsOnly, BlueprintReadOnly, Category="Layout|Uniform Grid Slot")
+	UPROPERTY(EditAnywhere, BlueprintReadOnly, Category="Layout|Uniform Grid Slot")
 	TEnumAsByte<EVerticalAlignment> VerticalAlignment;
 	
 	/** The row index of the cell this slot is in */
-	UPROPERTY(EditDefaultsOnly, BlueprintReadOnly, meta=( UIMin = "0" ), Category="Layout|Uniform Grid Slot")
+	UPROPERTY(EditAnywhere, BlueprintReadOnly, meta=( UIMin = "0" ), Category="Layout|Uniform Grid Slot")
 	int32 Row;
 	
 	/** The column index of the cell this slot is in */
-	UPROPERTY(EditDefaultsOnly, BlueprintReadOnly, meta=( UIMin = "0" ), Category="Layout|Uniform Grid Slot")
+	UPROPERTY(EditAnywhere, BlueprintReadOnly, meta=( UIMin = "0" ), Category="Layout|Uniform Grid Slot")
 	int32 Column;
 
 public:
diff --git a/Engine/Source/Runtime/UMG/Public/Components/VerticalBoxSlot.h b/Engine/Source/Runtime/UMG/Public/Components/VerticalBoxSlot.h
--- a/Engine/Source/Runtime/UMG/Public/Components/VerticalBoxSlot.h
+++ b/Engine/Source/Runtime/UMG/Public/Components/VerticalBoxSlot.h
@@ -15,19 +15,19 @@ class UMG_API UVerticalBoxSlot : public UPanelSlot
 public:
 	
 	/** The padding area between the slot and the content it contains. */
-	UPROPERTY(EditDefaultsOnly, BlueprintReadOnly, Category="Layout|Vertical Box Slot")
+	UPROPERTY(EditAnywhere, BlueprintReadOnly, Category="Layout|Vertical Box Slot")
 	FMargin Padding;
 
 	/** How much space this slot should occupy in the direction of the panel. */
-	UPROPERTY(EditDefaultsOnly, BlueprintReadOnly, Category="Layout|Vertical Box Slot")
+	UPROPERTY(EditAnywhere, BlueprintReadOnly, Category="Layout|Vertical Box Slot")
 	FSlateChildSize Size;
 
 	/** The alignment of the object horizontally. */
-	UPROPERTY(EditDefaultsOnly, BlueprintReadOnly, Category="Layout|Vertical Box Slot")
+	UPROPERTY(EditAnywhere, BlueprintReadOnly, Category="Layout|Vertical Box Slot")
 	TEnumAsByte<EHorizontalAlignment> HorizontalAlignment;
 
 	/** The alignment of the object vertically. */
-	UPROPERTY(EditDefaultsOnly, BlueprintReadOnly, Category="Layout|Vertical Box Slot")
+	UPROPERTY(EditAnywhere, BlueprintReadOnly, Category="Layout|Vertical Box Slot")
 	TEnumAsByte<EVerticalAlignment> VerticalAlignment;
 
 public:
diff --git a/Engine/Source/Runtime/UMG/Public/Components/Viewport.h b/Engine/Source/Runtime/UMG/Public/Components/Viewport.h
--- a/Engine/Source/Runtime/UMG/Public/Components/Viewport.h
+++ b/Engine/Source/Runtime/UMG/Public/Components/Viewport.h
@@ -221,7 +221,7 @@ class UMG_API UViewport : public UContentWidget
 {
 	GENERATED_UCLASS_BODY()
 
-	UPROPERTY(EditDefaultsOnly, Category=Appearance)
+	UPROPERTY(EditAnywhere, Category=Appearance)
 	FLinearColor BackgroundColor;
 
 	UFUNCTION(BlueprintCallable, Category="Viewport")
diff --git a/Engine/Source/Runtime/UMG/Public/Components/Widget.h b/Engine/Source/Runtime/UMG/Public/Components/Widget.h
--- a/Engine/Source/Runtime/UMG/Public/Components/Widget.h
+++ b/Engine/Source/Runtime/UMG/Public/Components/Widget.h
@@ -119,7 +119,7 @@ public:
 	UPanelSlot* Slot;
 
 	/** Sets whether this widget can be modified interactively by the user */
-	UPROPERTY(EditDefaultsOnly, BlueprintReadOnly, Category="Behavior")
+	UPROPERTY(EditAnywhere, BlueprintReadOnly, Category="Behavior")
 	bool bIsEnabled;
 
 	/** A bindable delegate for bIsEnabled */
@@ -127,7 +127,7 @@ public:
 	FGetBool bIsEnabledDelegate;
 
 	/** Tooltip text to show when the user hovers over the widget with the mouse */
-	UPROPERTY(EditDefaultsOnly, BlueprintReadOnly, Category="Behavior", meta=(MultiLine=true))
+	UPROPERTY(EditAnywhere, BlueprintReadOnly, Category="Behavior", meta=(MultiLine=true))
 	FText ToolTipText;
 
 	/** A bindable delegate for ToolTipText */
@@ -135,7 +135,7 @@ public:
 	FGetText ToolTipTextDelegate;
 
 	/** Tooltip widget to show when the user hovers over the widget with the mouse */
-	UPROPERTY(EditDefaultsOnly, Category="Behavior", AdvancedDisplay)
+	UPROPERTY(EditAnywhere, Category="Behavior", AdvancedDisplay)
 	UWidget* ToolTipWidget;
 
 	/** A bindable delegate for ToolTipWidget */
@@ -147,7 +147,7 @@ public:
 	TEnumAsByte<ESlateVisibility> Visiblity_DEPRECATED;
 
 	/** The visibility of the widget */
-	UPROPERTY(EditDefaultsOnly, Category="Behavior")
+	UPROPERTY(EditAnywhere, Category="Behavior")
 	ESlateVisibility Visibility;
 
 	/** A bindable delegate for Visibility */
@@ -155,7 +155,7 @@ public:
 	FGetSlateVisibility VisibilityDelegate;
 
 	/** The cursor to show when the mouse is over the widget */
-	UPROPERTY(EditDefaultsOnly, BlueprintReadOnly, Category="Behavior", AdvancedDisplay)
+	UPROPERTY(EditAnywhere, BlueprintReadOnly, Category="Behavior", AdvancedDisplay)
 	TEnumAsByte<EMouseCursor::Type> Cursor;
 
 	/** A bindable delegate for Cursor */
@@ -163,14 +163,14 @@ public:
 	FGetMouseCursor CursorDelegate;
 
 	/** The render transform of the widget allows for arbitrary 2D transforms to be applied to the widget. */
-	UPROPERTY(EditDefaultsOnly, BlueprintReadOnly, Category="Render Transform", meta=( DisplayName="Transform" ))
+	UPROPERTY(EditAnywhere, BlueprintReadOnly, Category="Render Transform", meta=( DisplayName="Transform" ))
 	FWidgetTransform RenderTransform;
 
 	/**
 	 * The render transform pivot controls the location about which transforms are applied.  
 	 * This value is a normalized coordinate about which things like rotations will occur.
 	 */
-	UPROPERTY(EditDefaultsOnly, BlueprintReadOnly, Category="Render Transform", meta=( DisplayName="Pivot" ))
+	UPROPERTY(EditAnywhere, BlueprintReadOnly, Category="Render Transform", meta=( DisplayName="Pivot" ))
 	FVector2D RenderTransformPivot;
 
 	/**
@@ -178,7 +178,7 @@ public:
 	 * navigation rules for this widget in the widget designer.  Those rules determine how navigation transitions
 	 * can occur between widgets.
 	 */
-	UPROPERTY(Instanced, EditDefaultsOnly, BlueprintReadOnly, Category="Navigation")
+	UPROPERTY(Instanced, EditAnywhere, BlueprintReadOnly, Category="Navigation")
 	class UWidgetNavigation* Navigation;
 
 #if WITH_EDITORONLY_DATA
@@ -435,8 +435,6 @@ public:
 
 	static FSizeParam ConvertSerializedSizeParamToRuntime(const FSlateChildSize& Input);
 
-	static void GatherChildren(UWidget* Root, TSet<UWidget*>& Children);
-	static void GatherAllChildren(UWidget* Root, TSet<UWidget*>& Children);
 	static UWidget* FindChildContainingDescendant(UWidget* Root, UWidget* Descendant);
 
 protected:
diff --git a/Engine/Source/Runtime/UMG/Public/Components/WidgetComponent.h b/Engine/Source/Runtime/UMG/Public/Components/WidgetComponent.h
--- a/Engine/Source/Runtime/UMG/Public/Components/WidgetComponent.h
+++ b/Engine/Source/Runtime/UMG/Public/Components/WidgetComponent.h
@@ -141,7 +141,7 @@ private:
 	FVector2D Pivot;
 
 	/** The Screen Space ZOrder that is used if the widget is displayed in screen space. */
-	UPROPERTY(EditDefaultsOnly, Category=UI)
+	UPROPERTY(EditAnywhere, Category=UI)
 	int32 ZOrder;
 	
 	/** The maximum distance from which a player can interact with this widget */
diff --git a/Engine/Source/Runtime/UMG/Public/Components/WidgetSwitcher.h b/Engine/Source/Runtime/UMG/Public/Components/WidgetSwitcher.h
--- a/Engine/Source/Runtime/UMG/Public/Components/WidgetSwitcher.h
+++ b/Engine/Source/Runtime/UMG/Public/Components/WidgetSwitcher.h
@@ -14,7 +14,7 @@ class UMG_API UWidgetSwitcher : public UPanelWidget
 
 public:
 	/** The slot index to display */
-	UPROPERTY(EditDefaultsOnly, BlueprintReadOnly, Category="Switcher", meta=( UIMin=0, ClampMin=0 ))
+	UPROPERTY(EditAnywhere, BlueprintReadOnly, Category="Switcher", meta=( UIMin=0, ClampMin=0 ))
 	int32 ActiveWidgetIndex;
 
 public:
diff --git a/Engine/Source/Runtime/UMG/Public/Components/WidgetSwitcherSlot.h b/Engine/Source/Runtime/UMG/Public/Components/WidgetSwitcherSlot.h
--- a/Engine/Source/Runtime/UMG/Public/Components/WidgetSwitcherSlot.h
+++ b/Engine/Source/Runtime/UMG/Public/Components/WidgetSwitcherSlot.h
@@ -16,15 +16,15 @@ class UMG_API UWidgetSwitcherSlot : public UPanelSlot
 public:
 	
 	/** The padding area between the slot and the content it contains. */
-	UPROPERTY(EditDefaultsOnly, BlueprintReadOnly, Category="Layout|Widget Switcher Slot")
+	UPROPERTY(EditAnywhere, BlueprintReadOnly, Category="Layout|Widget Switcher Slot")
 	FMargin Padding;
 
 	/** The alignment of the object horizontally. */
-	UPROPERTY(EditDefaultsOnly, BlueprintReadOnly, Category="Layout|Widget Switcher Slot")
+	UPROPERTY(EditAnywhere, BlueprintReadOnly, Category="Layout|Widget Switcher Slot")
 	TEnumAsByte<EHorizontalAlignment> HorizontalAlignment;
 
 	/** The alignment of the object vertically. */
-	UPROPERTY(EditDefaultsOnly, BlueprintReadOnly, Category="Layout|Widget Switcher Slot")
+	UPROPERTY(EditAnywhere, BlueprintReadOnly, Category="Layout|Widget Switcher Slot")
 	TEnumAsByte<EVerticalAlignment> VerticalAlignment;
 
 public:
diff --git a/Engine/Source/Runtime/UMG/Public/Components/WrapBox.h b/Engine/Source/Runtime/UMG/Public/Components/WrapBox.h
--- a/Engine/Source/Runtime/UMG/Public/Components/WrapBox.h
+++ b/Engine/Source/Runtime/UMG/Public/Components/WrapBox.h
@@ -18,7 +18,7 @@ class UMG_API UWrapBox : public UPanelWidget
 
 public:
 	/** The inner slot padding goes between slots sharing borders */
-	UPROPERTY(EditDefaultsOnly, BlueprintReadOnly, Category="Content Layout")
+	UPROPERTY(EditAnywhere, BlueprintReadOnly, Category="Content Layout")
 	FVector2D InnerSlotPadding;
 
 	/** Sets the inner slot padding goes between slots sharing borders */
diff --git a/Engine/Source/Runtime/UMG/Public/Components/WrapBoxSlot.h b/Engine/Source/Runtime/UMG/Public/Components/WrapBoxSlot.h
--- a/Engine/Source/Runtime/UMG/Public/Components/WrapBoxSlot.h
+++ b/Engine/Source/Runtime/UMG/Public/Components/WrapBoxSlot.h
@@ -15,26 +15,26 @@ class UMG_API UWrapBoxSlot : public UPanelSlot
 public:
 
 	/** The padding area between the slot and the content it contains. */
-	UPROPERTY(EditDefaultsOnly, BlueprintReadOnly, Category="Layout|Wrap Box Slot")
+	UPROPERTY(EditAnywhere, BlueprintReadOnly, Category="Layout|Wrap Box Slot")
 	FMargin Padding;
 
 	/** Should this slot fill the remaining space on the line? */
-	UPROPERTY(EditDefaultsOnly, BlueprintReadOnly, Category="Layout|Wrap Box Slot")
+	UPROPERTY(EditAnywhere, BlueprintReadOnly, Category="Layout|Wrap Box Slot")
 	bool bFillEmptySpace;
 
 	/**
 	 * If the total available space in the wrap panel drops below this threshold, this slot will attempt to fill an entire line.
 	 * NOTE: A value of 0, denotes no filling will occur.
 	 */
-	UPROPERTY(EditDefaultsOnly, BlueprintReadOnly, Category="Layout|Wrap Box Slot")
+	UPROPERTY(EditAnywhere, BlueprintReadOnly, Category="Layout|Wrap Box Slot")
 	float FillSpanWhenLessThan;
 
 	/** The alignment of the object horizontally. */
-	UPROPERTY(EditDefaultsOnly, BlueprintReadOnly, Category="Layout|Wrap Box Slot")
+	UPROPERTY(EditAnywhere, BlueprintReadOnly, Category="Layout|Wrap Box Slot")
 	TEnumAsByte<EHorizontalAlignment> HorizontalAlignment;
 
 	/** The alignment of the object vertically. */
-	UPROPERTY(EditDefaultsOnly, BlueprintReadOnly, Category="Layout|Wrap Box Slot")
+	UPROPERTY(EditAnywhere, BlueprintReadOnly, Category="Layout|Wrap Box Slot")
 	TEnumAsByte<EVerticalAlignment> VerticalAlignment;
 
 public:
