UMG - The designer now allows you to set a background image on the root widget to be visualized on the mock screen to get a better idea of scale.

diff --git a/Engine/Source/Editor/UMGEditor/Private/Designer/SDesignerView.cpp b/Engine/Source/Editor/UMGEditor/Private/Designer/SDesignerView.cpp
--- a/Engine/Source/Editor/UMGEditor/Private/Designer/SDesignerView.cpp
+++ b/Engine/Source/Editor/UMGEditor/Private/Designer/SDesignerView.cpp
@@ -225,12 +225,17 @@ void SDesignerView::Construct(const FArguments& InArgs, TSharedPtr<FWidgetBluepr
 
 					+ SOverlay::Slot()
 					[
-						SAssignNew(PreviewAreaConstraint, SBox)
-						.WidthOverride(this, &SDesignerView::GetPreviewAreaWidth)
-						.HeightOverride(this, &SDesignerView::GetPreviewAreaHeight)
+						SNew(SBorder)
+						.Padding(FMargin(0))
+						.BorderImage(this, &SDesignerView::GetPreviewBackground)
 						[
-							SAssignNew(PreviewSurface, SDPIScaler)
-							.DPIScale(this, &SDesignerView::GetPreviewDPIScale)
+							SAssignNew(PreviewAreaConstraint, SBox)
+							.WidthOverride(this, &SDesignerView::GetPreviewAreaWidth)
+							.HeightOverride(this, &SDesignerView::GetPreviewAreaHeight)
+							[
+								SAssignNew(PreviewSurface, SDPIScaler)
+								.DPIScale(this, &SDesignerView::GetPreviewDPIScale)
+							]
 						]
 					]
 				]
@@ -690,6 +695,20 @@ FOptionalSize SDesignerView::GetPreviewSizeHeight() const
 	return Size.Y;
 }
 
+const FSlateBrush* SDesignerView::GetPreviewBackground() const
+{
+	if ( UUserWidget* DefaultWidget = GetDefaultWidget() )
+	{
+		if ( DefaultWidget->PreviewBackground )
+		{
+			BackgroundImage.SetResourceObject(DefaultWidget->PreviewBackground);
+			return &BackgroundImage;
+		}
+	}
+
+	return nullptr;
+}
+
 void SDesignerView::GetPreviewAreaAndSize(FVector2D& Area, FVector2D& Size) const
 {
 	Area = FVector2D(PreviewWidth, PreviewHeight);
diff --git a/Engine/Source/Editor/UMGEditor/Private/Designer/SDesignerView.h b/Engine/Source/Editor/UMGEditor/Private/Designer/SDesignerView.h
--- a/Engine/Source/Editor/UMGEditor/Private/Designer/SDesignerView.h
+++ b/Engine/Source/Editor/UMGEditor/Private/Designer/SDesignerView.h
@@ -92,6 +92,8 @@ private:
 	/** The height of the preview widget for the UI */
 	FOptionalSize GetPreviewSizeHeight() const;
 
+	const FSlateBrush* GetPreviewBackground() const;
+
 	void GetPreviewAreaAndSize(FVector2D& Area, FVector2D& Size) const;
 
 	/** Gets the DPI scale that would be applied given the current preview width and height */
@@ -260,6 +262,9 @@ private:
 	/** The configured Height of the preview area, simulates screen size. */
 	int32 PreviewHeight;
 
+	/** The slate brush we use to hold the background image shown in the designer. */
+	mutable FSlateBrush BackgroundImage;
+
 	/** We cache the desired preview desired size to maintain the same size between compiles when it lags a frame behind and no widget is available. */
 	FVector2D CachedPreviewDesiredSize;
 
