UMG - Properly handling the F2 button for renaming a widget in the hierarchy view.

diff --git a/Engine/Source/Editor/UMGEditor/Private/Hierarchy/SHierarchyView.cpp b/Engine/Source/Editor/UMGEditor/Private/Hierarchy/SHierarchyView.cpp
--- a/Engine/Source/Editor/UMGEditor/Private/Hierarchy/SHierarchyView.cpp
+++ b/Engine/Source/Editor/UMGEditor/Private/Hierarchy/SHierarchyView.cpp
@@ -43,6 +43,14 @@ void SHierarchyView::Construct(const FArguments& InArgs, TSharedPtr<FWidgetBluep
 	FilterHandler->SetRootItems(&RootWidgets, &TreeRootWidgets);
 	FilterHandler->SetGetChildrenDelegate(TreeFilterHandler< TSharedPtr<FHierarchyModel> >::FOnGetChildren::CreateRaw(this, &SHierarchyView::WidgetHierarchy_OnGetChildren));
 
+	CommandList = MakeShareable(new FUICommandList);
+
+	CommandList->MapAction(
+		FGenericCommands::Get().Rename,
+		FExecuteAction::CreateSP(this, &SHierarchyView::BeginRename),
+		FCanExecuteAction::CreateSP(this, &SHierarchyView::CanRename)
+	);
+
 	ChildSlot
 	[
 		SNew(SBorder)
@@ -136,9 +144,26 @@ FReply SHierarchyView::OnKeyDown(const FGeometry& MyGeometry, const FKeyEvent& I
 		return FReply::Handled();
 	}
 
+	if ( CommandList->ProcessCommandBindings(InKeyEvent) )
+	{
+		return FReply::Handled();
+	}
+
 	return FReply::Unhandled();
 }
 
+void SHierarchyView::BeginRename()
+{
+	TArray< TSharedPtr<FHierarchyModel> > SelectedItems = WidgetTreeView->GetSelectedItems();
+	SelectedItems[0]->BeginRename();
+}
+
+bool SHierarchyView::CanRename() const
+{
+	TArray< TSharedPtr<FHierarchyModel> > SelectedItems = WidgetTreeView->GetSelectedItems();
+	return SelectedItems.Num() == 1 && SelectedItems[0]->CanRename();
+}
+
 void SHierarchyView::TransformWidgetToString(TSharedPtr<FHierarchyModel> Item, OUT TArray< FString >& Array)
 {
 	Array.Add( Item->GetText().ToString() );
@@ -196,6 +221,8 @@ TSharedPtr<SWidget> SHierarchyView::WidgetHierarchy_OnContextMenuOpening()
 
 	FWidgetBlueprintEditorUtils::CreateWidgetContextMenu(MenuBuilder, BlueprintEditor.Pin().ToSharedRef(), FVector2D(0, 0));
 
+	MenuBuilder.AddMenuEntry(FGenericCommands::Get().Rename);
+
 	return MenuBuilder.MakeWidget();
 }
 
diff --git a/Engine/Source/Editor/UMGEditor/Private/Hierarchy/SHierarchyView.h b/Engine/Source/Editor/UMGEditor/Private/Hierarchy/SHierarchyView.h
--- a/Engine/Source/Editor/UMGEditor/Private/Hierarchy/SHierarchyView.h
+++ b/Engine/Source/Editor/UMGEditor/Private/Hierarchy/SHierarchyView.h
@@ -13,8 +13,6 @@ class FWidgetBlueprintEditor;
 class UWidget;
 class UWidgetBlueprint;
 
-//TODO rename SUMGEditorHierarchy
-
 /**
  * The tree view presenting the widget hierarchy.  This allows users to edit the hierarchy of widgets easily by dragging and 
  * dropping them logically, which in some cases may be significantly easier than doing it visually in the widget designer.
@@ -48,6 +46,9 @@ private:
 	void WidgetHierarchy_OnExpansionChanged(TSharedPtr<FHierarchyModel> Item, bool bExpanded);
 
 private:
+	void BeginRename();
+	bool CanRename() const;
+
 	/** @returns the current blueprint being edited */
 	UWidgetBlueprint* GetBlueprint() const;
 
@@ -104,6 +105,9 @@ private:
 	/** Cached pointer to the blueprint editor that owns this tree. */
 	TWeakPtr<class FWidgetBlueprintEditor> BlueprintEditor;
 
+	/** Commands specific to the hierarchy. */
+	TSharedPtr<FUICommandList> CommandList;
+
 	/** Handles filtering the hierarchy based on an IFilter. */
 	TSharedPtr< TreeFilterHandler< TSharedPtr<FHierarchyModel> > > FilterHandler;
 
diff --git a/Engine/Source/Editor/UMGEditor/Private/Hierarchy/SHierarchyViewItem.cpp b/Engine/Source/Editor/UMGEditor/Private/Hierarchy/SHierarchyViewItem.cpp
--- a/Engine/Source/Editor/UMGEditor/Private/Hierarchy/SHierarchyViewItem.cpp
+++ b/Engine/Source/Editor/UMGEditor/Private/Hierarchy/SHierarchyViewItem.cpp
@@ -36,7 +36,7 @@ class FHierarchyWidgetDragDropOp : public FDecoratedDragDropOp
 public:
 	DRAG_DROP_OPERATOR_TYPE(FHierarchyWidgetDragDropOp, FDecoratedDragDropOp)
 
-	virtual ~FHierarchyWidgetDragDropOp();
+		virtual ~FHierarchyWidgetDragDropOp();
 
 	virtual void OnDrop(bool bDropWasHandled, const FPointerEvent& MouseEvent) override;
 
@@ -103,7 +103,7 @@ TOptional<EItemDropZone> ProcessHierarchyDragDrop(const FDragDropEvent& DragDrop
 {
 	UWidget* TargetTemplate = TargetItem.GetTemplate();
 
-	if ( TargetTemplate && (DropZone == EItemDropZone::AboveItem || DropZone == EItemDropZone::BelowItem) )
+	if ( TargetTemplate && ( DropZone == EItemDropZone::AboveItem || DropZone == EItemDropZone::BelowItem ) )
 	{
 		if ( UPanelWidget* TargetParentTemplate = Cast<UPanelWidget>(TargetTemplate->GetParent()) )
 		{
@@ -155,7 +155,7 @@ TOptional<EItemDropZone> ProcessHierarchyDragDrop(const FDragDropEvent& DragDrop
 			if ( bIsDrop )
 			{
 				UWidget* Widget = TemplateDragDropOp->Template->Create(Blueprint->WidgetTree);
-				
+
 				UPanelSlot* NewSlot = nullptr;
 				if ( Index.IsSet() )
 				{
@@ -285,6 +285,7 @@ TOptional<EItemDropZone> ProcessHierarchyDragDrop(const FDragDropEvent& DragDrop
 	return TOptional<EItemDropZone>();
 }
 
+
 //////////////////////////////////////////////////////////////////////////
 
 FHierarchyModel::FHierarchyModel()
@@ -387,6 +388,7 @@ bool FHierarchyModel::IsSelected() const
 FHierarchyRoot::FHierarchyRoot(TSharedPtr<FWidgetBlueprintEditor> InBlueprintEditor)
 	: BlueprintEditor(InBlueprintEditor)
 {
+	RootText = FText::Format(LOCTEXT("RootWidgetFormat", "[{0}]"), FText::FromString(BlueprintEditor.Pin()->GetBlueprintObj()->GetName()));
 }
 
 FName FHierarchyRoot::GetUniqueName() const
@@ -397,7 +399,7 @@ FName FHierarchyRoot::GetUniqueName() const
 
 FText FHierarchyRoot::GetText() const
 {
-	return LOCTEXT("Root", "[Root]");
+	return RootText;
 }
 
 const FSlateBrush* FHierarchyRoot::GetImage() const
@@ -810,11 +812,22 @@ void FHierarchyWidget::UpdateSelection()
 	bIsSelected = SelectedWidgets.Contains(Item);
 }
 
+bool FHierarchyWidget::CanRename() const
+{
+	return true;
+}
+
+void FHierarchyWidget::BeginRename()
+{
+	RenameEvent.ExecuteIfBound();
+}
+
 //////////////////////////////////////////////////////////////////////////
 
 void SHierarchyViewItem::Construct(const FArguments& InArgs, const TSharedRef< STableViewBase >& InOwnerTableView, TSharedPtr<FHierarchyModel> InModel)
 {
 	Model = InModel;
+	Model->RenameEvent.BindSP(this, &SHierarchyViewItem::BeginRename);
 
 	STableRow< TSharedPtr<FHierarchyModel> >::Construct(
 		STableRow< TSharedPtr<FHierarchyModel> >::FArguments()
@@ -844,7 +857,7 @@ void SHierarchyViewItem::Construct(const FArguments& InArgs, const TSharedRef< S
 			.Padding(2, 0, 0, 0)
 			.VAlign(VAlign_Center)
 			[
-				SNew(SInlineEditableTextBlock)
+				SAssignNew(EditBox, SInlineEditableTextBlock)
 				.Font(this, &SHierarchyViewItem::GetItemFont)
 				.Text(this, &SHierarchyViewItem::GetItemText)
 				.ToolTipText(Model->GetLabelToolTipText())
@@ -878,6 +891,11 @@ void SHierarchyViewItem::Construct(const FArguments& InArgs, const TSharedRef< S
 		InOwnerTableView);
 }
 
+SHierarchyViewItem::~SHierarchyViewItem()
+{
+	Model->RenameEvent.Unbind();
+}
+
 void SHierarchyViewItem::OnMouseEnter(const FGeometry& MyGeometry, const FPointerEvent& MouseEvent)
 {
 	STableRow< TSharedPtr<FHierarchyModel> >::OnMouseEnter(MyGeometry, MouseEvent);
@@ -909,6 +927,20 @@ void SHierarchyViewItem::OnNameTextCommited(const FText& InText, ETextCommit::Ty
 	Model->OnNameTextCommited(InText, CommitInfo);
 }
 
+bool SHierarchyViewItem::CanRename() const
+{
+	return Model->CanRename();
+}
+
+void SHierarchyViewItem::BeginRename()
+{
+	TSharedPtr<SInlineEditableTextBlock> SafeEditBox = EditBox.Pin();
+	if ( SafeEditBox.IsValid() )
+	{
+		SafeEditBox->EnterEditingMode();
+	}
+}
+
 FSlateFontInfo SHierarchyViewItem::GetItemFont() const
 {
 	return Model->GetFont();
diff --git a/Engine/Source/Editor/UMGEditor/Private/Hierarchy/SHierarchyViewItem.h b/Engine/Source/Editor/UMGEditor/Private/Hierarchy/SHierarchyViewItem.h
--- a/Engine/Source/Editor/UMGEditor/Private/Hierarchy/SHierarchyViewItem.h
+++ b/Engine/Source/Editor/UMGEditor/Private/Hierarchy/SHierarchyViewItem.h
@@ -62,6 +62,12 @@ public:
 	virtual bool IsExpanded() const { return true; }
 	virtual void SetExpanded(bool bIsExpanded) { }
 
+	virtual bool CanRename() const { return false; }
+	virtual void BeginRename() { }
+
+public:
+	FSimpleDelegate RenameEvent;
+
 protected:
 	virtual void GetChildren(TArray< TSharedPtr<FHierarchyModel> >& Children) = 0;
 	virtual void UpdateSelection() = 0;
@@ -104,6 +110,8 @@ protected:
 private:
 
 	TWeakPtr<FWidgetBlueprintEditor> BlueprintEditor;
+
+	FText RootText;
 };
 
 class FNamedSlotModel : public FHierarchyModel
@@ -210,6 +218,9 @@ public:
 		Item.GetTemplate()->bExpandedInDesigner = bIsExpanded;
 	}
 
+	virtual bool CanRename() const override;
+	virtual void BeginRename() override;
+
 protected:
 	virtual void GetChildren(TArray< TSharedPtr<FHierarchyModel> >& Children) override;
 	virtual void UpdateSelection() override;
@@ -234,6 +245,7 @@ public:
 	SLATE_END_ARGS()
 
 	void Construct(const FArguments& InArgs, const TSharedRef< STableViewBase >& InOwnerTableView, TSharedPtr<FHierarchyModel> InModel);
+	virtual ~SHierarchyViewItem();
 
 	// Begin SWidget
 	virtual bool IsHovered() const override;
@@ -241,7 +253,7 @@ public:
 	virtual void OnMouseLeave(const FPointerEvent& MouseEvent) override;
 	// End SWidget
 
-private:
+public:
 	FReply HandleDragDetected(const FGeometry& MyGeometry, const FPointerEvent& MouseEvent);
 	void HandleDragEnter(const FDragDropEvent& DragDropEvent);
 	void HandleDragLeave(const FDragDropEvent& DragDropEvent);
@@ -255,6 +267,11 @@ private:
 	/** Called when text is committed on the node */
 	void OnNameTextCommited(const FText& InText, ETextCommit::Type CommitInfo);
 
+	bool CanRename() const;
+	void BeginRename();
+
+private:
+
 	/** Gets the font to use for the text item, bold for customized named items */
 	FSlateFontInfo GetItemFont() const;
 
@@ -267,6 +284,11 @@ private:
 	/** Returns a brush representing the visibility item of the widget's visibility button */
 	FText GetVisibilityBrushForWidget() const;
 
+private:
+
+	/** Edit box for the name. */
+	TWeakPtr<SInlineEditableTextBlock> EditBox;
+
 	/* The mode that this tree item represents */
 	TSharedPtr<FHierarchyModel> Model;
 };
