Fixed build error in SHeierarchyView

diff --git a/Engine/Source/Editor/UMGEditor/Private/Hierarchy/SHierarchyView.cpp b/Engine/Source/Editor/UMGEditor/Private/Hierarchy/SHierarchyView.cpp
--- a/Engine/Source/Editor/UMGEditor/Private/Hierarchy/SHierarchyView.cpp
+++ b/Engine/Source/Editor/UMGEditor/Private/Hierarchy/SHierarchyView.cpp
@@ -19,6 +19,7 @@
 #include "WidgetBlueprintEditorUtils.h"
 
 #include "Components/PanelWidget.h"
+#include "GenericCommands.h"
 
 #define LOCTEXT_NAMESPACE "UMG"
 
