UMG - Single point delegates that you want to bind now can be specified using meta=( IsBindableEvent="True" ), instead of depending on ending with Event.  The Event will eventually be phased out as a flag to detect these.

diff --git a/Engine/Source/Editor/UMGEditor/Private/Customizations/DetailCustomizations.cpp b/Engine/Source/Editor/UMGEditor/Private/Customizations/DetailCustomizations.cpp
--- a/Engine/Source/Editor/UMGEditor/Private/Customizations/DetailCustomizations.cpp
+++ b/Engine/Source/Editor/UMGEditor/Private/Customizations/DetailCustomizations.cpp
@@ -21,7 +21,6 @@
 
 #define LOCTEXT_NAMESPACE "UMG"
 
-
 class SGraphSchemaActionButton : public SCompoundWidget, public FGCObject
 {
 public:
@@ -88,11 +87,6 @@ private:
 
 void FBlueprintWidgetCustomization::CreateEventCustomization( IDetailLayoutBuilder& DetailLayout, UDelegateProperty* Property, UWidget* Widget )
 {
-	if ( !Property->GetName().EndsWith(TEXT("Event")) )
-	{
-		return;
-	}
-
 	TSharedRef<IPropertyHandle> DelegatePropertyHandle = DetailLayout.GetProperty(Property->GetFName(), CastChecked<UClass>(Property->GetOuter()));
 
 	const bool bHasValidHandle = DelegatePropertyHandle->IsValidHandle();
@@ -289,6 +283,8 @@ void FBlueprintWidgetCustomization::CustomizeDetails( IDetailLayoutBuilder& Deta
 
 void FBlueprintWidgetCustomization::PerformBindingCustomization(IDetailLayoutBuilder& DetailLayout)
 {
+	static const FName IsBindableEventName(TEXT("IsBindableEvent"));
+
 	TArray< TWeakObjectPtr<UObject> > OutObjects;
 	DetailLayout.GetObjectsBeingCustomized(OutObjects);
 
@@ -303,7 +299,8 @@ void FBlueprintWidgetCustomization::PerformBindingCustomization(IDetailLayoutBui
 
 			if ( UDelegateProperty* DelegateProperty = Cast<UDelegateProperty>(*PropertyIt) )
 			{
-				if ( DelegateProperty->GetName().EndsWith(TEXT("Event")) )
+				//TODO Remove the code to use ones that end with "Event".  Prefer metadata flag.
+				if ( DelegateProperty->GetBoolMetaData(IsBindableEventName) || DelegateProperty->GetName().EndsWith(TEXT("Event")) )
 				{
 					CreateEventCustomization(DetailLayout, DelegateProperty, Widget);
 				}
