UMG - Sequencer - Make animation binding names less confusing.
+ For slots display the name of the widget in the slot and the panels name instead of the slot name.
+ Update the display name for the binding when the object's name changes.

diff --git a/Engine/Source/Editor/Sequencer/Private/SAnimationOutlinerView.cpp b/Engine/Source/Editor/Sequencer/Private/SAnimationOutlinerView.cpp
--- a/Engine/Source/Editor/Sequencer/Private/SAnimationOutlinerView.cpp
+++ b/Engine/Source/Editor/Sequencer/Private/SAnimationOutlinerView.cpp
@@ -28,7 +28,7 @@ void SAnimationOutlinerTreeNode::Construct( const FArguments& InArgs, TSharedRef
 
 	TSharedRef<SWidget> TextWidget = 
 		SNew( STextBlock )
-		.Text( Node->GetDisplayName() )
+		.Text(this, &SAnimationOutlinerTreeNode::GetDisplayName )
 		.Font( NodeFont );
 
 	TSharedRef<SWidget>	FinalWidget = 
@@ -155,6 +155,11 @@ EVisibility SAnimationOutlinerTreeNode::GetExpanderVisibility() const
 	return DisplayNode->GetNumChildren() > 0 ? EVisibility::Visible : EVisibility::Hidden;
 }
 
+FText SAnimationOutlinerTreeNode::GetDisplayName() const
+{
+	return DisplayNode->GetDisplayName();
+}
+
 void SAnimationOutlinerView::Construct( const FArguments& InArgs, TSharedRef<FSequencerDisplayNode> InRootNode, TSharedRef<FSequencer> InSequencer )
 {
 	SAnimationOutlinerViewBase::Construct( SAnimationOutlinerViewBase::FArguments(), InRootNode );
diff --git a/Engine/Source/Editor/Sequencer/Private/SAnimationOutlinerView.h b/Engine/Source/Editor/Sequencer/Private/SAnimationOutlinerView.h
--- a/Engine/Source/Editor/Sequencer/Private/SAnimationOutlinerView.h
+++ b/Engine/Source/Editor/Sequencer/Private/SAnimationOutlinerView.h
@@ -52,6 +52,12 @@ private:
 	 * @return The expander visibility of this node                                                              
 	 */
 	EVisibility GetExpanderVisibility() const;
+
+	/**
+	 * @return The display name for this node.
+	 */
+	FText GetDisplayName() const;
+
 private:
 	/** Layout node the widget is visualizing */
 	TSharedPtr<FSequencerDisplayNode> DisplayNode;
diff --git a/Engine/Source/Editor/Sequencer/Private/Sequencer.cpp b/Engine/Source/Editor/Sequencer/Private/Sequencer.cpp
--- a/Engine/Source/Editor/Sequencer/Private/Sequencer.cpp
+++ b/Engine/Source/Editor/Sequencer/Private/Sequencer.cpp
@@ -1296,6 +1296,11 @@ void FSequencer::KeyProperty(const TArray<UObject*>& ObjectsToKey, const class I
 	ObjectChangeListener->KeyProperty( ObjectsToKey, PropertyHandle );
 }
 
+TSharedRef<ISequencerObjectBindingManager> FSequencer::GetObjectBindingManager() const
+{
+	return ObjectBindingManager.ToSharedRef();
+}
+
 TArray< TWeakObjectPtr<UMovieSceneSection> > FSequencer::GetFilteringShotSections() const
 {
 	return FilteringShots;
diff --git a/Engine/Source/Editor/Sequencer/Private/Sequencer.h b/Engine/Source/Editor/Sequencer/Private/Sequencer.h
--- a/Engine/Source/Editor/Sequencer/Private/Sequencer.h
+++ b/Engine/Source/Editor/Sequencer/Private/Sequencer.h
@@ -90,6 +90,7 @@ public:
 	virtual void FilterToSelectedShotSections(bool bZoomToShotBounds = true) override;
 	virtual bool CanKeyProperty(const UClass& ObjectClass, const class IPropertyHandle& PropertyHandle) const override;
 	virtual void KeyProperty(const TArray<UObject*>& ObjectsToKey, const class IPropertyHandle& PropertyHandle) override;
+	virtual TSharedRef<ISequencerObjectBindingManager> GetObjectBindingManager() const override;
 
 	bool IsPerspectiveViewportPosessionEnabled() const { return bPerspectiveViewportPossessionEnabled; }
 
diff --git a/Engine/Source/Editor/Sequencer/Private/SequencerActorBindingManager.h b/Engine/Source/Editor/Sequencer/Private/SequencerActorBindingManager.h
--- a/Engine/Source/Editor/Sequencer/Private/SequencerActorBindingManager.h
+++ b/Engine/Source/Editor/Sequencer/Private/SequencerActorBindingManager.h
@@ -86,6 +86,7 @@ public:
 	virtual void BindPossessableObject( const FGuid& PossessableGuid, UObject& PossessedObject ) override;
 	virtual void UnbindPossessableObjects( const FGuid& PossessableGuid ) override;
 	virtual void GetRuntimeObjects( const TSharedRef<FMovieSceneInstance>& MovieSceneInstance, const FGuid& ObjectGuid, TArray<UObject*>& OutRuntimeObjects ) const override;
+	virtual bool TryGetObjectBindingDisplayName(const FGuid& ObjectGuid, FText& DisplayName) const override { return false; }
 
 protected:
 	
diff --git a/Engine/Source/Editor/Sequencer/Private/SequencerDisplayNode.cpp b/Engine/Source/Editor/Sequencer/Private/SequencerDisplayNode.cpp
--- a/Engine/Source/Editor/Sequencer/Private/SequencerDisplayNode.cpp
+++ b/Engine/Source/Editor/Sequencer/Private/SequencerDisplayNode.cpp
@@ -12,6 +12,7 @@
 #include "MovieSceneTrack.h"
 #include "CommonMovieSceneTools.h"
 #include "IKeyArea.h"
+#include "ISequencerObjectBindingManager.h"
 
 
 class SSequencerObjectTrack : public SLeafWidget
@@ -445,6 +446,13 @@ void FTrackNode::FixRowIndices()
 	}
 }
 
+FText FObjectBindingNode::GetDisplayName() const
+{
+	FText DisplayName;
+	return GetSequencer().GetObjectBindingManager()->TryGetObjectBindingDisplayName(ObjectBinding, DisplayName) ?
+		DisplayName : DefaultDisplayName;
+}
+
 float FObjectBindingNode::GetNodeHeight() const
 { 
 	return SequencerLayoutConstants::ObjectNodeHeight;
diff --git a/Engine/Source/Editor/Sequencer/Private/SequencerDisplayNode.h b/Engine/Source/Editor/Sequencer/Private/SequencerDisplayNode.h
--- a/Engine/Source/Editor/Sequencer/Private/SequencerDisplayNode.h
+++ b/Engine/Source/Editor/Sequencer/Private/SequencerDisplayNode.h
@@ -371,12 +371,12 @@ public:
 	FObjectBindingNode( FName NodeName, const FString& InObjectName, const FGuid& InObjectBinding, TSharedPtr<FSequencerDisplayNode> InParentNode, FSequencerNodeTree& InParentTree )
 		: FSequencerDisplayNode( NodeName, InParentNode, InParentTree )
 		, ObjectBinding( InObjectBinding )
-		, DisplayName( InObjectName )
+		, DefaultDisplayName( FText::FromString(InObjectName) )
 	{}
 
 	/** FSequencerDisplayNodeInterface */
 	virtual ESequencerNode::Type GetType() const override { return ESequencerNode::Object; }
-	virtual FText GetDisplayName() const override { return FText::FromString(DisplayName); }
+	virtual FText GetDisplayName() const override;
 	virtual float GetNodeHeight() const override;
 	virtual bool GetShotFilteredVisibilityToCache() const override;
 	
@@ -389,8 +389,8 @@ public:
 private:
 	/** The binding to live objects */
 	FGuid ObjectBinding;
-	/** Display name of the object */
-	FString DisplayName;
+	/** The default display name of the object which is used if the binding manager doesn't provide one. */
+	FText DefaultDisplayName;
 };
 
 /**
diff --git a/Engine/Source/Editor/Sequencer/Public/ISequencer.h b/Engine/Source/Editor/Sequencer/Public/ISequencer.h
--- a/Engine/Source/Editor/Sequencer/Public/ISequencer.h
+++ b/Engine/Source/Editor/Sequencer/Public/ISequencer.h
@@ -125,5 +125,5 @@ public:
 
 	virtual void UpdateRuntimeInstances() = 0;
 
-
+	virtual TSharedRef<ISequencerObjectBindingManager> GetObjectBindingManager() const = 0;
 };
diff --git a/Engine/Source/Editor/Sequencer/Public/ISequencerObjectBindingManager.h b/Engine/Source/Editor/Sequencer/Public/ISequencerObjectBindingManager.h
--- a/Engine/Source/Editor/Sequencer/Public/ISequencerObjectBindingManager.h
+++ b/Engine/Source/Editor/Sequencer/Public/ISequencerObjectBindingManager.h
@@ -52,4 +52,14 @@ public:
 	virtual void UnbindPossessableObjects( const FGuid& PossessableGuid ) = 0;
 	
 	virtual void GetRuntimeObjects( const TSharedRef<FMovieSceneInstance>& MovieSceneInstance, const FGuid& ObjectGuid, TArray<UObject*>& OutRuntimeObjects ) const = 0;
+
+	/**
+	 * 
+	 * Tries to get a display name for the binding represented by the guid.
+	 *
+	 * @param ObjectGuid	the guid for the object binding.
+	 * @param DisplayName	the display name for the object binding.
+	 * @returns true if DisplayName has been set to a valid display name, otherwise false.
+	 */
+	virtual bool TryGetObjectBindingDisplayName( const FGuid& ObjectGuid, FText& DisplayName ) const = 0;
 };
diff --git a/Engine/Source/Editor/UMGEditor/Private/Animation/UMGSequencerObjectBindingManager.cpp b/Engine/Source/Editor/UMGEditor/Private/Animation/UMGSequencerObjectBindingManager.cpp
--- a/Engine/Source/Editor/UMGEditor/Private/Animation/UMGSequencerObjectBindingManager.cpp
+++ b/Engine/Source/Editor/UMGEditor/Private/Animation/UMGSequencerObjectBindingManager.cpp
@@ -12,6 +12,8 @@
 #include "Animation/WidgetAnimation.h"
 #include "WidgetBlueprint.h"
 
+#define LOCTEXT_NAMESPACE "UMGSequencerObjectBindingManager"
+
 FUMGSequencerObjectBindingManager::FUMGSequencerObjectBindingManager( FWidgetBlueprintEditor& InWidgetBlueprintEditor, UWidgetAnimation& InAnimation )
 	: WidgetAnimation( &InAnimation )
 	, WidgetBlueprintEditor( InWidgetBlueprintEditor )
@@ -48,6 +50,7 @@ bool FUMGSequencerObjectBindingManager::CanPossessObject( UObject& Object ) cons
 void FUMGSequencerObjectBindingManager::BindPossessableObject( const FGuid& PossessableGuid, UObject& PossessedObject )
 {
 	PreviewObjectToGuidMap.Add( &PossessedObject, PossessableGuid );
+	GuidToPreviewObjectsMap.Add(PossessableGuid, &PossessedObject);
 
 	FWidgetAnimationBinding NewBinding;
 	NewBinding.AnimationGuid = PossessableGuid;
@@ -73,13 +76,13 @@ void FUMGSequencerObjectBindingManager::BindPossessableObject( const FGuid& Poss
 
 void FUMGSequencerObjectBindingManager::UnbindPossessableObjects( const FGuid& PossessableGuid )
 {
-	for( auto It = PreviewObjectToGuidMap.CreateIterator(); It; ++It )
+	TArray<TWeakObjectPtr<UObject>> PreviewObjects;
+	GuidToPreviewObjectsMap.MultiFind(PossessableGuid, PreviewObjects);
+	for (TWeakObjectPtr<UObject>& PreviewObject : PreviewObjects)
 	{
-		if( It.Value() == PossessableGuid )
-		{
-			It.RemoveCurrent();
-		}
+		PreviewObjectToGuidMap.Remove(PreviewObject);
 	}
+	GuidToPreviewObjectsMap.Remove(PossessableGuid);
 
 	UWidgetBlueprint* WidgetBlueprint = WidgetBlueprintEditor.GetWidgetBlueprintObj();
 
@@ -89,14 +92,40 @@ void FUMGSequencerObjectBindingManager::UnbindPossessableObjects( const FGuid& P
 
 void FUMGSequencerObjectBindingManager::GetRuntimeObjects( const TSharedRef<FMovieSceneInstance>& MovieSceneInstance, const FGuid& ObjectGuid, TArray<UObject*>& OutRuntimeObjects ) const
 {
-	for( auto It = PreviewObjectToGuidMap.CreateConstIterator(); It; ++It )
+	TArray<TWeakObjectPtr<UObject>> PreviewObjects;
+	GuidToPreviewObjectsMap.MultiFind(ObjectGuid, PreviewObjects);
+	for (TWeakObjectPtr<UObject>& PreviewObject : PreviewObjects)
+	{
+		OutRuntimeObjects.Add(PreviewObject.Get());
+	}
+}
+
+bool FUMGSequencerObjectBindingManager::TryGetObjectBindingDisplayName(const FGuid& ObjectGuid, FText& DisplayName) const
+{
+	// TODO: This gets called every frame for every bound object and could be a potential performance issue for a really complicated animation.
+	TArray<TWeakObjectPtr<UObject>> BindingObjects;
+	GuidToPreviewObjectsMap.MultiFind(ObjectGuid, BindingObjects);
+	if (BindingObjects.Num() == 0)
 	{
-		UObject* Object = It.Key().Get();
-		if( Object && It.Value() == ObjectGuid )
+		DisplayName = LOCTEXT("NoBoundObjects", "Invalid bound object");
+	}
+	else if (BindingObjects.Num() == 1)
+	{
+		UPanelSlot* SlotObject = Cast<UPanelSlot>(BindingObjects[0].Get());
+		if (SlotObject != nullptr)
+		{
+			DisplayName = FText::Format(LOCTEXT("SlotObject", "{0} ({1} Slot)"), FText::FromString(SlotObject->Content->GetName()), FText::FromString(SlotObject->Parent->GetName()));
+		}
+		else
 		{
-			OutRuntimeObjects.Add( Object );
+			DisplayName = FText::FromString(BindingObjects[0].Get()->GetName());
 		}
 	}
+	else
+	{
+		DisplayName = LOCTEXT("Multiple bound objects", "Multilple bound objects");
+	}
+	return true;
 }
 
 bool FUMGSequencerObjectBindingManager::HasValidWidgetAnimation() const
@@ -120,6 +149,7 @@ void FUMGSequencerObjectBindingManager::InitPreviewObjects()
 			if( FoundObject )
 			{
 				PreviewObjectToGuidMap.Add(FoundObject, Binding.AnimationGuid);
+				GuidToPreviewObjectsMap.Add(Binding.AnimationGuid, FoundObject);
 			}
 		}
 	}
@@ -128,8 +158,11 @@ void FUMGSequencerObjectBindingManager::InitPreviewObjects()
 void FUMGSequencerObjectBindingManager::OnWidgetPreviewUpdated()
 {
 	PreviewObjectToGuidMap.Empty();
+	GuidToPreviewObjectsMap.Empty();
 
 	InitPreviewObjects();
 
 	WidgetBlueprintEditor.GetSequencer()->UpdateRuntimeInstances();
-}
\ No newline at end of file
+}
+
+#undef LOCTEXT_NAMESPACE
\ No newline at end of file
diff --git a/Engine/Source/Editor/UMGEditor/Private/Animation/UMGSequencerObjectBindingManager.h b/Engine/Source/Editor/UMGEditor/Private/Animation/UMGSequencerObjectBindingManager.h
--- a/Engine/Source/Editor/UMGEditor/Private/Animation/UMGSequencerObjectBindingManager.h
+++ b/Engine/Source/Editor/UMGEditor/Private/Animation/UMGSequencerObjectBindingManager.h
@@ -25,6 +25,7 @@ public:
 	virtual void BindPossessableObject( const FGuid& PossessableGuid, UObject& PossessedObject ) override;
 	virtual void UnbindPossessableObjects( const FGuid& PossessableGuid ) override;
 	virtual void GetRuntimeObjects( const TSharedRef<FMovieSceneInstance>& MovieSceneInstance, const FGuid& ObjectGuid, TArray<UObject*>& OutRuntimeObjects ) const override;
+	virtual bool TryGetObjectBindingDisplayName(const FGuid& ObjectGuid, FText& DisplayName) const override;
 	
 	/** @return true if the current animation is valid */
 	bool HasValidWidgetAnimation() const;
@@ -39,7 +40,9 @@ private:
 	void OnWidgetPreviewUpdated();
 private:
 	/** Mapping of preview objects to sequencer guids */
-	TMap< TWeakObjectPtr<UObject>, FGuid> PreviewObjectToGuidMap;
+	TMap<TWeakObjectPtr<UObject>, FGuid> PreviewObjectToGuidMap;
+	TMultiMap<FGuid, TWeakObjectPtr<UObject>> GuidToPreviewObjectsMap;
+
 
 	UWidgetAnimation* WidgetAnimation;
 
