UMG - Sequencer - Fix a crash during drag drop.

diff --git a/Engine/Source/Editor/UMGEditor/Private/Animation/UMGSequencerObjectBindingManager.cpp b/Engine/Source/Editor/UMGEditor/Private/Animation/UMGSequencerObjectBindingManager.cpp
--- a/Engine/Source/Editor/UMGEditor/Private/Animation/UMGSequencerObjectBindingManager.cpp
+++ b/Engine/Source/Editor/UMGEditor/Private/Animation/UMGSequencerObjectBindingManager.cpp
@@ -114,7 +114,13 @@ bool FUMGSequencerObjectBindingManager::TryGetObjectBindingDisplayName(const FGu
 		UPanelSlot* SlotObject = Cast<UPanelSlot>(BindingObjects[0].Get());
 		if (SlotObject != nullptr)
 		{
-			DisplayName = FText::Format(LOCTEXT("SlotObject", "{0} ({1} Slot)"), FText::FromString(SlotObject->Content->GetName()), FText::FromString(SlotObject->Parent->GetName()));
+			FText PanelName = SlotObject->Parent != nullptr 
+				? FText::FromString(SlotObject->Parent->GetName())
+				: LOCTEXT("InvalidPanel", "Invalid Panel");
+			FText ContentName = SlotObject->Content != nullptr
+				? FText::FromString(SlotObject->Content->GetName())
+				: LOCTEXT("EmptyPanel", "Empty");
+			DisplayName = FText::Format(LOCTEXT("SlotObject", "{0} ({1} Slot)"), ContentName, PanelName);
 		}
 		else
 		{
