UMG - Sequencer - Make the "No Animation Selected" text block handling more consistent.

diff --git a/Engine/Source/Editor/UMGEditor/Private/WidgetBlueprintEditor.cpp b/Engine/Source/Editor/UMGEditor/Private/WidgetBlueprintEditor.cpp
--- a/Engine/Source/Editor/UMGEditor/Private/WidgetBlueprintEditor.cpp
+++ b/Engine/Source/Editor/UMGEditor/Private/WidgetBlueprintEditor.cpp
@@ -459,11 +459,24 @@ TSharedRef<SWidget> FWidgetBlueprintEditor::CreateSequencerWidget()
 		.AddMetaData<FTagMetaData>(FTagMetaData(TEXT("Sequencer")));
 	SequencerOverlay = SequencerOverlayRef;
 
+	TSharedRef<STextBlock> NoAnimationTextBlockRef = 
+		SNew(STextBlock)
+		.TextStyle(FEditorStyle::Get(), "UMGEditor.NoAnimationFont")
+		.Text(LOCTEXT("NoAnimationSelected", "No Animation Selected"));
+	NoAnimationTextBlock = NoAnimationTextBlockRef;
+
 	SequencerOverlayRef->AddSlot(0)
 	[
 		GetSequencer()->GetSequencerWidget()
 	];
 
+	SequencerOverlayRef->AddSlot(1)
+		.HAlign(HAlign_Center)
+		.VAlign(VAlign_Center)
+	[
+		NoAnimationTextBlockRef
+	];
+
 	return SequencerOverlayRef;
 }
 
@@ -549,30 +562,24 @@ void FWidgetBlueprintEditor::ChangeViewedAnimation( UWidgetAnimation& InAnimatio
 	SequencerObjectBindingManager->InitPreviewObjects();
 
 	TSharedPtr<SOverlay> SequencerOverlayPin = SequencerOverlay.Pin();
-	if( &InAnimationToView == UWidgetAnimation::GetNullAnimation() && SequencerOverlayPin.IsValid() )
-	{
-		Sequencer->GetSequencerWidget()->SetEnabled(false);
-		// Disable sequencer from interaction
-		SequencerOverlayPin->AddSlot(1)
-		.HAlign(HAlign_Center)
-		.VAlign(VAlign_Center)
-		[
-			SNew( STextBlock )
-			.TextStyle( FEditorStyle::Get(), "UMGEditor.NoAnimationFont" )
-			.Text( LOCTEXT("NoAnimationSelected","No Animation Selected") )
-		];
-
-		SequencerOverlayPin->SetVisibility( EVisibility::HitTestInvisible );
-	}
-	else if( SequencerOverlayPin.IsValid() && SequencerOverlayPin->GetNumWidgets() > 1 )
+	if (SequencerOverlayPin.IsValid())
 	{
-		Sequencer->GetSequencerWidget()->SetEnabled(true);
-
-		SequencerOverlayPin->RemoveSlot(1);
-		// Allow sequencer to be interacted with
-		SequencerOverlayPin->SetVisibility( EVisibility::SelfHitTestInvisible );
+		TSharedPtr<STextBlock> NoAnimationTextBlockPin = NoAnimationTextBlock.Pin();
+		if( &InAnimationToView == UWidgetAnimation::GetNullAnimation())
+		{
+			// Disable sequencer from interaction
+			Sequencer->GetSequencerWidget()->SetEnabled(false);
+			NoAnimationTextBlockPin->SetVisibility(EVisibility::Visible);
+			SequencerOverlayPin->SetVisibility( EVisibility::HitTestInvisible );
+		}
+		else
+		{
+			// Allow sequencer to be interacted with
+			Sequencer->GetSequencerWidget()->SetEnabled(true);
+			NoAnimationTextBlockPin->SetVisibility(EVisibility::Collapsed);
+			SequencerOverlayPin->SetVisibility( EVisibility::SelfHitTestInvisible );
+		}
 	}
-
 	InvalidatePreview();
 }
 
diff --git a/Engine/Source/Editor/UMGEditor/Private/WidgetBlueprintEditor.h b/Engine/Source/Editor/UMGEditor/Private/WidgetBlueprintEditor.h
--- a/Engine/Source/Editor/UMGEditor/Private/WidgetBlueprintEditor.h
+++ b/Engine/Source/Editor/UMGEditor/Private/WidgetBlueprintEditor.h
@@ -182,6 +182,9 @@ private:
 	/** Overlay used to display UI on top of sequencer */
 	TWeakPtr<SOverlay> SequencerOverlay;
 
+	/** A text block which is displayed in the overlay when no animation is selected. */
+	TWeakPtr<STextBlock> NoAnimationTextBlock;
+
 	/** Manager for handling bindings to sequence animations */
 	TSharedPtr<class FUMGSequencerObjectBindingManager> SequencerObjectBindingManager;
 
