UMG - Sequencer - Fix cases where changing widgets position and size in a canvas through the designer in auto-key mode didn't generate keyframes.
+ Change the UMGObjectBindingManager to track slots the same way the FWidgetAnimationBindings do where the Guid is associated with the slot's content.
+ Fix the WidgetBlueprintEditorUtils.ImportPropertiesFromText to call pre-change callbacks so that sequencer can handle them.
+ Update the STransformHandle so that it sets slot offsets through a utility function which calls the correct change callbacks so that sequencer can handle them.

diff --git a/Engine/Source/Editor/Sequencer/Private/Sequencer.cpp b/Engine/Source/Editor/Sequencer/Private/Sequencer.cpp
--- a/Engine/Source/Editor/Sequencer/Private/Sequencer.cpp
+++ b/Engine/Source/Editor/Sequencer/Private/Sequencer.cpp
@@ -530,13 +530,16 @@ FGuid FSequencer::GetHandleToObject( UObject* Object )
 
 	FGuid ObjectGuid = ObjectBindingManager->FindGuidForObject( *FocusedMovieScene, *Object );
 
-	// Make sure that the possessable is still valid, if it's not remove the binding so new one 
-	// can be created.  This can happen due to undo.
-	FMovieScenePossessable* Possessable = FocusedMovieScene->FindPossessable(ObjectGuid);
-	if ( Possessable == nullptr )
+	if (ObjectGuid.IsValid())
 	{
-		ObjectBindingManager->UnbindPossessableObjects(ObjectGuid);
-		ObjectGuid.Invalidate();
+		// Make sure that the possessable is still valid, if it's not remove the binding so new one 
+		// can be created.  This can happen due to undo.
+		FMovieScenePossessable* Possessable = FocusedMovieScene->FindPossessable(ObjectGuid);
+		if (Possessable == nullptr )
+		{
+			ObjectBindingManager->UnbindPossessableObjects(ObjectGuid);
+			ObjectGuid.Invalidate();
+		}
 	}
 	
 	bool bPossessableAdded = false;
diff --git a/Engine/Source/Editor/UMGEditor/Private/Animation/UMGSequencerObjectBindingManager.cpp b/Engine/Source/Editor/UMGEditor/Private/Animation/UMGSequencerObjectBindingManager.cpp
--- a/Engine/Source/Editor/UMGEditor/Private/Animation/UMGSequencerObjectBindingManager.cpp
+++ b/Engine/Source/Editor/UMGEditor/Private/Animation/UMGSequencerObjectBindingManager.cpp
@@ -34,11 +34,24 @@ void FUMGSequencerObjectBindingManager::AddReferencedObjects(FReferenceCollector
 
 FGuid FUMGSequencerObjectBindingManager::FindGuidForObject( const UMovieScene& MovieScene, UObject& Object ) const
 {
+	// Slot guids are tracked by their content.
+	UPanelSlot* Slot = Cast<UPanelSlot>(&Object);
+	if (Slot != nullptr)
+	{
+		return SlotContentPreviewObjectToGuidMap.FindRef(Slot->Content);
+	}
 	return PreviewObjectToGuidMap.FindRef( &Object );
 }
 
 bool FUMGSequencerObjectBindingManager::CanPossessObject( UObject& Object ) const
 {
+	// Can't possess empty slots.
+	UPanelSlot* Slot = Cast<UPanelSlot>(&Object);
+	if (Slot != nullptr && Slot->Content == nullptr)
+	{
+		return false;
+	}
+
 	UWidgetBlueprint* WidgetBlueprint = WidgetBlueprintEditor.GetWidgetBlueprintObj();
 
 	// Only preview widgets in this blueprint can be possessed
@@ -49,29 +62,31 @@ bool FUMGSequencerObjectBindingManager::CanPossessObject( UObject& Object ) cons
 
 void FUMGSequencerObjectBindingManager::BindPossessableObject( const FGuid& PossessableGuid, UObject& PossessedObject )
 {
-	PreviewObjectToGuidMap.Add( &PossessedObject, PossessableGuid );
-	GuidToPreviewObjectsMap.Add(PossessableGuid, &PossessedObject);
-
-	FWidgetAnimationBinding NewBinding;
-	NewBinding.AnimationGuid = PossessableGuid;
-
-	UPanelSlot* Slot = Cast<UPanelSlot>( &PossessedObject );
-	if( Slot && Slot->Content )
+	UPanelSlot* PossessedSlot = Cast<UPanelSlot>(&PossessedObject);
+	if ( PossessedSlot != nullptr && PossessedSlot->Content != nullptr )
 	{
-		// Save the name of the widget containing the slots.  This is the object to look up that contains the slot itself(the thing we are animating)
-		NewBinding.SlotWidgetName = Slot->GetFName();
-		NewBinding.WidgetName = Slot->Content->GetFName();
+		SlotContentPreviewObjectToGuidMap.Add(PossessedSlot->Content, PossessableGuid);
+		GuidToSlotContentPreviewObjectsMap.Add(PossessableGuid, PossessedSlot->Content);
 
+		// Save the name of the widget containing the slots.  This is the object to look up that contains the slot itself(the thing we are animating)
+		FWidgetAnimationBinding NewBinding;
+		NewBinding.AnimationGuid = PossessableGuid;
+		NewBinding.SlotWidgetName = PossessedSlot->GetFName();
+		NewBinding.WidgetName = PossessedSlot->Content->GetFName();
 
 		WidgetAnimation->AnimationBindings.Add(NewBinding);
 	}
-	else if( !Slot )
+	else if( PossessedSlot == nullptr )
 	{
+		PreviewObjectToGuidMap.Add(&PossessedObject, PossessableGuid);
+		GuidToPreviewObjectsMap.Add(PossessableGuid, &PossessedObject);
+
+		FWidgetAnimationBinding NewBinding;
+		NewBinding.AnimationGuid = PossessableGuid;
 		NewBinding.WidgetName = PossessedObject.GetFName();
 
 		WidgetAnimation->AnimationBindings.Add(NewBinding);
 	}
-
 }
 
 void FUMGSequencerObjectBindingManager::UnbindPossessableObjects( const FGuid& PossessableGuid )
@@ -84,6 +99,14 @@ void FUMGSequencerObjectBindingManager::UnbindPossessableObjects( const FGuid& P
 	}
 	GuidToPreviewObjectsMap.Remove(PossessableGuid);
 
+	TArray<TWeakObjectPtr<UObject>> SlotContentPreviewObjects;
+	GuidToSlotContentPreviewObjectsMap.MultiFind(PossessableGuid, SlotContentPreviewObjects);
+	for (TWeakObjectPtr<UObject>& SlotContentPreviewObject : SlotContentPreviewObjects)
+	{
+		SlotContentPreviewObjectToGuidMap.Remove(SlotContentPreviewObject);
+	}
+	GuidToSlotContentPreviewObjectsMap.Remove(PossessableGuid);
+
 	UWidgetBlueprint* WidgetBlueprint = WidgetBlueprintEditor.GetWidgetBlueprintObj();
 
 	WidgetAnimation->AnimationBindings.RemoveAll( [&]( const FWidgetAnimationBinding& Binding ) { return Binding.AnimationGuid == PossessableGuid; } );
@@ -98,6 +121,14 @@ void FUMGSequencerObjectBindingManager::GetRuntimeObjects( const TSharedRef<FMov
 	{
 		OutRuntimeObjects.Add(PreviewObject.Get());
 	}
+
+	TArray<TWeakObjectPtr<UObject>> SlotContentPreviewObjects;
+	GuidToSlotContentPreviewObjectsMap.MultiFind(ObjectGuid, SlotContentPreviewObjects);
+	for (TWeakObjectPtr<UObject>& SlotContentPreviewObject : SlotContentPreviewObjects)
+	{
+		UWidget* ContentWidget = Cast<UWidget>(SlotContentPreviewObject.Get());
+		OutRuntimeObjects.Add(ContentWidget->Slot);
+	}
 }
 
 bool FUMGSequencerObjectBindingManager::TryGetObjectBindingDisplayName(const FGuid& ObjectGuid, FText& DisplayName) const
@@ -105,31 +136,28 @@ bool FUMGSequencerObjectBindingManager::TryGetObjectBindingDisplayName(const FGu
 	// TODO: This gets called every frame for every bound object and could be a potential performance issue for a really complicated animation.
 	TArray<TWeakObjectPtr<UObject>> BindingObjects;
 	GuidToPreviewObjectsMap.MultiFind(ObjectGuid, BindingObjects);
-	if (BindingObjects.Num() == 0)
+	TArray<TWeakObjectPtr<UObject>> SlotContentBindingObjects;
+	GuidToSlotContentPreviewObjectsMap.MultiFind(ObjectGuid, SlotContentBindingObjects);
+	if (BindingObjects.Num() == 0 && SlotContentBindingObjects.Num() == 0)
+	{
+		DisplayName = LOCTEXT("NoBoundObjects", "No bound objects");
+	}
+	else if (BindingObjects.Num() + SlotContentBindingObjects.Num() > 1)
 	{
-		DisplayName = LOCTEXT("NoBoundObjects", "Invalid bound object");
+		DisplayName = LOCTEXT("Multiple bound objects", "Multilple bound objects");
 	}
 	else if (BindingObjects.Num() == 1)
 	{
-		UPanelSlot* SlotObject = Cast<UPanelSlot>(BindingObjects[0].Get());
-		if (SlotObject != nullptr)
-		{
-			FText PanelName = SlotObject->Parent != nullptr 
-				? FText::FromString(SlotObject->Parent->GetName())
-				: LOCTEXT("InvalidPanel", "Invalid Panel");
-			FText ContentName = SlotObject->Content != nullptr
-				? FText::FromString(SlotObject->Content->GetName())
-				: LOCTEXT("EmptyPanel", "Empty");
-			DisplayName = FText::Format(LOCTEXT("SlotObject", "{0} ({1} Slot)"), ContentName, PanelName);
-		}
-		else
-		{
-			DisplayName = FText::FromString(BindingObjects[0].Get()->GetName());
-		}
+		DisplayName = FText::FromString(BindingObjects[0].Get()->GetName());
 	}
-	else
+	else // SlotContentBindingObjects.Num() == 1
 	{
-		DisplayName = LOCTEXT("Multiple bound objects", "Multilple bound objects");
+		UWidget* SlotContent = Cast<UWidget>(SlotContentBindingObjects[0].Get());
+		FText PanelName = SlotContent->Slot != nullptr && SlotContent->Slot->Parent != nullptr
+			? FText::FromString(SlotContent->Slot->Parent->GetName())
+			: LOCTEXT("InvalidPanel", "Invalid Panel");
+		FText ContentName = FText::FromString(SlotContent->GetName());
+		DisplayName = FText::Format(LOCTEXT("SlotObject", "{0} ({1} Slot)"), ContentName, PanelName);
 	}
 	return true;
 }
@@ -152,10 +180,19 @@ void FUMGSequencerObjectBindingManager::InitPreviewObjects()
 		for(const FWidgetAnimationBinding& Binding : WidgetAnimation->AnimationBindings)
 		{
 			UObject* FoundObject = Binding.FindRuntimeObject( *WidgetTree );
-			if( FoundObject )
+			if ( FoundObject )
 			{
-				PreviewObjectToGuidMap.Add(FoundObject, Binding.AnimationGuid);
-				GuidToPreviewObjectsMap.Add(Binding.AnimationGuid, FoundObject);
+				UPanelSlot* FoundSlot = Cast<UPanelSlot>(FoundObject);
+				if (FoundSlot == nullptr)
+				{
+					PreviewObjectToGuidMap.Add(FoundObject, Binding.AnimationGuid);
+					GuidToPreviewObjectsMap.Add(Binding.AnimationGuid, FoundObject);
+				}
+				else
+				{
+					SlotContentPreviewObjectToGuidMap.Add(FoundSlot->Content, Binding.AnimationGuid);
+					GuidToSlotContentPreviewObjectsMap.Add(Binding.AnimationGuid, FoundSlot->Content);
+				}
 			}
 		}
 	}
@@ -165,6 +202,8 @@ void FUMGSequencerObjectBindingManager::OnWidgetPreviewUpdated()
 {
 	PreviewObjectToGuidMap.Empty();
 	GuidToPreviewObjectsMap.Empty();
+	SlotContentPreviewObjectToGuidMap.Empty();
+	GuidToSlotContentPreviewObjectsMap.Empty();
 
 	InitPreviewObjects();
 
diff --git a/Engine/Source/Editor/UMGEditor/Private/Animation/UMGSequencerObjectBindingManager.h b/Engine/Source/Editor/UMGEditor/Private/Animation/UMGSequencerObjectBindingManager.h
--- a/Engine/Source/Editor/UMGEditor/Private/Animation/UMGSequencerObjectBindingManager.h
+++ b/Engine/Source/Editor/UMGEditor/Private/Animation/UMGSequencerObjectBindingManager.h
@@ -41,8 +41,9 @@ private:
 private:
 	/** Mapping of preview objects to sequencer guids */
 	TMap<TWeakObjectPtr<UObject>, FGuid> PreviewObjectToGuidMap;
+	TMap<TWeakObjectPtr<UObject>, FGuid> SlotContentPreviewObjectToGuidMap;
 	TMultiMap<FGuid, TWeakObjectPtr<UObject>> GuidToPreviewObjectsMap;
-
+	TMultiMap<FGuid, TWeakObjectPtr<UObject>> GuidToSlotContentPreviewObjectsMap;
 
 	UWidgetAnimation* WidgetAnimation;
 
diff --git a/Engine/Source/Editor/UMGEditor/Private/Designer/STransformHandle.cpp b/Engine/Source/Editor/UMGEditor/Private/Designer/STransformHandle.cpp
--- a/Engine/Source/Editor/UMGEditor/Private/Designer/STransformHandle.cpp
+++ b/Engine/Source/Editor/UMGEditor/Private/Designer/STransformHandle.cpp
@@ -219,7 +219,10 @@ void STransformHandle::Resize(UCanvasPanelSlot* Slot, const FVector2D& Direction
 		}
 	}
 
-	Slot->SetOffsets(Offsets);
+	FAnchorData NewLayoutData = LayoutData;
+	NewLayoutData.Offsets = Offsets;
+	static const FName LayoutDataName(TEXT("LayoutData"));
+	FObjectEditorUtils::SetPropertyValue<UCanvasPanelSlot, FAnchorData>(Slot, LayoutDataName, NewLayoutData);
 }
 
 FCursorReply STransformHandle::OnCursorQuery(const FGeometry& MyGeometry, const FPointerEvent& MouseEvent) const
diff --git a/Engine/Source/Editor/UMGEditor/Private/WidgetBlueprintEditorUtils.cpp b/Engine/Source/Editor/UMGEditor/Private/WidgetBlueprintEditorUtils.cpp
--- a/Engine/Source/Editor/UMGEditor/Private/WidgetBlueprintEditorUtils.cpp
+++ b/Engine/Source/Editor/UMGEditor/Private/WidgetBlueprintEditorUtils.cpp
@@ -770,6 +770,10 @@ void FWidgetBlueprintEditorUtils::ImportPropertiesFromText(UObject* Object, cons
 		{
 			if ( UProperty* Property = FindField<UProperty>(Object->GetClass(), Entry.Key) )
 			{
+				FEditPropertyChain PropertyChain;
+				PropertyChain.AddHead(Property);
+				Object->PreEditChange(PropertyChain);
+
 				Property->ImportText(*Entry.Value, Property->ContainerPtrToValuePtr<uint8>(Object), 0, Object);
 
 				FPropertyChangedEvent ChangedEvent(Property);
