UMG - Pasting widgets at the last click location now when triggered with Ctrl+V.

diff --git a/Engine/Source/Editor/UMGEditor/Private/Designer/SDesignerView.cpp b/Engine/Source/Editor/UMGEditor/Private/Designer/SDesignerView.cpp
--- a/Engine/Source/Editor/UMGEditor/Private/Designer/SDesignerView.cpp
+++ b/Engine/Source/Editor/UMGEditor/Private/Designer/SDesignerView.cpp
@@ -183,6 +183,8 @@ void SDesignerView::Construct(const FArguments& InArgs, TSharedPtr<FWidgetBluepr
 
 	HoveredWidgetOutlineFade = FCurveSequence(0.0f, 0.15f);
 
+	SelectedWidgetContextMenuLocation = FVector2D(0, 0);
+
 	bMovingExistingWidget = false;
 
 	// TODO UMG - Register these with the module through some public interface to allow for new extensions to be registered.
@@ -1208,7 +1210,7 @@ FReply SDesignerView::OnMouseMove(const FGeometry& MyGeometry, const FPointerEve
 		return SurfaceHandled;
 	}
 
-	if ( MouseEvent.IsMouseButtonDown(EKeys::LeftMouseButton) )
+	if ( MouseEvent.IsMouseButtonDown(EKeys::LeftMouseButton) && HasMouseCapture() )
 	{
 		FWidgetReference SelectedWidget = GetSelectedWidget();
 
@@ -1283,7 +1285,7 @@ void SDesignerView::OnMouseLeave(const FPointerEvent& MouseEvent)
 
 FReply SDesignerView::OnKeyDown(const FGeometry& MyGeometry, const FKeyEvent& InKeyEvent)
 {
-	BlueprintEditor.Pin()->PasteDropLocation = FVector2D(0, 0);
+	BlueprintEditor.Pin()->PasteDropLocation = SelectedWidgetContextMenuLocation;
 
 	if ( BlueprintEditor.Pin()->DesignerCommandList->ProcessCommandBindings(InKeyEvent) )
 	{
