Rename UMG DetailCustomization files to UMDetailsCustomization to avoid confusion with other generic details customization files

diff --git a/Engine/Source/Editor/UMGEditor/Private/Customizations/DetailCustomizations.cpp b/Engine/Source/Editor/UMGEditor/Private/Customizations/DetailCustomizations.cpp
deleted file mode 100644
diff --git a/Engine/Source/Editor/UMGEditor/Private/Customizations/DetailCustomizations.h b/Engine/Source/Editor/UMGEditor/Private/Customizations/DetailCustomizations.h
deleted file mode 100644
diff --git a/Engine/Source/Editor/UMGEditor/Private/Customizations/UMGDetailCustomizations.cpp b/Engine/Source/Editor/UMGEditor/Private/Customizations/UMGDetailCustomizations.cpp
new file mode 100644
--- /dev/null
+++ b/Engine/Source/Editor/UMGEditor/Private/Customizations/UMGDetailCustomizations.cpp
@@ -0,0 +1,316 @@
+// Copyright 1998-2015 Epic Games, Inc. All Rights Reserved.
+
+#include "UMGEditorPrivatePCH.h"
+
+#include "UMGDetailCustomizations.h"
+#include "BlueprintModes/WidgetBlueprintApplicationModes.h"
+#include "PropertyEditing.h"
+#include "ObjectEditorUtils.h"
+#include "WidgetGraphSchema.h"
+#include "ScopedTransaction.h"
+#include "BlueprintEditorUtils.h"
+#include "WidgetGraphSchema.h"
+#include "Components/Widget.h"
+#include "WidgetBlueprint.h"
+#include "K2Node_ComponentBoundEvent.h"
+#include "Kismet2/KismetEditorUtilities.h"
+#include "EdGraphSchema_K2_Actions.h"
+#include "WidgetBlueprintEditor.h"
+#include "Components/PanelSlot.h"
+#include "SPropertyBinding.h"
+
+#define LOCTEXT_NAMESPACE "UMG"
+
+class SGraphSchemaActionButton : public SCompoundWidget, public FGCObject
+{
+public:
+
+	SLATE_BEGIN_ARGS(SGraphSchemaActionButton) {}
+		/** Slot for this designers content (optional) */
+		SLATE_DEFAULT_SLOT(FArguments, Content)
+	SLATE_END_ARGS()
+
+	void Construct(const FArguments& InArgs, TSharedPtr<FWidgetBlueprintEditor> InEditor, TSharedPtr<FEdGraphSchemaAction> InClickAction)
+	{
+		Editor = InEditor;
+		Action = InClickAction;
+
+		ChildSlot
+		[
+			SNew(SButton)
+			.ButtonStyle(FEditorStyle::Get(), "FlatButton.Success")
+			.TextStyle(FEditorStyle::Get(), "NormalText")
+			.HAlign(HAlign_Center)
+			.ForegroundColor(FSlateColor::UseForeground())
+			.ToolTipText(FText::FromString(Action->TooltipDescription))
+			.OnClicked(this, &SGraphSchemaActionButton::AddOrViewEventBinding)
+			[
+				InArgs._Content.Widget
+			]
+		];
+	}
+
+private:
+	FReply AddOrViewEventBinding()
+	{
+		UBlueprint* Blueprint = Editor.Pin()->GetBlueprintObj();
+
+		UEdGraph* TargetGraph = nullptr;
+		if ( Blueprint->UbergraphPages.Num() > 0 )
+		{
+			TargetGraph = Blueprint->UbergraphPages[0]; // Just use the first graph
+		}
+
+		if ( TargetGraph != nullptr )
+		{
+			Editor.Pin()->SetCurrentMode(FWidgetBlueprintApplicationModes::GraphMode);
+
+			// Figure out a decent place to stick the node
+			const FVector2D NewNodePos = TargetGraph->GetGoodPlaceForNewNode();
+
+			Action->PerformAction(TargetGraph, nullptr, NewNodePos);
+		}
+
+		return FReply::Handled();
+	}
+
+	virtual void AddReferencedObjects(FReferenceCollector& Collector) override
+	{
+		Action->AddReferencedObjects(Collector);
+	}
+
+private:
+	TWeakPtr<FWidgetBlueprintEditor> Editor;
+
+	TSharedPtr<FEdGraphSchemaAction> Action;
+};
+
+void FBlueprintWidgetCustomization::CreateEventCustomization( IDetailLayoutBuilder& DetailLayout, UDelegateProperty* Property, UWidget* Widget )
+{
+	TSharedRef<IPropertyHandle> DelegatePropertyHandle = DetailLayout.GetProperty(Property->GetFName(), CastChecked<UClass>(Property->GetOuter()));
+
+	const bool bHasValidHandle = DelegatePropertyHandle->IsValidHandle();
+	if(!bHasValidHandle)
+	{
+		return;
+	}
+
+	IDetailCategoryBuilder& PropertyCategory = DetailLayout.EditCategory(FObjectEditorUtils::GetCategoryFName(Property), FText::GetEmpty(), ECategoryPriority::Uncommon);
+
+	IDetailPropertyRow& PropertyRow = PropertyCategory.AddProperty(DelegatePropertyHandle);
+	PropertyRow.OverrideResetToDefault(true, FSimpleDelegate::CreateSP(this, &FBlueprintWidgetCustomization::ResetToDefault_RemoveBinding, DelegatePropertyHandle));
+
+	FString LabelStr = Property->GetName();
+	LabelStr.RemoveFromEnd(TEXT("Event"));
+
+	FText Label = FText::FromString(LabelStr);
+
+	const bool bShowChildren = true;
+	PropertyRow.CustomWidget(bShowChildren)
+		.NameContent()
+		[
+			SNew(SHorizontalBox)
+
+			+ SHorizontalBox::Slot()
+			.AutoWidth()
+			.VAlign(VAlign_Center)
+			.Padding(0,0,5,0)
+			[
+				SNew(SImage)
+				.Image(FEditorStyle::GetBrush("GraphEditor.Event_16x"))
+			]
+
+			+ SHorizontalBox::Slot()
+			.VAlign(VAlign_Center)
+			[
+				SNew(STextBlock)
+				.Text(Label)
+			]
+		]
+		.ValueContent()
+		.MinDesiredWidth(200)
+		.MaxDesiredWidth(250)
+		[
+			SNew(SPropertyBinding, Editor.Pin().ToSharedRef(), Property, DelegatePropertyHandle)
+			.GeneratePureBindings(false)
+		];
+}
+
+void FBlueprintWidgetCustomization::ResetToDefault_RemoveBinding(TSharedRef<IPropertyHandle> PropertyHandle)
+{
+	const FScopedTransaction Transaction(LOCTEXT("BindDelegate", "Remove Binding"));
+
+	Blueprint->Modify();
+
+	TArray<UObject*> OuterObjects;
+	PropertyHandle->GetOuterObjects(OuterObjects);
+	for ( UObject* SelectedObject : OuterObjects )
+	{
+		FDelegateEditorBinding Binding;
+		Binding.ObjectName = SelectedObject->GetName();
+		Binding.PropertyName = PropertyHandle->GetProperty()->GetFName();
+
+		Blueprint->Bindings.Remove(Binding);
+	}
+
+	FBlueprintEditorUtils::MarkBlueprintAsStructurallyModified(Blueprint);
+}
+
+void FBlueprintWidgetCustomization::CreateMulticastEventCustomization(IDetailLayoutBuilder& DetailLayout, FName ThisComponentName, UClass* PropertyClass, UMulticastDelegateProperty* DelegateProperty)
+{
+	const FString AddString = FString(TEXT("Add "));
+	const FString ViewString = FString(TEXT("View "));
+
+	const UEdGraphSchema_K2* K2Schema = GetDefault<UEdGraphSchema_K2>();
+
+	if ( !K2Schema->CanUserKismetAccessVariable(DelegateProperty, PropertyClass, UEdGraphSchema_K2::MustBeDelegate) )
+	{
+		return;
+	}
+
+	FString PropertyTooltip = DelegateProperty->GetToolTipText().ToString();
+	if ( PropertyTooltip.IsEmpty() )
+	{
+		PropertyTooltip = DelegateProperty->GetName();
+	}
+
+	// Add on category for delegate property
+	const FString EventCategory = FObjectEditorUtils::GetCategory(DelegateProperty);
+
+	UObjectProperty* ComponentProperty = FindField<UObjectProperty>(Blueprint->SkeletonGeneratedClass, ThisComponentName);
+
+	if ( !ComponentProperty )
+	{
+		return;
+	}
+
+	const UK2Node* EventNode = FKismetEditorUtilities::FindBoundEventForComponent(Blueprint, DelegateProperty->GetFName(), ComponentProperty->GetFName());
+
+	TSharedPtr<FEdGraphSchemaAction> ClickAction;
+	TSharedPtr<SWidget> ButtonContent;
+
+	if ( EventNode )
+	{
+		TSharedPtr<FEdGraphSchemaAction_K2ViewNode> NewDelegateNode = TSharedPtr<FEdGraphSchemaAction_K2ViewNode>(new FEdGraphSchemaAction_K2ViewNode(EventCategory, FText::FromString(ViewString + DelegateProperty->GetName()), PropertyTooltip, K2Schema->AG_LevelReference));
+		NewDelegateNode->NodePtr = EventNode;
+
+		ClickAction = NewDelegateNode;
+
+		ButtonContent = SNew(STextBlock)
+			.Text(LOCTEXT("ViewEvent", "View"));
+	}
+	else
+	{
+		TSharedPtr<FEdGraphSchemaAction_K2NewNode> NewDelegateNode = TSharedPtr<FEdGraphSchemaAction_K2NewNode>(new FEdGraphSchemaAction_K2NewNode(EventCategory, FText::FromString(AddString + DelegateProperty->GetName()), PropertyTooltip, K2Schema->AG_LevelReference));
+
+		UK2Node_ComponentBoundEvent* NewComponentEvent = NewObject<UK2Node_ComponentBoundEvent>(Blueprint, UK2Node_ComponentBoundEvent::StaticClass());
+		NewComponentEvent->InitializeComponentBoundEventParams(ComponentProperty, DelegateProperty);
+		NewDelegateNode->NodeTemplate = NewComponentEvent;
+		NewDelegateNode->bGotoNode = true;
+
+		ClickAction = NewDelegateNode;
+
+		ButtonContent = SNew(SImage)
+			.Image(FEditorStyle::GetBrush("Plus"));
+	}
+
+	TSharedRef<IPropertyHandle> DelegatePropertyHandle = DetailLayout.GetProperty(DelegateProperty->GetFName(), CastChecked<UClass>(DelegateProperty->GetOuter()));
+
+	IDetailCategoryBuilder& PropertyCategory = DetailLayout.EditCategory("Events", LOCTEXT("Events", "Events"), ECategoryPriority::Uncommon);
+
+	FText DelegatePropertyName = FText::FromString(DelegateProperty->GetName());
+	PropertyCategory.AddCustomRow(DelegatePropertyName)
+	.NameContent()
+	[
+		SNew(SHorizontalBox)
+
+		+ SHorizontalBox::Slot()
+		.AutoWidth()
+		.VAlign(VAlign_Center)
+		.Padding(0,0,5,0)
+		[
+			SNew(SImage)
+			.Image(FEditorStyle::GetBrush("GraphEditor.Event_16x"))
+		]
+
+		+ SHorizontalBox::Slot()
+		.VAlign(VAlign_Center)
+		[
+			SNew(STextBlock)
+			.Font(IDetailLayoutBuilder::GetDetailFont())
+			.Text(DelegatePropertyName)
+		]
+	]
+	.ValueContent()
+	.MinDesiredWidth(150)
+	.MaxDesiredWidth(200)
+	[
+		SNew(SGraphSchemaActionButton, Editor.Pin(), ClickAction)
+		[
+			ButtonContent.ToSharedRef()
+		]
+	];
+}
+
+void FBlueprintWidgetCustomization::CustomizeDetails( IDetailLayoutBuilder& DetailLayout )
+{
+	static const FName LayoutCategoryKey(TEXT("Layout"));
+
+	TArray< TWeakObjectPtr<UObject> > OutObjects;
+	DetailLayout.GetObjectsBeingCustomized(OutObjects);
+	
+	if ( OutObjects.Num() == 1 )
+	{
+		if ( UWidget* Widget = Cast<UWidget>(OutObjects[0].Get()) )
+		{
+			if ( Widget->Slot )
+			{
+				UClass* SlotClass = Widget->Slot->GetClass();
+				FText LayoutCatName = FText::Format(LOCTEXT("SlotNameFmt", "Slot ({0})"), SlotClass->GetDisplayNameText());
+
+				DetailLayout.EditCategory(LayoutCategoryKey, LayoutCatName, ECategoryPriority::TypeSpecific);
+			}
+			else
+			{
+				auto& Category = DetailLayout.EditCategory(LayoutCategoryKey);
+				// TODO UMG Hide Category
+			}
+		}
+	}
+
+	PerformBindingCustomization(DetailLayout);
+}
+
+void FBlueprintWidgetCustomization::PerformBindingCustomization(IDetailLayoutBuilder& DetailLayout)
+{
+	static const FName IsBindableEventName(TEXT("IsBindableEvent"));
+
+	TArray< TWeakObjectPtr<UObject> > OutObjects;
+	DetailLayout.GetObjectsBeingCustomized(OutObjects);
+
+	if ( OutObjects.Num() == 1 )
+	{
+		UWidget* Widget = Cast<UWidget>(OutObjects[0].Get());
+		UClass* PropertyClass = OutObjects[0].Get()->GetClass();
+
+		for ( TFieldIterator<UProperty> PropertyIt(PropertyClass, EFieldIteratorFlags::IncludeSuper); PropertyIt; ++PropertyIt )
+		{
+			UProperty* Property = *PropertyIt;
+
+			if ( UDelegateProperty* DelegateProperty = Cast<UDelegateProperty>(*PropertyIt) )
+			{
+				//TODO Remove the code to use ones that end with "Event".  Prefer metadata flag.
+				if ( DelegateProperty->GetBoolMetaData(IsBindableEventName) || DelegateProperty->GetName().EndsWith(TEXT("Event")) )
+				{
+					CreateEventCustomization(DetailLayout, DelegateProperty, Widget);
+				}
+			}
+			else if ( UMulticastDelegateProperty* MulticastDelegateProperty = Cast<UMulticastDelegateProperty>(Property) )
+			{
+				CreateMulticastEventCustomization(DetailLayout, OutObjects[0].Get()->GetFName(), PropertyClass, MulticastDelegateProperty);
+			}
+		}
+	}
+}
+
+#undef LOCTEXT_NAMESPACE
diff --git a/Engine/Source/Editor/UMGEditor/Private/Customizations/UMGDetailCustomizations.h b/Engine/Source/Editor/UMGEditor/Private/Customizations/UMGDetailCustomizations.h
new file mode 100644
--- /dev/null
+++ b/Engine/Source/Editor/UMGEditor/Private/Customizations/UMGDetailCustomizations.h
@@ -0,0 +1,45 @@
+// Copyright 1998-2015 Epic Games, Inc. All Rights Reserved.
+
+#pragma once
+#include "IDetailCustomization.h"
+
+class FWidgetBlueprintEditor;
+class UWidget;
+class IPropertyHandle;
+class UWidgetBlueprint;
+
+/**
+ * Provides the customization for all UWidgets.  Bindings, style disabling...etc.
+ */
+class FBlueprintWidgetCustomization : public IDetailCustomization
+{
+public:
+	/** Makes a new instance of this detail layout class for a specific detail view requesting it */
+	static TSharedRef<class IDetailCustomization> MakeInstance(TSharedRef<FWidgetBlueprintEditor> InEditor, UBlueprint* InBlueprint)
+	{
+		return MakeShareable(new FBlueprintWidgetCustomization(InEditor, InBlueprint));
+	}
+
+	FBlueprintWidgetCustomization(TSharedRef<FWidgetBlueprintEditor> InEditor, UBlueprint* InBlueprint)
+		: Editor(InEditor)
+		, Blueprint(CastChecked<UWidgetBlueprint>(InBlueprint))
+	{
+	}
+	
+	/** IDetailCustomization interface */
+	virtual void CustomizeDetails( IDetailLayoutBuilder& DetailLayout ) override;
+
+private:
+	void PerformBindingCustomization(IDetailLayoutBuilder& DetailLayout);
+
+	void CreateEventCustomization( IDetailLayoutBuilder& DetailLayout, UDelegateProperty* Property, UWidget* Widget );
+
+	void CreateMulticastEventCustomization(IDetailLayoutBuilder& DetailLayout, FName ThisComponentName, UClass* PropertyClass, UMulticastDelegateProperty* Property);
+
+	void ResetToDefault_RemoveBinding(TSharedRef<IPropertyHandle> PropertyHandle);
+private:
+
+	TWeakPtr<FWidgetBlueprintEditor> Editor;
+
+	UWidgetBlueprint* Blueprint;
+};
diff --git a/Engine/Source/Editor/UMGEditor/Private/Details/SWidgetDetailsView.cpp b/Engine/Source/Editor/UMGEditor/Private/Details/SWidgetDetailsView.cpp
--- a/Engine/Source/Editor/UMGEditor/Private/Details/SWidgetDetailsView.cpp
+++ b/Engine/Source/Editor/UMGEditor/Private/Details/SWidgetDetailsView.cpp
@@ -11,7 +11,7 @@
 #include "DetailWidgetExtensionHandler.h"
 #include "EditorClassUtils.h"
 
-#include "Customizations/DetailCustomizations.h"
+#include "Customizations/UMGDetailCustomizations.h"
 #include "Customizations/SlateBrushCustomization.h"
 #include "Customizations/SlateFontInfoCustomization.h"
 
