UMG - Sequencer - Fix an issue where undoing deleting a track would break the object/animation association.

diff --git a/Engine/Source/Editor/UMGEditor/Private/Animation/UMGSequencerObjectBindingManager.cpp b/Engine/Source/Editor/UMGEditor/Private/Animation/UMGSequencerObjectBindingManager.cpp
--- a/Engine/Source/Editor/UMGEditor/Private/Animation/UMGSequencerObjectBindingManager.cpp
+++ b/Engine/Source/Editor/UMGEditor/Private/Animation/UMGSequencerObjectBindingManager.cpp
@@ -109,6 +109,7 @@ void FUMGSequencerObjectBindingManager::UnbindPossessableObjects( const FGuid& P
 
 	UWidgetBlueprint* WidgetBlueprint = WidgetBlueprintEditor.GetWidgetBlueprintObj();
 
+	WidgetAnimation->Modify();
 	WidgetAnimation->AnimationBindings.RemoveAll( [&]( const FWidgetAnimationBinding& Binding ) { return Binding.AnimationGuid == PossessableGuid; } );
 
 }
