UMG - UE-6692 - The grid lines now use the snap size in the designer.

diff --git a/Engine/Source/Editor/UMGEditor/Private/Designer/SDesignSurface.cpp b/Engine/Source/Editor/UMGEditor/Private/Designer/SDesignSurface.cpp
--- a/Engine/Source/Editor/UMGEditor/Private/Designer/SDesignSurface.cpp
+++ b/Engine/Source/Editor/UMGEditor/Private/Designer/SDesignSurface.cpp
@@ -117,7 +117,6 @@ void SDesignSurface::Construct(const FArguments& InArgs)
 	ZoomLevel = ZoomLevels->GetDefaultZoomLevel();
 	PreviousZoomLevel = ZoomLevels->GetDefaultZoomLevel();
 	PostChangedZoom();
-	SnapGridSize = 16.0f;
 	AllowContinousZoomInterpolation = InArgs._AllowContinousZoomInterpolation;
 	bIsPanning = false;
 
@@ -470,11 +469,6 @@ FVector2D SDesignSurface::GetViewOffset() const
 	return ViewOffset;
 }
 
-float SDesignSurface::GetSnapGridSize() const
-{
-	return SnapGridSize;
-}
-
 FVector2D SDesignSurface::GraphCoordToPanelCoord(const FVector2D& GraphSpaceCoordinate) const
 {
 	return ( GraphSpaceCoordinate - GetViewOffset() ) * GetZoomAmount();
@@ -485,11 +479,21 @@ FVector2D SDesignSurface::PanelCoordToGraphCoord(const FVector2D& PanelSpaceCoor
 	return PanelSpaceCoordinate / GetZoomAmount() + GetViewOffset();
 }
 
+int32 SDesignSurface::GetGraphRulePeriod() const
+{
+	return (int32)FEditorStyle::GetFloat("Graph.Panel.GridRulePeriod");
+}
+
+float SDesignSurface::GetGridScaleAmount() const
+{
+	return 1;
+}
+
 void SDesignSurface::PaintBackgroundAsLines(const FSlateBrush* BackgroundImage, const FGeometry& AllottedGeometry, const FSlateRect& MyClippingRect, FSlateWindowElementList& OutDrawElements, int32& DrawLayerId) const
 {
 	const bool bAntialias = false;
 
-	const int32 RulePeriod = (int32)FEditorStyle::GetFloat("Graph.Panel.GridRulePeriod");
+	const int32 RulePeriod = GetGraphRulePeriod();
 	check(RulePeriod > 0);
 
 	const FLinearColor RegularColor(FEditorStyle::GetColor("Graph.Panel.GridLineColor"));
@@ -497,7 +501,7 @@ void SDesignSurface::PaintBackgroundAsLines(const FSlateBrush* BackgroundImage,
 	const FLinearColor CenterColor(FEditorStyle::GetColor("Graph.Panel.GridCenterColor"));
 	const float GraphSmallestGridSize = 8.0f;
 	const float RawZoomFactor = GetZoomAmount();
-	const float NominalGridSize = GetSnapGridSize();
+	const float NominalGridSize = GetSnapGridSize() * GetGridScaleAmount();
 
 	float ZoomFactor = RawZoomFactor;
 	float Inflation = 1.0f;
diff --git a/Engine/Source/Editor/UMGEditor/Private/Designer/SDesignSurface.h b/Engine/Source/Editor/UMGEditor/Private/Designer/SDesignSurface.h
--- a/Engine/Source/Editor/UMGEditor/Private/Designer/SDesignSurface.h
+++ b/Engine/Source/Editor/UMGEditor/Private/Designer/SDesignSurface.h
@@ -36,9 +36,7 @@ public:
 protected:
 	void PaintBackgroundAsLines(const FSlateBrush* BackgroundImage, const FGeometry& AllottedGeometry, const FSlateRect& MyClippingRect, FSlateWindowElementList& OutDrawElements, int32& DrawLayerId) const;
 
-	/** Get the grid snap size */
-	float GetSnapGridSize() const;
-
+	/** Gets the current zoom factor. */
 	float GetZoomAmount() const;
 
 	void ChangeZoomLevel(int32 ZoomLevelDelta, const FVector2D& WidgetSpaceZoomOrigin, bool bOverrideZoomLimiting);
@@ -56,19 +54,21 @@ protected:
 
 	FVector2D GetViewOffset() const;
 
-	virtual FSlateRect ComputeAreaBounds() const;
-
 	FSlateRect ComputeSensibleBounds() const;
 
 	FVector2D GraphCoordToPanelCoord(const FVector2D& GraphSpaceCoordinate) const;
 	FVector2D PanelCoordToGraphCoord(const FVector2D& PanelSpaceCoordinate) const;
 
+protected:
+	virtual FSlateRect ComputeAreaBounds() const;
+	virtual float GetGridScaleAmount() const;
+	virtual int32 GetGraphRulePeriod() const;
+	virtual int32 GetSnapGridSize() const = 0;
+
 protected:
 	/** The position within the graph at which the user is looking */
 	FVector2D ViewOffset;
 
-	float SnapGridSize;
-
 	/** Previous Zoom Level */
 	int32 PreviousZoomLevel;
 
diff --git a/Engine/Source/Editor/UMGEditor/Private/Designer/SDesignerView.cpp b/Engine/Source/Editor/UMGEditor/Private/Designer/SDesignerView.cpp
--- a/Engine/Source/Editor/UMGEditor/Private/Designer/SDesignerView.cpp
+++ b/Engine/Source/Editor/UMGEditor/Private/Designer/SDesignerView.cpp
@@ -760,6 +760,22 @@ FSlateRect SDesignerView::ComputeAreaBounds() const
 	return FSlateRect(0, 0, GetPreviewAreaWidth().Get(), GetPreviewAreaHeight().Get());
 }
 
+int32 SDesignerView::GetSnapGridSize() const
+{
+	const UWidgetDesignerSettings* DesignerSettings = GetDefault<UWidgetDesignerSettings>();
+	return DesignerSettings->GridSnapSize;
+}
+
+int32 SDesignerView::GetGraphRulePeriod() const
+{
+	return 10;
+}
+
+float SDesignerView::GetGridScaleAmount() const
+{
+	return GetPreviewDPIScale();
+}
+
 EVisibility SDesignerView::GetInfoBarVisibility() const
 {
 	if ( DesignerMessage != EDesignerMessage::None )
diff --git a/Engine/Source/Editor/UMGEditor/Private/Designer/SDesignerView.h b/Engine/Source/Editor/UMGEditor/Private/Designer/SDesignerView.h
--- a/Engine/Source/Editor/UMGEditor/Private/Designer/SDesignerView.h
+++ b/Engine/Source/Editor/UMGEditor/Private/Designer/SDesignerView.h
@@ -76,6 +76,12 @@ public:
 	virtual void AddReferencedObjects(FReferenceCollector& Collector) override;
 	// End of FGCObject interface
 
+protected:
+	virtual FSlateRect ComputeAreaBounds() const override;
+	virtual int32 GetGraphRulePeriod() const override;
+	virtual float GetGridScaleAmount() const override;
+	virtual int32 GetSnapGridSize() const override;
+
 private:
 	/** Establishes the resolution and aspect ratio to use on construction from config settings */
 	void SetStartupResolution();
@@ -99,8 +105,6 @@ private:
 	/** Gets the DPI scale that would be applied given the current preview width and height */
 	float GetPreviewDPIScale() const;
 
-	virtual FSlateRect ComputeAreaBounds() const override;
-
 	/** Adds any pending selected widgets to the selection set */
 	void ResolvePendingSelectedWidgets();
 
