UMG - Fixing the right click-rename option to properly call the command list.

diff --git a/Engine/Source/Editor/UMGEditor/Private/Hierarchy/SHierarchyView.cpp b/Engine/Source/Editor/UMGEditor/Private/Hierarchy/SHierarchyView.cpp
--- a/Engine/Source/Editor/UMGEditor/Private/Hierarchy/SHierarchyView.cpp
+++ b/Engine/Source/Editor/UMGEditor/Private/Hierarchy/SHierarchyView.cpp
@@ -218,7 +218,7 @@ void SHierarchyView::OnBlueprintChanged(UBlueprint* InBlueprint)
 
 TSharedPtr<SWidget> SHierarchyView::WidgetHierarchy_OnContextMenuOpening()
 {
-	FMenuBuilder MenuBuilder(true, nullptr);
+	FMenuBuilder MenuBuilder(true, CommandList);
 
 	FWidgetBlueprintEditorUtils::CreateWidgetContextMenu(MenuBuilder, BlueprintEditor.Pin().ToSharedRef(), FVector2D(0, 0));
 
