UMG - UE-14808 - The drag/drop code in the designer now filters out non-panels from the possible parents you can drop onto, enabling users to drop new widgets onto a canvas even if say, an image is blocking your view, it will just move to the next highest parent that's a UPanelWidget.

diff --git a/Engine/Source/Editor/UMGEditor/Private/Designer/SDesignerView.cpp b/Engine/Source/Editor/UMGEditor/Private/Designer/SDesignerView.cpp
--- a/Engine/Source/Editor/UMGEditor/Private/Designer/SDesignerView.cpp
+++ b/Engine/Source/Editor/UMGEditor/Private/Designer/SDesignerView.cpp
@@ -1068,7 +1068,7 @@ SDesignerView::FWidgetHitResult::FWidgetHitResult()
 {
 }
 
-bool SDesignerView::FindWidgetUnderCursor(const FGeometry& MyGeometry, const FPointerEvent& MouseEvent, FWidgetHitResult& HitResult)
+bool SDesignerView::FindWidgetUnderCursor(const FGeometry& MyGeometry, const FPointerEvent& MouseEvent, TSubclassOf<UWidget> FindType, FWidgetHitResult& HitResult)
 {
 	//@TODO UMG Make it so you can request dropable widgets only, to find the first parentable.
 
@@ -1100,6 +1100,13 @@ bool SDesignerView::FindWidgetUnderCursor(const FGeometry& MyGeometry, const FPo
 				WidgetUnderCursor = nullptr;
 				continue;
 			}
+
+			// Ignore widgets that don't pass our find widget filter.
+			if ( WidgetUnderCursor->GetClass()->IsChildOf(FindType) == false )
+			{
+				WidgetUnderCursor = nullptr;
+				continue;
+			}
 			
 			// We successfully found a widget that's accessible from the root.
 			if ( WidgetUnderCursor )
@@ -1148,7 +1155,7 @@ FReply SDesignerView::OnMouseButtonDown(const FGeometry& MyGeometry, const FPoin
 
 	//TODO UMG Undoable Selection
 	FWidgetHitResult HitResult;
-	if ( FindWidgetUnderCursor(MyGeometry, MouseEvent, HitResult) )
+	if ( FindWidgetUnderCursor(MyGeometry, MouseEvent, UWidget::StaticClass(), HitResult) )
 	{
 		SelectedWidgetContextMenuLocation = HitResult.WidgetArranged.Geometry.AbsoluteToLocal(MouseEvent.GetScreenSpacePosition());
 
@@ -1273,7 +1280,7 @@ FReply SDesignerView::OnMouseMove(const FGeometry& MyGeometry, const FPointerEve
 	// Update the hovered widget under the mouse
 	auto PinnedBlueprintEditor = BlueprintEditor.Pin();
 	FWidgetHitResult HitResult;
-	if ( FindWidgetUnderCursor(MyGeometry, MouseEvent, HitResult) )
+	if ( FindWidgetUnderCursor(MyGeometry, MouseEvent, UWidget::StaticClass(), HitResult) )
 	{
 		PinnedBlueprintEditor->SetHoveredWidget(HitResult.Widget);
 	}
@@ -1674,7 +1681,7 @@ UWidget* SDesignerView::ProcessDropAndAddWidget(const FGeometry& MyGeometry, con
 	UWidget* Target = nullptr;
 
 	FWidgetHitResult HitResult;
-	if ( FindWidgetUnderCursor(MyGeometry, DragDropEvent, HitResult) )
+	if ( FindWidgetUnderCursor(MyGeometry, DragDropEvent, UPanelWidget::StaticClass(), HitResult) )
 	{
 		Target = bIsPreview ? HitResult.Widget.GetPreview() : HitResult.Widget.GetTemplate();
 	}
diff --git a/Engine/Source/Editor/UMGEditor/Private/Designer/SDesignerView.h b/Engine/Source/Editor/UMGEditor/Private/Designer/SDesignerView.h
--- a/Engine/Source/Editor/UMGEditor/Private/Designer/SDesignerView.h
+++ b/Engine/Source/Editor/UMGEditor/Private/Designer/SDesignerView.h
@@ -188,7 +188,7 @@ private:
 	};
 
 	/** @returns Gets the widget under the cursor based on a mouse pointer event. */
-	bool FindWidgetUnderCursor(const FGeometry& MyGeometry, const FPointerEvent& MouseEvent, FWidgetHitResult& HitResult);
+	bool FindWidgetUnderCursor(const FGeometry& MyGeometry, const FPointerEvent& MouseEvent, TSubclassOf<UWidget> FindType, FWidgetHitResult& HitResult);
 
 private:
 	FReply HandleZoomToFitClicked();
diff --git a/Engine/Source/Runtime/Engine/Classes/Engine/UserInterfaceSettings.h b/Engine/Source/Runtime/Engine/Classes/Engine/UserInterfaceSettings.h
--- a/Engine/Source/Runtime/Engine/Classes/Engine/UserInterfaceSettings.h
+++ b/Engine/Source/Runtime/Engine/Classes/Engine/UserInterfaceSettings.h
@@ -93,7 +93,7 @@ public:
 
 	void ForceLoadResources();
 
-	/** Gets the current scale of the UI based on the size */
+	/** Gets the current scale of the UI based on the size of a viewport */
 	float GetDPIScaleBasedOnSize(FIntPoint Size) const;
 
 #if WITH_EDITOR
