Integrating 2548936 into ML - #umg Fixed crash when user attempts to drag-create a new widget in a panel that can't take any more children

diff --git a/Engine/Source/Editor/UMGEditor/Private/Hierarchy/SHierarchyViewItem.cpp b/Engine/Source/Editor/UMGEditor/Private/Hierarchy/SHierarchyViewItem.cpp
--- a/Engine/Source/Editor/UMGEditor/Private/Hierarchy/SHierarchyViewItem.cpp
+++ b/Engine/Source/Editor/UMGEditor/Private/Hierarchy/SHierarchyViewItem.cpp
@@ -134,6 +134,7 @@ TOptional<EItemDropZone> ProcessHierarchyDragDrop(const FDragDropEvent& DragDrop
 	TSharedPtr<FWidgetTemplateDragDropOp> TemplateDragDropOp = DragDropEvent.GetOperationAs<FWidgetTemplateDragDropOp>();
 	if ( TemplateDragDropOp.IsValid() )
 	{
+		TemplateDragDropOp->ResetToDefaultToolTip();
 		TemplateDragDropOp->SetCursorOverride(TOptional<EMouseCursor::Type>());
 
 		// Are we adding to the root?
@@ -151,32 +152,40 @@ TOptional<EItemDropZone> ProcessHierarchyDragDrop(const FDragDropEvent& DragDrop
 		// Are we adding to a panel?
 		else if ( UPanelWidget* Parent = Cast<UPanelWidget>(TargetItem.GetTemplate()) )
 		{
-			// TODO UMG Allow showing a preview of this.
-			if ( bIsDrop )
+			if (!Parent->CanAddMoreChildren())
 			{
-				UWidget* Widget = TemplateDragDropOp->Template->Create(Blueprint->WidgetTree);
-
-				UPanelSlot* NewSlot = nullptr;
-				if ( Index.IsSet() )
-				{
-					NewSlot = Parent->InsertChildAt(Index.GetValue(), Widget);
-				}
-				else
+				TemplateDragDropOp->CurrentHoverText = LOCTEXT("NoAdditionalChildren", "Widget can't accept additional children.");
+			}
+			else
+			{
+				// TODO UMG Allow showing a preview of this.
+				if (bIsDrop)
 				{
-					NewSlot = Parent->AddChild(Widget);
+					UWidget* Widget = TemplateDragDropOp->Template->Create(Blueprint->WidgetTree);
+
+					UPanelSlot* NewSlot = nullptr;
+					if (Index.IsSet())
+					{
+						NewSlot = Parent->InsertChildAt(Index.GetValue(), Widget);
+					}
+					else
+					{
+						NewSlot = Parent->AddChild(Widget);
+					}
+					check(NewSlot);
+
+					FBlueprintEditorUtils::MarkBlueprintAsStructurallyModified(Blueprint);
 				}
-				check(NewSlot);
 
-				FBlueprintEditorUtils::MarkBlueprintAsStructurallyModified(Blueprint);
+				return EItemDropZone::OntoItem;
 			}
-
-			return EItemDropZone::OntoItem;
 		}
 		else
 		{
-			TemplateDragDropOp->SetCursorOverride(EMouseCursor::SlashedCircle);
+			TemplateDragDropOp->CurrentHoverText = LOCTEXT("CantHaveChildren", "Widget can't have children.");
 		}
 
+		TemplateDragDropOp->SetCursorOverride(EMouseCursor::SlashedCircle);
 		return TOptional<EItemDropZone>();
 	}
 
