BP/UMG - Updating the IBlueprintCompiler interface to include the compiler options passed in for Pre/Post compile.  Improving the logic in the UMG module to correctly reregister the Widget Components.  Also reducing when it reregisters avoiding cases where it's unnessesary (skeletal compile..etc).

diff --git a/Engine/Source/Editor/KismetCompiler/Private/KismetCompilerModule.cpp b/Engine/Source/Editor/KismetCompiler/Private/KismetCompilerModule.cpp
--- a/Engine/Source/Editor/KismetCompiler/Private/KismetCompilerModule.cpp
+++ b/Engine/Source/Editor/KismetCompiler/Private/KismetCompilerModule.cpp
@@ -140,7 +140,7 @@ void FKismet2CompilerModule::CompileBlueprint(class UBlueprint* Blueprint, const
 
 	for ( IBlueprintCompiler* Compiler : Compilers )
 	{
-		Compiler->PreCompile(Blueprint);
+		Compiler->PreCompile(Blueprint, CompileOptions);
 	}
 
 	if (CompileOptions.CompileType != EKismetCompileType::Cpp)
@@ -218,7 +218,7 @@ void FKismet2CompilerModule::CompileBlueprint(class UBlueprint* Blueprint, const
 
 	for ( IBlueprintCompiler* Compiler : Compilers )
 	{
-		Compiler->PostCompile(Blueprint);
+		Compiler->PostCompile(Blueprint, CompileOptions);
 	}
 
 	UPackage* Package = Blueprint->GetOutermost();
diff --git a/Engine/Source/Editor/KismetCompiler/Public/KismetCompilerModule.h b/Engine/Source/Editor/KismetCompiler/Public/KismetCompilerModule.h
--- a/Engine/Source/Editor/KismetCompiler/Public/KismetCompilerModule.h
+++ b/Engine/Source/Editor/KismetCompiler/Public/KismetCompilerModule.h
@@ -15,12 +15,22 @@ class UBlueprint;
 class IBlueprintCompiler
 {
 public:
-	virtual void PreCompile(UBlueprint* Blueprint) = 0;
+	virtual void PreCompile(UBlueprint* Blueprint, const FKismetCompilerOptions& CompileOptions)
+	{
+		PreCompile(Blueprint);
+	}
 
 	virtual bool CanCompile(const UBlueprint* Blueprint) = 0;
 	virtual void Compile(UBlueprint* Blueprint, const FKismetCompilerOptions& CompileOptions, FCompilerResultsLog& Results, TArray<UObject*>* ObjLoaded) = 0;
 	
-	virtual void PostCompile(UBlueprint* Blueprint) = 0;
+	virtual void PostCompile(UBlueprint* Blueprint, const FKismetCompilerOptions& CompileOptions)
+	{
+		PostCompile(Blueprint);
+	}
+
+protected:
+	virtual void PreCompile(UBlueprint* Blueprint) { }
+	virtual void PostCompile(UBlueprint* Blueprint) { }
 };
 
 class IKismetCompilerInterface : public IModuleInterface
diff --git a/Engine/Source/Editor/UMGEditor/Private/UMGEditorModule.cpp b/Engine/Source/Editor/UMGEditor/Private/UMGEditorModule.cpp
--- a/Engine/Source/Editor/UMGEditor/Private/UMGEditorModule.cpp
+++ b/Engine/Source/Editor/UMGEditor/Private/UMGEditorModule.cpp
@@ -105,14 +105,9 @@ public:
 		return Cast<UWidgetBlueprint>(Blueprint) != nullptr;
 	}
 
-	void PreCompile(UBlueprint* Blueprint) override
+	void PreCompile(UBlueprint* Blueprint, const FKismetCompilerOptions& CompileOptions) override
 	{
-		if (!CanCompile(Blueprint))
-		{
-			return;
-		}
-
-		if ( CompileCount == 0 )
+		if ( ReRegister == nullptr && CanCompile(Blueprint) && CompileOptions.CompileType == EKismetCompileType::Full )
 		{
 			ReRegister = new TComponentReregisterContext<UWidgetComponent>();
 		}
@@ -130,19 +125,19 @@ public:
 		}
 	}
 
-	void PostCompile(UBlueprint* Blueprint) override
+	void PostCompile(UBlueprint* Blueprint, const FKismetCompilerOptions& CompileOptions) override
 	{
 		CompileCount--;
 
-		if ( ReRegister && CompileCount == 0 )
+		if ( CompileCount == 0 && ReRegister )
 		{
 			delete ReRegister;
 			ReRegister = nullptr;
-		}
-		
-		if ( GIsEditor && GEditor )
-		{
-			GEditor->RedrawAllViewports(true);
+
+			if ( GIsEditor && GEditor )
+			{
+				GEditor->RedrawAllViewports(true);
+			}
 		}
 	}
 
