UMG - UETOOL-91 - You can now use the arrow keys in the designer to nudge widgets in the canvas by the grid snap amount.  If grid snap is enabled it will actually keep it snapped to the grid, otherwise it will just move by the current snap amount absolutely.

diff --git a/Engine/Source/Editor/UMGEditor/Private/Designer/SDesignerView.cpp b/Engine/Source/Editor/UMGEditor/Private/Designer/SDesignerView.cpp
--- a/Engine/Source/Editor/UMGEditor/Private/Designer/SDesignerView.cpp
+++ b/Engine/Source/Editor/UMGEditor/Private/Designer/SDesignerView.cpp
@@ -1320,9 +1320,94 @@ FReply SDesignerView::OnKeyDown(const FGeometry& MyGeometry, const FKeyEvent& In
 		return FReply::Handled();
 	}
 
+	const UWidgetDesignerSettings* DesignerSettings = GetDefault<UWidgetDesignerSettings>();
+
+	if ( InKeyEvent.GetKey() == EKeys::Up )
+	{
+		return NudgeSelectedWidget(FVector2D(0, -DesignerSettings->GridSnapSize));
+	}
+	else if ( InKeyEvent.GetKey() == EKeys::Down )
+	{
+		return NudgeSelectedWidget(FVector2D(0, DesignerSettings->GridSnapSize));
+	}
+	else if ( InKeyEvent.GetKey() == EKeys::Left )
+	{
+		return NudgeSelectedWidget(FVector2D(-DesignerSettings->GridSnapSize, 0));
+	}
+	else if ( InKeyEvent.GetKey() == EKeys::Right )
+	{
+		return NudgeSelectedWidget(FVector2D(DesignerSettings->GridSnapSize, 0));
+	}
+
+	return FReply::Unhandled();
+}
+
+FReply SDesignerView::OnKeyUp(const FGeometry& MyGeometry, const FKeyEvent& InKeyEvent)
+{
 	return FReply::Unhandled();
 }
 
+FReply SDesignerView::NudgeSelectedWidget(FVector2D Nudge)
+{
+	for ( const FWidgetReference& WidgetRef : GetSelectedWidgets() )
+	{
+		if ( WidgetRef.IsValid() )
+		{
+			UWidget* TemplateWidget = WidgetRef.GetTemplate();
+			UWidget* PreviewWidget = WidgetRef.GetPreview();
+
+			if ( TemplateWidget && PreviewWidget )
+			{
+				UCanvasPanelSlot* TemplateSlot = Cast<UCanvasPanelSlot>(TemplateWidget->Slot);
+				UCanvasPanelSlot* PreviewSlot = Cast<UCanvasPanelSlot>(PreviewWidget->Slot);
+
+				if ( TemplateSlot && PreviewSlot )
+				{
+					const UWidgetDesignerSettings* DesignerSettings = GetDefault<UWidgetDesignerSettings>();
+
+					const FVector2D OldPosition = TemplateSlot->GetPosition();
+					FVector2D NewPosition = OldPosition + Nudge;
+
+					// Determine the new position aligned to the grid.
+					if ( DesignerSettings->GridSnapEnabled )
+					{
+						if ( Nudge.X != 0 )
+						{
+							NewPosition.X = ( (int32)NewPosition.X ) - ( ( (int32)NewPosition.X ) % DesignerSettings->GridSnapSize );
+						}
+						if ( Nudge.Y != 0 )
+						{
+							NewPosition.Y = ( (int32)NewPosition.Y ) - ( ( (int32)NewPosition.Y ) % DesignerSettings->GridSnapSize );
+						}
+					}
+
+					// Offset the size by the same amount moved if we're anchoring along that axis.
+					FVector2D NewSize = TemplateSlot->GetSize();
+					if ( TemplateSlot->GetAnchors().IsStretchedHorizontal() )
+					{
+						NewSize.X -= NewPosition.X - OldPosition.X;
+					}
+					if ( TemplateSlot->GetAnchors().IsStretchedVertical() )
+					{
+						NewSize.Y -= NewPosition.Y - OldPosition.Y;
+					}
+
+					FScopedTransaction Transaction(LOCTEXT("Designer_NudgeWidget", "Nudge Widget"));
+
+					TemplateSlot->Modify();
+
+					TemplateSlot->SetPosition(NewPosition);
+					TemplateSlot->SetSize(NewSize);
+					PreviewSlot->SetPosition(NewPosition);
+					PreviewSlot->SetSize(NewSize);
+				}
+			}
+		}
+	}
+
+	return FReply::Handled();
+}
+
 void SDesignerView::ShowContextMenu(const FGeometry& MyGeometry, const FPointerEvent& MouseEvent)
 {
 	FMenuBuilder MenuBuilder(true, nullptr);
@@ -1910,7 +1995,7 @@ UWidget* SDesignerView::ProcessDropAndAddWidget(const FGeometry& MyGeometry, con
 					}
 
 					// HACK UMG: In order to correctly drop items into the canvas that have a non-zero anchor,
-					// we need to know the layout information after slate has performed a prepass.  So we have
+					// we need to know the layout information after slate has performed a pre-pass.  So we have
 					// to rebase the layout and reinterpret the new position based on anchor point layout data.
 					// This should be pulled out into an extension of some kind so that this can be fixed for
 					// other widgets as well that may need to do work like this.
diff --git a/Engine/Source/Editor/UMGEditor/Private/Designer/SDesignerView.h b/Engine/Source/Editor/UMGEditor/Private/Designer/SDesignerView.h
--- a/Engine/Source/Editor/UMGEditor/Private/Designer/SDesignerView.h
+++ b/Engine/Source/Editor/UMGEditor/Private/Designer/SDesignerView.h
@@ -47,6 +47,7 @@ public:
 	virtual void OnMouseLeave(const FPointerEvent& MouseEvent) override;
 
 	virtual FReply OnKeyDown(const FGeometry& MyGeometry, const FKeyEvent& InKeyEvent) override;
+	virtual FReply OnKeyUp(const FGeometry& MyGeometry, const FKeyEvent& InKeyEvent) override;
 
 	virtual void Tick( const FGeometry& AllottedGeometry, const double InCurrentTime, const float InDeltaTime ) override;
 
@@ -139,6 +140,9 @@ private:
 	/** @return Formatted text for the given resolution params */
 	FText GetResolutionText(int32 Width, int32 Height, const FString& AspectRatio) const;
 
+	/** Move the selected widget by the nudge amount. */
+	FReply NudgeSelectedWidget(FVector2D Nudge);
+
 	FText GetCurrentResolutionText() const;
 	FText GetCurrentDPIScaleText() const;
 	FSlateColor GetResolutionTextColorAndOpacity() const;
