Fixed CIS for OSX builds.

diff --git a/Engine/Source/Editor/UMGEditor/Private/Designer/SDesignerView.cpp b/Engine/Source/Editor/UMGEditor/Private/Designer/SDesignerView.cpp
--- a/Engine/Source/Editor/UMGEditor/Private/Designer/SDesignerView.cpp
+++ b/Engine/Source/Editor/UMGEditor/Private/Designer/SDesignerView.cpp
@@ -1353,13 +1353,13 @@ FReply SDesignerView::NudgeSelectedWidget(FVector2D Nudge)
 	{
 		if ( WidgetRef.IsValid() )
 		{
-			UWidget* TemplateWidget = WidgetRef.GetTemplate();
-			UWidget* PreviewWidget = WidgetRef.GetPreview();
+			UWidget* CurrentTemplateWidget = WidgetRef.GetTemplate();
+			UWidget* CurrentPreviewWidget = WidgetRef.GetPreview();
 
-			if ( TemplateWidget && PreviewWidget )
+			if (CurrentTemplateWidget && CurrentPreviewWidget)
 			{
-				UCanvasPanelSlot* TemplateSlot = Cast<UCanvasPanelSlot>(TemplateWidget->Slot);
-				UCanvasPanelSlot* PreviewSlot = Cast<UCanvasPanelSlot>(PreviewWidget->Slot);
+				UCanvasPanelSlot* TemplateSlot = Cast<UCanvasPanelSlot>(CurrentTemplateWidget->Slot);
+				UCanvasPanelSlot* PreviewSlot = Cast<UCanvasPanelSlot>(CurrentPreviewWidget->Slot);
 
 				if ( TemplateSlot && PreviewSlot )
 				{
