[UE-15679] Fixed widget blueprint not compiling on load, which would cause other (dependent) blueprints to fail compilation under some circumstances. Widget blueprints are effectively never data only, because blueprints that declare variables cannot be data only

diff --git a/Engine/Source/Editor/UMGEditor/Public/WidgetBlueprint.h b/Engine/Source/Editor/UMGEditor/Public/WidgetBlueprint.h
--- a/Engine/Source/Editor/UMGEditor/Public/WidgetBlueprint.h
+++ b/Engine/Source/Editor/UMGEditor/Public/WidgetBlueprint.h
@@ -220,6 +220,9 @@ public:
 	}
 
 	virtual void GetReparentingRules(TSet< const UClass* >& AllowedChildrenOfClasses, TSet< const UClass* >& DisallowedChildrenOfClasses) const override;
+
+	/** UWidget blueprints are never data only, should always compile on load (data only blueprints cannot declare new variables) */
+	virtual bool AlwaysCompileOnLoad() const override { return true; }
 	// End of UBlueprint interface
 
 	static bool ValidateGeneratedClass(const UClass* InClass);
diff --git a/Engine/Source/Editor/UnrealEd/Private/Kismet2/BlueprintEditorUtils.cpp b/Engine/Source/Editor/UnrealEd/Private/Kismet2/BlueprintEditorUtils.cpp
--- a/Engine/Source/Editor/UnrealEd/Private/Kismet2/BlueprintEditorUtils.cpp
+++ b/Engine/Source/Editor/UnrealEd/Private/Kismet2/BlueprintEditorUtils.cpp
@@ -2787,6 +2787,10 @@ bool FBlueprintEditorUtils::IsDataOnlyBlueprint(const UBlueprint* Blueprint)
 		return false;
 	}
 
+	if (Blueprint->AlwaysCompileOnLoad())
+	{
+		return false;
+	}
 
 	if( Blueprint->ParentClass->IsChildOf( UActorComponent::StaticClass() ) )
 	{
diff --git a/Engine/Source/Runtime/Engine/Classes/Engine/Blueprint.h b/Engine/Source/Runtime/Engine/Classes/Engine/Blueprint.h
--- a/Engine/Source/Runtime/Engine/Classes/Engine/Blueprint.h
+++ b/Engine/Source/Runtime/Engine/Classes/Engine/Blueprint.h
@@ -530,6 +530,11 @@ public:
 
 	virtual void GetReparentingRules(TSet< const UClass* >& AllowedChildrenOfClasses, TSet< const UClass* >& DisallowedChildrenOfClasses) const;
 
+	/**
+	* Allows derived blueprints to require compilation on load, otherwise they may get treated as data only and not compiled on load.
+	*/
+	virtual bool AlwaysCompileOnLoad() const { return false; }
+
 private:
 
 	/** Sets the current object being debugged */
