UMG - Fixing a bug that was showing the category field in the details panel when nothing is selected.

diff --git a/Engine/Source/Editor/UMGEditor/Private/Details/SWidgetDetailsView.cpp b/Engine/Source/Editor/UMGEditor/Private/Details/SWidgetDetailsView.cpp
--- a/Engine/Source/Editor/UMGEditor/Private/Details/SWidgetDetailsView.cpp
+++ b/Engine/Source/Editor/UMGEditor/Private/Details/SWidgetDetailsView.cpp
@@ -290,6 +290,11 @@ EVisibility SWidgetDetailsView::GetNameAreaVisibility() const
 
 EVisibility SWidgetDetailsView::GetCategoryAreaVisibility() const
 {
+	if ( SelectedObjects.Num() == 0 )
+	{
+		return EVisibility::Collapsed;
+	}
+
 	return IsWidgetCDOSelected() ? EVisibility::Collapsed : EVisibility::Visible;
 }
 
