UMG - The Widget compiler now warns if you have events that return EventReplys that are unconnected to anything.

diff --git a/Engine/Source/Editor/UMGEditor/Private/WidgetBlueprintCompiler.cpp b/Engine/Source/Editor/UMGEditor/Private/WidgetBlueprintCompiler.cpp
--- a/Engine/Source/Editor/UMGEditor/Private/WidgetBlueprintCompiler.cpp
+++ b/Engine/Source/Editor/UMGEditor/Private/WidgetBlueprintCompiler.cpp
@@ -13,7 +13,7 @@
 #include "Kismet2/BlueprintEditorUtils.h"
 #include "NamedSlot.h"
 
-///-------------------------------------------------------------
+#define LOCTEXT_NAMESPACE "UMG"
 
 FWidgetBlueprintCompiler::FWidgetBlueprintCompiler(UWidgetBlueprint* SourceSketch, FCompilerResultsLog& InMessageLog, const FKismetCompilerOptions& InCompilerOptions, TArray<UObject*>* InObjLoaded)
 	: Super(SourceSketch, InMessageLog, InCompilerOptions, InObjLoaded)
@@ -320,6 +320,37 @@ void FWidgetBlueprintCompiler::SpawnNewClass(const FString& NewClassName)
 	NewClass = NewWidgetBlueprintClass;
 }
 
+void FWidgetBlueprintCompiler::PrecompileFunction(FKismetFunctionContext& Context)
+{
+	Super::PrecompileFunction(Context);
+
+	VerifyEventReplysAreNotEmpty(Context);
+}
+
+void FWidgetBlueprintCompiler::VerifyEventReplysAreNotEmpty(FKismetFunctionContext& Context)
+{
+	TArray<UK2Node_FunctionResult*> FunctionResults;
+	Context.SourceGraph->GetNodesOfClass<UK2Node_FunctionResult>(FunctionResults);
+
+	UScriptStruct* EventReplyStruct = FEventReply::StaticStruct();
+	FEdGraphPinType EventReplyPinType(Schema->PC_Struct, TEXT(""), EventReplyStruct, /*bIsArray =*/false, /*bIsReference =*/false);
+
+	for ( UK2Node_FunctionResult* FunctionResult : FunctionResults )
+	{
+		for ( UEdGraphPin* ReturnPin : FunctionResult->Pins )
+		{
+			if ( ReturnPin->PinType == EventReplyPinType )
+			{
+				const bool IsUnconnectedEventReply = ReturnPin->Direction == EEdGraphPinDirection::EGPD_Input && ReturnPin->LinkedTo.Num() == 0;
+				if ( IsUnconnectedEventReply )
+				{
+					MessageLog.Warning(*LOCTEXT("MissingEventReply_Warning", "Event Reply @@ should not be empty.  Return a reply such as Handled or Unhandled.").ToString(), ReturnPin);
+				}
+			}
+		}
+	}
+}
+
 bool FWidgetBlueprintCompiler::ValidateGeneratedClass(UBlueprintGeneratedClass* Class)
 {
 	bool SuperResult = Super::ValidateGeneratedClass(Class);
@@ -327,3 +358,5 @@ bool FWidgetBlueprintCompiler::ValidateGeneratedClass(UBlueprintGeneratedClass*
 
 	return SuperResult && Result;
 }
+
+#undef LOCTEXT_NAMESPACE
diff --git a/Engine/Source/Editor/UMGEditor/Private/WidgetBlueprintCompiler.h b/Engine/Source/Editor/UMGEditor/Private/WidgetBlueprintCompiler.h
--- a/Engine/Source/Editor/UMGEditor/Private/WidgetBlueprintCompiler.h
+++ b/Engine/Source/Editor/UMGEditor/Private/WidgetBlueprintCompiler.h
@@ -32,6 +32,7 @@ protected:
 	//virtual UEdGraphSchema_K2* CreateSchema() override;
 	virtual void CreateFunctionList() override;
 	virtual void SpawnNewClass(const FString& NewClassName) override;
+	virtual void PrecompileFunction(FKismetFunctionContext& Context) override;
 	virtual void CleanAndSanitizeClass(UBlueprintGeneratedClass* ClassToClean, UObject*& OldCDO) override;
 	virtual void SaveSubObjectsFromCleanAndSanitizeClass(FSubobjectCollection& SubObjectsToSave, UBlueprintGeneratedClass* ClassToClean, UObject*& OldCDO) override;
 	virtual void EnsureProperGeneratedClass(UClass*& TargetClass) override;
@@ -40,6 +41,8 @@ protected:
 	virtual bool ValidateGeneratedClass(UBlueprintGeneratedClass* Class) override;
 	// End FKismetCompilerContext
 
+	void VerifyEventReplysAreNotEmpty(FKismetFunctionContext& Context);
+
 protected:
 	UWidgetBlueprintGeneratedClass* NewWidgetBlueprintClass;
 
