Ensured that everything in the editor using a TTextFilter is reporting syntax error information back to the user

diff --git a/Engine/Plugins/Editor/PluginBrowser/Source/PluginBrowser/Private/SPluginBrowser.cpp b/Engine/Plugins/Editor/PluginBrowser/Source/PluginBrowser/Private/SPluginBrowser.cpp
--- a/Engine/Plugins/Editor/PluginBrowser/Source/PluginBrowser/Private/SPluginBrowser.cpp
+++ b/Engine/Plugins/Editor/PluginBrowser/Source/PluginBrowser/Private/SPluginBrowser.cpp
@@ -95,7 +95,7 @@ void SPluginBrowser::Construct( const FArguments& Args )
 				+SHorizontalBox::Slot()
 				.Padding( PaddingAmount )
 				[
-					SNew( SSearchBox )
+					SAssignNew( SearchBoxPtr, SSearchBox )
 					.OnTextChanged( this, &SPluginBrowser::SearchBox_OnPluginSearchTextChanged )
 				]
 			]
@@ -166,6 +166,7 @@ FReply SPluginBrowser::HandleRestartEditorButtonClicked() const
 void SPluginBrowser::SearchBox_OnPluginSearchTextChanged( const FText& NewText )
 {
 	PluginTextFilter->SetRawFilterText( NewText );
+	SearchBoxPtr->SetError( PluginTextFilter->GetFilterErrorText() );
 }
 
 
diff --git a/Engine/Plugins/Editor/PluginBrowser/Source/PluginBrowser/Private/SPluginBrowser.h b/Engine/Plugins/Editor/PluginBrowser/Source/PluginBrowser/Private/SPluginBrowser.h
--- a/Engine/Plugins/Editor/PluginBrowser/Source/PluginBrowser/Private/SPluginBrowser.h
+++ b/Engine/Plugins/Editor/PluginBrowser/Source/PluginBrowser/Private/SPluginBrowser.h
@@ -76,6 +76,9 @@ private:
 	/** The plugin list widget */
 	TSharedPtr< class SPluginTileList > PluginList;
 
+	/** The plugin search box widget */
+	TSharedPtr< class SSearchBox > SearchBoxPtr;
+
 	/** Text filter object for typing in filter text to the search box */
 	TSharedPtr< FPluginTextFilter > PluginTextFilter;
 
diff --git a/Engine/Source/Developer/AutomationWindow/Private/SAutomationWindow.cpp b/Engine/Source/Developer/AutomationWindow/Private/SAutomationWindow.cpp
--- a/Engine/Source/Developer/AutomationWindow/Private/SAutomationWindow.cpp
+++ b/Engine/Source/Developer/AutomationWindow/Private/SAutomationWindow.cpp
@@ -1566,6 +1566,7 @@ FReply SAutomationWindow::RunTests()
 void SAutomationWindow::OnFilterTextChanged( const FText& InFilterText )
 {
 	AutomationTextFilter->SetRawFilterText( InFilterText );
+	AutomationSearchBox->SetError( AutomationTextFilter->GetFilterErrorText() );
 
 	//update the widget
 	OnRefreshTestCallback();
diff --git a/Engine/Source/Developer/Profiler/Private/Widgets/SFiltersAndPresets.cpp b/Engine/Source/Developer/Profiler/Private/Widgets/SFiltersAndPresets.cpp
--- a/Engine/Source/Developer/Profiler/Private/Widgets/SFiltersAndPresets.cpp
+++ b/Engine/Source/Developer/Profiler/Private/Widgets/SFiltersAndPresets.cpp
@@ -1131,6 +1131,7 @@ bool SFiltersAndPresets::GroupAndStatTableRow_ShouldBeEnabled( const uint32 Stat
 void SFiltersAndPresets::SearchBox_OnTextChanged( const FText& InFilterText )
 {
 	GroupAndStatTextFilter->SetRawFilterText( InFilterText );
+	GroupAndStatSearchBox->SetError( GroupAndStatTextFilter->GetFilterErrorText() );
 	ApplyFiltering();
 }
 
diff --git a/Engine/Source/Editor/AddContentDialog/Private/SAddContentWidget.cpp b/Engine/Source/Editor/AddContentDialog/Private/SAddContentWidget.cpp
--- a/Engine/Source/Editor/AddContentDialog/Private/SAddContentWidget.cpp
+++ b/Engine/Source/Editor/AddContentDialog/Private/SAddContentWidget.cpp
@@ -52,7 +52,7 @@ void SAddContentWidget::Construct(const FArguments& InArgs)
 					.AutoHeight()
 					.Padding(FMargin(0, 0, 0, 5))
 					[
-						SNew(SSearchBox)
+						SAssignNew(SearchBoxPtr, SSearchBox)
 						.OnTextChanged(this, &SAddContentWidget::SearchTextChanged)
 					]
 
@@ -304,6 +304,7 @@ ECheckBoxState SAddContentWidget::GetCategoryCheckBoxCheckState(FCategoryViewMod
 void SAddContentWidget::SearchTextChanged(const FText& SearchText)
 {
 	ViewModel->SetSearchText(SearchText);
+	SearchBoxPtr->SetError(ViewModel->GetSearchErrorText());
 }
 
 void SAddContentWidget::ContentSourceTileViewSelectionChanged( TSharedPtr<FContentSourceViewModel> SelectedContentSource, ESelectInfo::Type SelectInfo )
diff --git a/Engine/Source/Editor/AddContentDialog/Private/SAddContentWidget.h b/Engine/Source/Editor/AddContentDialog/Private/SAddContentWidget.h
--- a/Engine/Source/Editor/AddContentDialog/Private/SAddContentWidget.h
+++ b/Engine/Source/Editor/AddContentDialog/Private/SAddContentWidget.h
@@ -84,4 +84,7 @@ private:
 
 	/** The content sources which the user has selected for addition to the project. */
 	TArray<TSharedPtr<IContentSource>> ContentSourcesToAdd;
+
+	/** The search box */
+	TSharedPtr<class SSearchBox> SearchBoxPtr;
 };
\ No newline at end of file
diff --git a/Engine/Source/Editor/AddContentDialog/Private/ViewModels/AddContentWidgetViewModel.cpp b/Engine/Source/Editor/AddContentDialog/Private/ViewModels/AddContentWidgetViewModel.cpp
--- a/Engine/Source/Editor/AddContentDialog/Private/ViewModels/AddContentWidgetViewModel.cpp
+++ b/Engine/Source/Editor/AddContentDialog/Private/ViewModels/AddContentWidgetViewModel.cpp
@@ -61,6 +61,11 @@ void FAddContentWidgetViewModel::SetSearchText(FText SearchTextIn)
 	UpdateFilteredContentSourcesAndSelection(true);
 }
 
+FText FAddContentWidgetViewModel::GetSearchErrorText() const
+{
+	return ContentSourceFilter->GetFilterErrorText();
+}
+
 const TArray<TSharedPtr<FContentSourceViewModel>>* FAddContentWidgetViewModel::GetContentSources()
 {
 	return &FilteredContentSourceViewModels;
diff --git a/Engine/Source/Editor/AddContentDialog/Private/ViewModels/AddContentWidgetViewModel.h b/Engine/Source/Editor/AddContentDialog/Private/ViewModels/AddContentWidgetViewModel.h
--- a/Engine/Source/Editor/AddContentDialog/Private/ViewModels/AddContentWidgetViewModel.h
+++ b/Engine/Source/Editor/AddContentDialog/Private/ViewModels/AddContentWidgetViewModel.h
@@ -32,6 +32,9 @@ public:
 	/** Sets search text which should be used to filter the content sources. */
 	void SetSearchText(FText SearchTextIn);
 
+	/** Get any error information that was reported the last time SetSearchText was called. */
+	FText GetSearchErrorText() const;
+
 	/** Gets a filtered array of content sources which match both the selected category and the search
 		text if it has been set. */
 	const TArray<TSharedPtr<FContentSourceViewModel>>* GetContentSources();
diff --git a/Engine/Source/Editor/ContentBrowser/Private/SPathView.cpp b/Engine/Source/Editor/ContentBrowser/Private/SPathView.cpp
--- a/Engine/Source/Editor/ContentBrowser/Private/SPathView.cpp
+++ b/Engine/Source/Editor/ContentBrowser/Private/SPathView.cpp
@@ -878,6 +878,7 @@ void SPathView::TreeExpansionChanged( TSharedPtr< FTreeItem > TreeItem, bool bIs
 void SPathView::OnAssetTreeSearchBoxChanged( const FText& InSearchText )
 {
 	SearchBoxFolderFilter->SetRawFilterText( InSearchText );
+	SearchBoxPtr->SetError( SearchBoxFolderFilter->GetFilterErrorText() );
 }
 
 void SPathView::FilterUpdated()
diff --git a/Engine/Source/Editor/FoliageEdit/Private/SFoliagePalette.cpp b/Engine/Source/Editor/FoliageEdit/Private/SFoliagePalette.cpp
--- a/Engine/Source/Editor/FoliageEdit/Private/SFoliagePalette.cpp
+++ b/Engine/Source/Editor/FoliageEdit/Private/SFoliagePalette.cpp
@@ -235,7 +235,7 @@ void SFoliagePalette::Construct(const FArguments& InArgs)
 				.VAlign(VAlign_Center)
 				.Padding(6.f, 0.f)
 				[
-					SNew(SSearchBox)
+					SAssignNew(SearchBoxPtr, SSearchBox)
 					.HintText(LOCTEXT("SearchFoliagePaletteHint", "Search Foliage"))
 					.OnTextChanged(this, &SFoliagePalette::OnSearchTextChanged)
 				]
@@ -641,6 +641,7 @@ void SFoliagePalette::GetPaletteItemFilterString(FFoliagePaletteItemModelPtr Pal
 void SFoliagePalette::OnSearchTextChanged(const FText& InFilterText)
 {
 	TypeFilter->SetRawFilterText(InFilterText);
+	SearchBoxPtr->SetError(TypeFilter->GetFilterErrorText());
 	UpdatePalette();
 }
 
diff --git a/Engine/Source/Editor/FoliageEdit/Private/SFoliagePalette.h b/Engine/Source/Editor/FoliageEdit/Private/SFoliagePalette.h
--- a/Engine/Source/Editor/FoliageEdit/Private/SFoliagePalette.h
+++ b/Engine/Source/Editor/FoliageEdit/Private/SFoliagePalette.h
@@ -253,6 +253,9 @@ private:
 	/** Foliage mesh details widget  */
 	TSharedPtr<class IDetailsView> DetailsWidget;
 
+	/** Foliage items search box widget */
+	TSharedPtr<class SSearchBox> SearchBoxPtr;
+
 	/** Command list for binding functions for the context menu. */
 	TSharedPtr<FUICommandList> UICommandList;
 
diff --git a/Engine/Source/Editor/GameProjectGeneration/Private/SProjectBrowser.cpp b/Engine/Source/Editor/GameProjectGeneration/Private/SProjectBrowser.cpp
--- a/Engine/Source/Editor/GameProjectGeneration/Private/SProjectBrowser.cpp
+++ b/Engine/Source/Editor/GameProjectGeneration/Private/SProjectBrowser.cpp
@@ -156,7 +156,7 @@ void SProjectBrowser::Construct( const FArguments& InArgs )
 
 						+SOverlay::Slot()
 						[
-							SNew(SSearchBox)
+							SAssignNew(SearchBoxPtr, SSearchBox)
 							.HintText(LOCTEXT("FilterHint", "Filter Projects..."))
 							.OnTextChanged(this, &SProjectBrowser::OnFilterTextChanged)
 						]
@@ -1265,6 +1265,7 @@ FReply SProjectBrowser::HandleMarketplaceTabButtonClicked()
 void SProjectBrowser::OnFilterTextChanged(const FText& InText)
 {
 	ProjectItemFilter.SetRawFilterText(InText);
+	SearchBoxPtr->SetError(ProjectItemFilter.GetFilterErrorText());
 	PopulateFilteredProjectCategories();
 }
 
diff --git a/Engine/Source/Editor/GameProjectGeneration/Private/SProjectBrowser.h b/Engine/Source/Editor/GameProjectGeneration/Private/SProjectBrowser.h
--- a/Engine/Source/Editor/GameProjectGeneration/Private/SProjectBrowser.h
+++ b/Engine/Source/Editor/GameProjectGeneration/Private/SProjectBrowser.h
@@ -132,6 +132,9 @@ private:
 	// Holds the collection of project categories.
 	TArray<TSharedRef<FProjectCategory> > ProjectCategories;
 
+	/** Search box used to set the filter text */
+	TSharedPtr<class SSearchBox> SearchBoxPtr;
+
 	/** Filter that is used to test for the visibility of projects */
 	typedef TTextFilter<const TSharedPtr<FProjectItem>> ProjectItemTextFilter;
 	ProjectItemTextFilter ProjectItemFilter;
diff --git a/Engine/Source/Editor/GameplayAbilitiesEditor/Private/SGameplayAttributeWidget.cpp b/Engine/Source/Editor/GameplayAbilitiesEditor/Private/SGameplayAttributeWidget.cpp
--- a/Engine/Source/Editor/GameplayAbilitiesEditor/Private/SGameplayAttributeWidget.cpp
+++ b/Engine/Source/Editor/GameplayAbilitiesEditor/Private/SGameplayAttributeWidget.cpp
@@ -321,6 +321,7 @@ TSharedPtr<FAttributeViewerNode> SAttributeListWidget::UpdatePropertyOptions()
 void SAttributeListWidget::OnFilterTextChanged(const FText& InFilterText)
 {
 	AttributeTextFilter->SetRawFilterText(InFilterText);
+	SearchBoxPtr->SetError(AttributeTextFilter->GetFilterErrorText());
 
 	UpdatePropertyOptions();
 }
diff --git a/Engine/Source/Editor/Layers/Private/SLayerBrowser.cpp b/Engine/Source/Editor/Layers/Private/SLayerBrowser.cpp
--- a/Engine/Source/Editor/Layers/Private/SLayerBrowser.cpp
+++ b/Engine/Source/Editor/Layers/Private/SLayerBrowser.cpp
@@ -36,10 +36,10 @@ void SLayerBrowser::Construct(const FArguments& InArgs)
 			+ SVerticalBox::Slot()
 			.AutoHeight()
 			[
-				SNew(SSearchBox)
+				SAssignNew(SearchBoxPtr, SSearchBox)
 				.ToolTipText(LOCTEXT("FilterSearchToolTip", "Type here to search layers"))
 				.HintText(LOCTEXT("FilterSearchHint", "Search Layers"))
-				.OnTextChanged(SearchBoxLayerFilter.Get(), &LayerTextFilter::SetRawFilterText)
+				.OnTextChanged(this, &SLayerBrowser::OnFilterTextChanged)
 			]
 
 			+ SVerticalBox::Slot()
@@ -166,4 +166,10 @@ FText SLayerBrowser::GetLayerContentsHeaderText() const
 	return FText::Format(LOCTEXT("SelectedContentsLabel", "{0} Contents"), FText::FromString(SelectedLayerViewModel->GetName()));
 }
 
+void SLayerBrowser::OnFilterTextChanged( const FText& InNewText )
+{
+	SearchBoxLayerFilter->SetRawFilterText(InNewText);
+	SearchBoxPtr->SetError(SearchBoxLayerFilter->GetFilterErrorText());
+}
+
 #undef LOCTEXT_NAMESPACE
diff --git a/Engine/Source/Editor/Layers/Private/SLayerBrowser.h b/Engine/Source/Editor/Layers/Private/SLayerBrowser.h
--- a/Engine/Source/Editor/Layers/Private/SLayerBrowser.h
+++ b/Engine/Source/Editor/Layers/Private/SLayerBrowser.h
@@ -329,6 +329,7 @@ private:
 		return SNew( SLayersCommandsMenu, LayerCollectionViewModel.ToSharedRef() );
 	}
 
+	void OnFilterTextChanged( const FText& InNewText );
 
 private:
 
@@ -347,6 +348,9 @@ private:
 	/**	 */
 	TSharedPtr< SBorder > LayerContentsHeader;
 
+	/** */
+	TSharedPtr< class SSearchBox> SearchBoxPtr;
+
 	/**	 */
 	TSharedPtr< LayerTextFilter > SearchBoxLayerFilter;
 
diff --git a/Engine/Source/Editor/SceneOutliner/Private/SSceneOutliner.cpp b/Engine/Source/Editor/SceneOutliner/Private/SSceneOutliner.cpp
--- a/Engine/Source/Editor/SceneOutliner/Private/SSceneOutliner.cpp
+++ b/Engine/Source/Editor/SceneOutliner/Private/SSceneOutliner.cpp
@@ -1928,6 +1928,7 @@ namespace SceneOutliner
 	void SSceneOutliner::OnFilterTextChanged( const FText& InFilterText )
 	{
 		SearchBoxFilter->SetRawFilterText( InFilterText );
+		FilterTextBoxWidget->SetError( SearchBoxFilter->GetFilterErrorText() );
 	}
 
 	void SSceneOutliner::OnFilterTextCommitted( const FText& InFilterText, ETextCommit::Type CommitInfo )
diff --git a/Engine/Source/Editor/UMGEditor/Private/Hierarchy/SHierarchyView.cpp b/Engine/Source/Editor/UMGEditor/Private/Hierarchy/SHierarchyView.cpp
--- a/Engine/Source/Editor/UMGEditor/Private/Hierarchy/SHierarchyView.cpp
+++ b/Engine/Source/Editor/UMGEditor/Private/Hierarchy/SHierarchyView.cpp
@@ -63,7 +63,7 @@ void SHierarchyView::Construct(const FArguments& InArgs, TSharedPtr<FWidgetBluep
 			.Padding(4)
 			.AutoHeight()
 			[
-				SNew(SSearchBox)
+				SAssignNew(SearchBoxPtr, SSearchBox)
 				.HintText(LOCTEXT("SearchWidgets", "Search Widgets"))
 				.OnTextChanged(this, &SHierarchyView::OnSearchChanged)
 			]
@@ -175,6 +175,7 @@ void SHierarchyView::OnSearchChanged(const FText& InFilterText)
 	bRefreshRequested = true;
 	FilterHandler->SetIsEnabled(!InFilterText.IsEmpty());
 	SearchBoxWidgetFilter->SetRawFilterText(InFilterText);
+	SearchBoxPtr->SetError(SearchBoxWidgetFilter->GetFilterErrorText());
 }
 
 FText SHierarchyView::GetSearchText() const
diff --git a/Engine/Source/Editor/UMGEditor/Private/Hierarchy/SHierarchyView.h b/Engine/Source/Editor/UMGEditor/Private/Hierarchy/SHierarchyView.h
--- a/Engine/Source/Editor/UMGEditor/Private/Hierarchy/SHierarchyView.h
+++ b/Engine/Source/Editor/UMGEditor/Private/Hierarchy/SHierarchyView.h
@@ -123,6 +123,9 @@ private:
 	/** The widget hierarchy slate treeview widget */
 	TSharedPtr< STreeView< TSharedPtr<FHierarchyModel> > > WidgetTreeView;
 
+	/** The search box used to update the filter text */
+	TSharedPtr<class SSearchBox> SearchBoxPtr;
+
 	/** The filter used by the search box */
 	TSharedPtr<WidgetTextFilter> SearchBoxWidgetFilter;
 
diff --git a/Engine/Source/Editor/UMGEditor/Private/Palette/SPaletteView.cpp b/Engine/Source/Editor/UMGEditor/Private/Palette/SPaletteView.cpp
--- a/Engine/Source/Editor/UMGEditor/Private/Palette/SPaletteView.cpp
+++ b/Engine/Source/Editor/UMGEditor/Private/Palette/SPaletteView.cpp
@@ -210,7 +210,7 @@ void SPaletteView::Construct(const FArguments& InArgs, TSharedPtr<FBlueprintEdit
 		.Padding(4)
 		.AutoHeight()
 		[
-			SNew(SSearchBox)
+			SAssignNew(SearchBoxPtr, SSearchBox)
 			.HintText(LOCTEXT("SearchTemplates", "Search Palette"))
 			.OnTextChanged(this, &SPaletteView::OnSearchChanged)
 		]
@@ -257,6 +257,7 @@ void SPaletteView::OnSearchChanged(const FText& InFilterText)
 	bRefreshRequested = true;
 	FilterHandler->SetIsEnabled(!InFilterText.IsEmpty());
 	WidgetFilter->SetRawFilterText(InFilterText);
+	SearchBoxPtr->SetError(WidgetFilter->GetFilterErrorText());
 	SearchText = InFilterText;
 }
 
diff --git a/Engine/Source/Editor/UMGEditor/Private/Palette/SPaletteView.h b/Engine/Source/Editor/UMGEditor/Private/Palette/SPaletteView.h
--- a/Engine/Source/Editor/UMGEditor/Private/Palette/SPaletteView.h
+++ b/Engine/Source/Editor/UMGEditor/Private/Palette/SPaletteView.h
@@ -96,6 +96,9 @@ private:
 
 	TSharedPtr< STreeView< TSharedPtr<FWidgetViewModel> > > WidgetTemplatesView;
 
+	/** The search box used to update the filter text */
+	TSharedPtr<class SSearchBox> SearchBoxPtr;
+
 	/** The filter instance which is used by the TreeFilterHandler to filter the TreeView. */
 	TSharedPtr<WidgetViewModelTextFilter> WidgetFilter;
 
diff --git a/Engine/Source/Editor/UMGEditor/Private/TabFactory/AnimationTabSummoner.cpp b/Engine/Source/Editor/UMGEditor/Private/TabFactory/AnimationTabSummoner.cpp
--- a/Engine/Source/Editor/UMGEditor/Private/TabFactory/AnimationTabSummoner.cpp
+++ b/Engine/Source/Editor/UMGEditor/Private/TabFactory/AnimationTabSummoner.cpp
@@ -239,7 +239,7 @@ public:
 					.Padding(2.0f, 0.0f)
 					.VAlign( VAlign_Center )
 					[
-						SNew(SSearchBox)
+						SAssignNew(SearchBoxPtr, SSearchBox)
 						.HintText(LOCTEXT("Search Animations", "Search Animations"))
 						.OnTextChanged(this, &SUMGAnimationList::OnSearchChanged)
 					]
@@ -347,6 +347,7 @@ private:
 			TTextFilter<UWidgetAnimation*> TextFilter(TTextFilter<UWidgetAnimation*>::FItemToStringArray::CreateStatic(&Local::UpdateFilterStrings));
 
 			TextFilter.SetRawFilterText(InSearchText);
+			SearchBoxPtr->SetError(TextFilter.GetFilterErrorText());
 
 			Animations.Reset();
 
@@ -362,6 +363,8 @@ private:
 		}
 		else
 		{
+			SearchBoxPtr->SetError(FText::GetEmpty());
+
 			// Just regenerate the whole list
 			UpdateAnimationList();
 		}
@@ -504,6 +507,7 @@ private:
 	TWeakPtr<FWidgetBlueprintEditor> BlueprintEditor;
 	TSharedPtr<SWidgetAnimationListView> AnimationListView;
 	TArray< TSharedPtr<FWidgetAnimationListItem> > Animations;
+	TSharedPtr<SSearchBox> SearchBoxPtr;
 };
 
 
