Sequencer - Fix shadowed variable warnings.

diff --git a/Engine/Source/Developer/MovieSceneTools/Private/PropertyTrackEditor.cpp b/Engine/Source/Developer/MovieSceneTools/Private/PropertyTrackEditor.cpp
--- a/Engine/Source/Developer/MovieSceneTools/Private/PropertyTrackEditor.cpp
+++ b/Engine/Source/Developer/MovieSceneTools/Private/PropertyTrackEditor.cpp
@@ -278,9 +278,9 @@ public:
 class FBytePropertySection : public FPropertySection
 {
 public:
-	FBytePropertySection(UMovieSceneSection& InSectionObject, FName SectionName, UEnum* Enum)
-		: FPropertySection(InSectionObject, SectionName)
-		, Enum(Enum) {}
+	FBytePropertySection(UMovieSceneSection& InSectionObject, FName InSectionName, UEnum* InEnum)
+		: FPropertySection(InSectionObject, InSectionName)
+		, Enum(InEnum) {}
 
 	virtual void GenerateSectionLayout(class ISectionLayoutBuilder& LayoutBuilder) const override
 	{
diff --git a/Engine/Source/Editor/Sequencer/Private/KeyPropertyParams.cpp b/Engine/Source/Editor/Sequencer/Private/KeyPropertyParams.cpp
--- a/Engine/Source/Editor/Sequencer/Private/KeyPropertyParams.cpp
+++ b/Engine/Source/Editor/Sequencer/Private/KeyPropertyParams.cpp
@@ -16,7 +16,7 @@ void PropertyHandleToPropertyPath(const UClass* OwnerClass, const IPropertyHandl
 
 FCanKeyPropertyParams::FCanKeyPropertyParams()
 {
-	UClass* ObjectClass = nullptr;
+	ObjectClass = nullptr;
 }
 
 FCanKeyPropertyParams::FCanKeyPropertyParams(UClass* InObjectClass, TArray<UProperty*> InPropertyPath)
diff --git a/Engine/Source/Editor/Sequencer/Private/SequencerDisplayNode.cpp b/Engine/Source/Editor/Sequencer/Private/SequencerDisplayNode.cpp
--- a/Engine/Source/Editor/Sequencer/Private/SequencerDisplayNode.cpp
+++ b/Engine/Source/Editor/Sequencer/Private/SequencerDisplayNode.cpp
@@ -345,12 +345,12 @@ bool FSectionKeyAreaNode::GetShotFilteredVisibilityToCache() const
 TSharedRef<SWidget> FSectionKeyAreaNode::GenerateEditWidgetForOutliner()
 {
 	// @todo - Sequencer - Support multiple sections/key areas?
-	TArray<TSharedRef<IKeyArea>> KeyAreas = GetAllKeyAreas();
-	if (KeyAreas.Num() > 0)
+	TArray<TSharedRef<IKeyArea>> AllKeyAreas = GetAllKeyAreas();
+	if (AllKeyAreas.Num() > 0)
 	{
-		if (KeyAreas[0]->CanCreateKeyEditor())
+		if (AllKeyAreas[0]->CanCreateKeyEditor())
 		{
-			return KeyAreas[0]->CreateKeyEditor(&GetSequencer());
+			return AllKeyAreas[0]->CreateKeyEditor(&GetSequencer());
 		}
 	}
 	return FSequencerDisplayNode::GenerateEditWidgetForOutliner();
diff --git a/Engine/Source/Editor/UMGEditor/Private/WidgetBlueprintEditor.cpp b/Engine/Source/Editor/UMGEditor/Private/WidgetBlueprintEditor.cpp
--- a/Engine/Source/Editor/UMGEditor/Private/WidgetBlueprintEditor.cpp
+++ b/Engine/Source/Editor/UMGEditor/Private/WidgetBlueprintEditor.cpp
@@ -760,7 +760,7 @@ void GetBindableObjects(UWidget* RootWidget, TArray<FObjectAndDisplayName>& Bind
 	}
 }
 
-TSharedRef<SWidget> FWidgetBlueprintEditor::OnGetAnimationAddMenuContent(TSharedRef<ISequencer> Sequencer)
+TSharedRef<SWidget> FWidgetBlueprintEditor::OnGetAnimationAddMenuContent(TSharedRef<ISequencer> InSequencer)
 {
 	TArray<FObjectAndDisplayName> BindableObjects;
 	GetBindableObjects(GetPreview()->GetRootWidget(), BindableObjects);
@@ -769,7 +769,7 @@ TSharedRef<SWidget> FWidgetBlueprintEditor::OnGetAnimationAddMenuContent(TShared
 	FMenuBuilder AddMenuBuilder(true, NULL);
 	for (FObjectAndDisplayName& BindableObject : BindableObjects)
 	{
-		FGuid BoundObjectGuid = SequencerObjectBindingManager->FindGuidForObject(*Sequencer->GetFocusedMovieScene(), *BindableObject.Object);
+		FGuid BoundObjectGuid = SequencerObjectBindingManager->FindGuidForObject(*InSequencer->GetFocusedMovieScene(), *BindableObject.Object);
 		if (BoundObjectGuid.IsValid() == false)
 		{
 			FUIAction AddMenuAction(FExecuteAction::CreateSP(this, &FWidgetBlueprintEditor::AddObjectToAnimation, BindableObject.Object));
