Widgets removed by replacing them in the UMG editor hierarchy are moved to the transient package to stop name clashes.
[UE-16568] - Replacing a Widget and then reusing it's name for another widget will cause a crash

#change Renaming the replaced widget into the transient package in FWidgetBlueprintEditorUtils::ReplaceWidgets(). Approach copied from DeleteWidgets().
#change Add extra defensive check to FWidgetBlueprintEditorUtils::VerifyWidgetRename() to make sure that the rename won't cause a crash by running a dummy rename with the REN_Test flag. Performance feels fine while editing the name in the editable text box.

diff --git a/Engine/Source/Editor/UMGEditor/Private/WidgetBlueprintEditorUtils.cpp b/Engine/Source/Editor/UMGEditor/Private/WidgetBlueprintEditorUtils.cpp
--- a/Engine/Source/Editor/UMGEditor/Private/WidgetBlueprintEditorUtils.cpp
+++ b/Engine/Source/Editor/UMGEditor/Private/WidgetBlueprintEditorUtils.cpp
@@ -101,6 +101,30 @@ bool FWidgetBlueprintEditorUtils::VerifyWidgetRename(TSharedRef<class FWidgetBlu
 			bIsSameWidget = true;
 		}
 	}
+	else
+	{
+		// Not an existing widget in the tree BUT it still mustn't create a UObject name clash
+		UWidget* WidgetPreview = Widget.GetPreview();
+		if (WidgetPreview)
+		{
+			// Dummy rename with flag REN_Test returns if rename is possible
+			if (!WidgetPreview->Rename(*NewNameString, nullptr, REN_Test))
+			{
+				OutErrorMessage = LOCTEXT("ExistingObjectName", "Existing Object Name");
+				return false;
+			}
+		}
+		UWidget* WidgetTemplate = Widget.GetTemplate();
+		if (WidgetTemplate)
+		{
+			// Dummy rename with flag REN_Test returns if rename is possible
+			if (!WidgetTemplate->Rename(*NewNameString, nullptr, REN_Test))
+			{
+				OutErrorMessage = LOCTEXT("ExistingObjectName", "Existing Object Name");
+				return false;
+			}
+		}
+	}
 
 	FKismetNameValidator Validator(Blueprint);
 
@@ -489,6 +513,9 @@ void FWidgetBlueprintEditorUtils::ReplaceWidgets(UWidgetBlueprint* BP, TSet<FWid
 
 				NewReplacementWidget->AddChild(Widget);
 			}
+
+			// Rename the removed widget to the transient package so that it doesn't conflict with future widgets sharing the same name.
+			ExistingPanel->Rename(nullptr, nullptr);
 		}
 	}
 
