Enabled dragging of mutliple widgets to reparent in the UMG hierarchy
[UE-13345] - Unable to move multiple UMG widgets hierarchy at once

diff --git a/Engine/Source/Editor/UMGEditor/Private/Hierarchy/SHierarchyViewItem.cpp b/Engine/Source/Editor/UMGEditor/Private/Hierarchy/SHierarchyViewItem.cpp
--- a/Engine/Source/Editor/UMGEditor/Private/Hierarchy/SHierarchyViewItem.cpp
+++ b/Engine/Source/Editor/UMGEditor/Private/Hierarchy/SHierarchyViewItem.cpp
@@ -36,50 +36,73 @@ class FHierarchyWidgetDragDropOp : public FDecoratedDragDropOp
 public:
 	DRAG_DROP_OPERATOR_TYPE(FHierarchyWidgetDragDropOp, FDecoratedDragDropOp)
 
-		virtual ~FHierarchyWidgetDragDropOp();
+	virtual ~FHierarchyWidgetDragDropOp();
 
 	virtual void OnDrop(bool bDropWasHandled, const FPointerEvent& MouseEvent) override;
 
-	/** The slot properties for the old slot the widget was in, is used to attempt to reapply the same layout information */
-	TMap<FName, FString> ExportedSlotProperties;
+	struct FItem
+	{
+		/** The slot properties for the old slot the widget was in, is used to attempt to reapply the same layout information */
+		TMap<FName, FString> ExportedSlotProperties;
 
-	/** The widget being dragged and dropped */
-	FWidgetReference Widget;
+		/** The widget being dragged and dropped */
+		FWidgetReference Widget;
 
-	/** The original parent of the widget. */
-	UWidget* WidgetParent;
+		/** The original parent of the widget. */
+		UWidget* WidgetParent;
+	};
+
+	TArray<FItem> DraggedWidgets;
 
 	/** The widget being dragged and dropped */
 	FScopedTransaction* Transaction;
 
 	/** Constructs a new drag/drop operation */
-	static TSharedRef<FHierarchyWidgetDragDropOp> New(UWidgetBlueprint* Blueprint, FWidgetReference InWidget);
+	static TSharedRef<FHierarchyWidgetDragDropOp> New(UWidgetBlueprint* Blueprint, const TArray<FWidgetReference>& InWidgets);
 };
 
-TSharedRef<FHierarchyWidgetDragDropOp> FHierarchyWidgetDragDropOp::New(UWidgetBlueprint* Blueprint, FWidgetReference InWidget)
+TSharedRef<FHierarchyWidgetDragDropOp> FHierarchyWidgetDragDropOp::New(UWidgetBlueprint* Blueprint, const TArray<FWidgetReference>& InWidgets)
 {
+	check(InWidgets.Num() > 0);
+
 	TSharedRef<FHierarchyWidgetDragDropOp> Operation = MakeShareable(new FHierarchyWidgetDragDropOp());
-	Operation->Widget = InWidget;
-	Operation->DefaultHoverText = InWidget.GetTemplate()->GetLabelText();
-	Operation->CurrentHoverText = InWidget.GetTemplate()->GetLabelText();
-	Operation->Construct();
 
-	FWidgetBlueprintEditorUtils::ExportPropertiesToText(InWidget.GetTemplate()->Slot, Operation->ExportedSlotProperties);
+	// Add an FItem for each widget in the drag operation
+	for (const auto& Widget : InWidgets)
+	{
+		FItem DraggedWidget;
 
-	Operation->Transaction = new FScopedTransaction(LOCTEXT("Designer_MoveWidget", "Move Widget"));
+		DraggedWidget.Widget = Widget;
 
-	Blueprint->WidgetTree->SetFlags(RF_Transactional);
-	Blueprint->WidgetTree->Modify();
+		FWidgetBlueprintEditorUtils::ExportPropertiesToText(Widget.GetTemplate()->Slot, DraggedWidget.ExportedSlotProperties);
 
-	UWidget* Widget = Operation->Widget.GetTemplate();
-	Widget->Modify();
+		UWidget* WidgetTemplate = Widget.GetTemplate();
+		WidgetTemplate->Modify();
 
-	Operation->WidgetParent = Widget->GetParent();
+		DraggedWidget.WidgetParent = WidgetTemplate->GetParent();
+		if (DraggedWidget.WidgetParent)
+		{
+			DraggedWidget.WidgetParent->Modify();
+		}
 
-	if ( Operation->WidgetParent )
+		Operation->DraggedWidgets.Add(DraggedWidget);
+	}
+
+	// Set the display text and the transaction name based on whether we're dragging a single or multiple widgets
+	if (InWidgets.Num() == 1)
 	{
-		Operation->WidgetParent->Modify();
+		Operation->CurrentHoverText = Operation->DefaultHoverText = InWidgets[0].GetTemplate()->GetLabelText();
+		Operation->Transaction = new FScopedTransaction(LOCTEXT("Designer_MoveWidget", "Move Widget"));
 	}
+	else
+	{
+		Operation->CurrentHoverText = Operation->DefaultHoverText = LOCTEXT("Designer_DragMultipleWidgets", "Multiple Widgets");
+		Operation->Transaction = new FScopedTransaction(LOCTEXT("Designer_MoveWidgets", "Move Widgets"));
+	}
+	Operation->Construct();
+	
+	Blueprint->WidgetTree->SetFlags(RF_Transactional);
+	Blueprint->WidgetTree->Modify();
 
 	return Operation;
 }
@@ -199,7 +222,11 @@ TOptional<EItemDropZone> ProcessHierarchyDragDrop(const FDragDropEvent& DragDrop
 		// the null case and we should be adding it as the root widget.
 		if ( TargetItem.IsValid() )
 		{
-			const bool bIsDraggedObject = TargetItem.GetTemplate() == HierarchyDragDropOp->Widget.GetTemplate();
+			const bool bIsDraggedObject = HierarchyDragDropOp->DraggedWidgets.ContainsByPredicate([TargetItem](const FHierarchyWidgetDragDropOp::FItem& DraggedItem)
+			{
+				return DraggedItem.Widget == TargetItem;
+			});
+
 			if ( bIsDraggedObject )
 			{
 				HierarchyDragDropOp->SetCursorOverride(EMouseCursor::SlashedCircle);
@@ -221,16 +248,19 @@ TOptional<EItemDropZone> ProcessHierarchyDragDrop(const FDragDropEvent& DragDrop
 				return TOptional<EItemDropZone>();
 			}
 
-			UWidget* TemplateWidget = HierarchyDragDropOp->Widget.GetTemplate();
-
-			// Verify that the new location we're placing the widget is not inside of its existing children.
 			bool bFoundNewParentInChildSet = false;
-			Blueprint->WidgetTree->ForWidgetAndChildren(TemplateWidget, [&] (UWidget* Widget) {
-				if ( NewParent == Widget )
-				{
-					bFoundNewParentInChildSet = true;
-				}
-			});
+			for (const auto& DraggedWidget : HierarchyDragDropOp->DraggedWidgets)
+			{
+				UWidget* TemplateWidget = DraggedWidget.Widget.GetTemplate();
+
+				// Verify that the new location we're placing the widget is not inside of its existing children.
+				Blueprint->WidgetTree->ForWidgetAndChildren(TemplateWidget, [&](UWidget* Widget) {
+					if (NewParent == Widget)
+					{
+						bFoundNewParentInChildSet = true;
+					}
+				});
+			}
 
 			if ( bFoundNewParentInChildSet )
 			{
@@ -241,43 +271,50 @@ TOptional<EItemDropZone> ProcessHierarchyDragDrop(const FDragDropEvent& DragDrop
 
 			if ( bIsDrop )
 			{
-				if ( Index.IsSet() )
+				NewParent->SetFlags(RF_Transactional);
+				NewParent->Modify();
+
+				TSet<FWidgetReference> SelectedTemplates;
+
+				for (const auto& DraggedWidget : HierarchyDragDropOp->DraggedWidgets)
 				{
-					// If we're inserting at an index, and the widget we're moving is already
-					// in the hierarchy before the point we're moving it to, we need to reduce the index
-					// count by one, because the whole set is about to be shifted when it's removed.
-					const bool bInsertInSameParent = TemplateWidget->GetParent() == NewParent;
-					const bool bNeedToDropIndex = NewParent->GetChildIndex(TemplateWidget) < Index.GetValue();
+					UWidget* TemplateWidget = DraggedWidget.Widget.GetTemplate();
 
-					if ( bInsertInSameParent && bNeedToDropIndex )
+					if (Index.IsSet())
 					{
-						Index = Index.GetValue() - 1;
+						// If we're inserting at an index, and the widget we're moving is already
+						// in the hierarchy before the point we're moving it to, we need to reduce the index
+						// count by one, because the whole set is about to be shifted when it's removed.
+						const bool bInsertInSameParent = TemplateWidget->GetParent() == NewParent;
+						const bool bNeedToDropIndex = NewParent->GetChildIndex(TemplateWidget) < Index.GetValue();
+
+						if (bInsertInSameParent && bNeedToDropIndex)
+						{
+							Index = Index.GetValue() - 1;
+						}
 					}
-				}
 
-				TemplateWidget->RemoveFromParent();
+					TemplateWidget->RemoveFromParent();
 
-				NewParent->SetFlags(RF_Transactional);
-				NewParent->Modify();
+					UPanelSlot* NewSlot = nullptr;
+					if (Index.IsSet())
+					{
+						NewSlot = NewParent->InsertChildAt(Index.GetValue(), TemplateWidget);
+						Index = Index.GetValue() + 1;
+					}
+					else
+					{
+						NewSlot = NewParent->AddChild(TemplateWidget);
+					}
+					check(NewSlot);
 
-				UPanelSlot* NewSlot = nullptr;
-				if ( Index.IsSet() )
-				{
-					NewSlot = NewParent->InsertChildAt(Index.GetValue(), TemplateWidget);
-				}
-				else
-				{
-					NewSlot = NewParent->AddChild(TemplateWidget);
+					// Import the old slot properties
+					FWidgetBlueprintEditorUtils::ImportPropertiesFromText(NewSlot, DraggedWidget.ExportedSlotProperties);
+
+					SelectedTemplates.Add(BlueprintEditor->GetReferenceFromTemplate(TemplateWidget));
 				}
-				check(NewSlot);
 
-				// Import the old slot properties
-				FWidgetBlueprintEditorUtils::ImportPropertiesFromText(NewSlot, HierarchyDragDropOp->ExportedSlotProperties);
 				FBlueprintEditorUtils::MarkBlueprintAsStructurallyModified(Blueprint);
-
-				TSet<FWidgetReference> SelectedTemplates;
-				SelectedTemplates.Add(BlueprintEditor->GetReferenceFromTemplate(TemplateWidget));
-
 				BlueprintEditor->SelectWidgets(SelectedTemplates, false);
 			}
 
@@ -297,9 +334,10 @@ TOptional<EItemDropZone> ProcessHierarchyDragDrop(const FDragDropEvent& DragDrop
 
 //////////////////////////////////////////////////////////////////////////
 
-FHierarchyModel::FHierarchyModel()
+FHierarchyModel::FHierarchyModel(TSharedPtr<FWidgetBlueprintEditor> InBlueprintEditor)
 	: bInitialized(false)
 	, bIsSelected(false)
+	, BlueprintEditor(InBlueprintEditor)
 {
 
 }
@@ -311,6 +349,38 @@ TOptional<EItemDropZone> FHierarchyModel::HandleCanAcceptDrop(const FDragDropEve
 
 FReply FHierarchyModel::HandleDragDetected(const FGeometry& MyGeometry, const FPointerEvent& MouseEvent)
 {
+	if (!IsRoot())
+	{
+		TArray<FWidgetReference> DraggedItems;
+
+		// Dragging multiple items?
+		if (bIsSelected)
+		{
+			const TSet<FWidgetReference>& SelectedWidgets = BlueprintEditor.Pin()->GetSelectedWidgets();
+			if (SelectedWidgets.Num() > 1)
+			{
+				for (const auto& SelectedWidget : SelectedWidgets)
+				{
+					DraggedItems.Add(SelectedWidget);
+				}
+			}
+		}
+
+		if (DraggedItems.Num() == 0)
+		{
+			FWidgetReference ThisItem = AsDraggedWidgetReference();
+			if (ThisItem.IsValid())
+			{
+				DraggedItems.Add(ThisItem);
+			}
+		}
+
+		if (DraggedItems.Num() > 0)
+		{
+			return FReply::Handled().BeginDragDrop(FHierarchyWidgetDragDropOp::New(BlueprintEditor.Pin()->GetWidgetBlueprintObj(), DraggedItems));
+		}
+	}
+
 	return FReply::Unhandled();
 }
 
@@ -395,7 +465,7 @@ bool FHierarchyModel::IsSelected() const
 //////////////////////////////////////////////////////////////////////////
 
 FHierarchyRoot::FHierarchyRoot(TSharedPtr<FWidgetBlueprintEditor> InBlueprintEditor)
-	: BlueprintEditor(InBlueprintEditor)
+	: FHierarchyModel(InBlueprintEditor)
 {
 	RootText = FText::Format(LOCTEXT("RootWidgetFormat", "[{0}]"), FText::FromString(BlueprintEditor.Pin()->GetBlueprintObj()->GetName()));
 }
@@ -474,9 +544,9 @@ FReply FHierarchyRoot::HandleAcceptDrop(FDragDropEvent const& DragDropEvent, EIt
 //////////////////////////////////////////////////////////////////////////
 
 FNamedSlotModel::FNamedSlotModel(FWidgetReference InItem, FName InSlotName, TSharedPtr<FWidgetBlueprintEditor> InBlueprintEditor)
-	: Item(InItem)
+	: FHierarchyModel(InBlueprintEditor)
+	, Item(InItem)
 	, SlotName(InSlotName)
-	, BlueprintEditor(InBlueprintEditor)
 {
 }
 
@@ -561,6 +631,20 @@ void FNamedSlotModel::UpdateSelection()
 	//}
 }
 
+FWidgetReference FNamedSlotModel::AsDraggedWidgetReference() const
+{
+	if (INamedSlotInterface* NamedSlotHost = Cast<INamedSlotInterface>(Item.GetTemplate()))
+	{
+		// Only assign content to the named slot if it is null.
+		if (UWidget* Content = NamedSlotHost->GetContentForSlot(SlotName))
+		{
+			return BlueprintEditor.Pin()->GetReferenceFromTemplate(Content);
+		}
+	}
+
+	return FWidgetReference();
+}
+
 TOptional<EItemDropZone> FNamedSlotModel::HandleCanAcceptDrop(const FDragDropEvent& DragDropEvent, EItemDropZone DropZone)
 {
 	TSharedPtr<FWidgetTemplateDragDropOp> TemplateDragDropOp = DragDropEvent.GetOperationAs<FWidgetTemplateDragDropOp>();
@@ -614,28 +698,11 @@ FReply FNamedSlotModel::HandleAcceptDrop(FDragDropEvent const& DragDropEvent, EI
 	return FReply::Unhandled();
 }
 
-FReply FNamedSlotModel::HandleDragDetected(const FGeometry& MyGeometry, const FPointerEvent& MouseEvent)
-{
-	if ( INamedSlotInterface* NamedSlotHost = Cast<INamedSlotInterface>(Item.GetTemplate()) )
-	{
-		// Only assign content to the named slot if it is null.
-		if ( UWidget* Content = NamedSlotHost->GetContentForSlot(SlotName) )
-		{
-			FWidgetReference ContentRef = BlueprintEditor.Pin()->GetReferenceFromTemplate(Content);
-			check(ContentRef.IsValid());
-
-			return FReply::Handled().BeginDragDrop(FHierarchyWidgetDragDropOp::New(BlueprintEditor.Pin()->GetWidgetBlueprintObj(), ContentRef));
-		}
-	}
-
-	return FReply::Unhandled();
-}
-
 //////////////////////////////////////////////////////////////////////////
 
 FHierarchyWidget::FHierarchyWidget(FWidgetReference InItem, TSharedPtr<FWidgetBlueprintEditor> InBlueprintEditor)
-	: Item(InItem)
-	, BlueprintEditor(InBlueprintEditor)
+	: FHierarchyModel(InBlueprintEditor)
+	, Item(InItem)
 {
 }
 
@@ -721,18 +788,6 @@ TOptional<EItemDropZone> FHierarchyWidget::HandleCanAcceptDrop(const FDragDropEv
 	return ProcessHierarchyDragDrop(DragDropEvent, DropZone, bIsDrop, BlueprintEditor.Pin(), Item);
 }
 
-FReply FHierarchyWidget::HandleDragDetected(const FGeometry& MyGeometry, const FPointerEvent& MouseEvent)
-{
-	const bool bIsRoot = Item.GetTemplate()->GetParent() == nullptr;
-
-	if ( !bIsRoot )
-	{
-		return FReply::Handled().BeginDragDrop(FHierarchyWidgetDragDropOp::New(BlueprintEditor.Pin()->GetWidgetBlueprintObj(), Item));
-	}
-
-	return FReply::Unhandled();
-}
-
 void FHierarchyWidget::HandleDragLeave(const FDragDropEvent& DragDropEvent)
 {
 	TSharedPtr<FDecoratedDragDropOp> DecoratedDragDropOp = DragDropEvent.GetOperationAs<FDecoratedDragDropOp>();
diff --git a/Engine/Source/Editor/UMGEditor/Private/Hierarchy/SHierarchyViewItem.h b/Engine/Source/Editor/UMGEditor/Private/Hierarchy/SHierarchyViewItem.h
--- a/Engine/Source/Editor/UMGEditor/Private/Hierarchy/SHierarchyViewItem.h
+++ b/Engine/Source/Editor/UMGEditor/Private/Hierarchy/SHierarchyViewItem.h
@@ -14,7 +14,7 @@ class FWidgetBlueprintEditor;
 class FHierarchyModel : public TSharedFromThis < FHierarchyModel >
 {
 public:
-	FHierarchyModel();
+	FHierarchyModel(TSharedPtr<FWidgetBlueprintEditor> InBlueprintEditor);
 
 	/** Gets the unique name of the item used to restore item expansion. */
 	virtual FName GetUniqueName() const = 0;
@@ -32,7 +32,7 @@ public:
 
 	virtual FSlateFontInfo GetFont() const = 0;
 
-	virtual FReply HandleDragDetected(const FGeometry& MyGeometry, const FPointerEvent& MouseEvent);
+	FReply HandleDragDetected(const FGeometry& MyGeometry, const FPointerEvent& MouseEvent);
 	virtual void HandleDragEnter(const FDragDropEvent& DragDropEvent);
 	virtual void HandleDragLeave(const FDragDropEvent& DragDropEvent);
 
@@ -65,12 +65,15 @@ public:
 	virtual bool CanRename() const { return false; }
 	virtual void BeginRename() { }
 
+	virtual bool IsRoot() const { return false; }
+
 public:
 	FSimpleDelegate RenameEvent;
 
 protected:
 	virtual void GetChildren(TArray< TSharedPtr<FHierarchyModel> >& Children) = 0;
 	virtual void UpdateSelection() = 0;
+	virtual FWidgetReference AsDraggedWidgetReference() const { return FWidgetReference(); }
 
 private:
 	void InitializeChildren();
@@ -79,6 +82,7 @@ protected:
 
 	bool bInitialized;
 	bool bIsSelected;
+	TWeakPtr<FWidgetBlueprintEditor> BlueprintEditor;
 	TArray< TSharedPtr<FHierarchyModel> > Models;
 };
 
@@ -103,14 +107,13 @@ public:
 	virtual TOptional<EItemDropZone> HandleCanAcceptDrop(const FDragDropEvent& DragDropEvent, EItemDropZone DropZone) override;
 	virtual FReply HandleAcceptDrop(FDragDropEvent const& DragDropEvent, EItemDropZone DropZone) override;
 
+	virtual bool IsRoot() const override { return true; }
+
 protected:
 	virtual void GetChildren(TArray< TSharedPtr<FHierarchyModel> >& Children) override;
 	virtual void UpdateSelection() override;
 
 private:
-
-	TWeakPtr<FWidgetBlueprintEditor> BlueprintEditor;
-
 	FText RootText;
 };
 
@@ -134,18 +137,16 @@ public:
 	
 	virtual TOptional<EItemDropZone> HandleCanAcceptDrop(const FDragDropEvent& DragDropEvent, EItemDropZone DropZone) override;
 	virtual FReply HandleAcceptDrop(FDragDropEvent const& DragDropEvent, EItemDropZone DropZone) override;
-	virtual FReply HandleDragDetected(const FGeometry& MyGeometry, const FPointerEvent& MouseEvent) override;
 
 protected:
 	virtual void GetChildren(TArray< TSharedPtr<FHierarchyModel> >& Children) override;
 	virtual void UpdateSelection() override;
+	virtual FWidgetReference AsDraggedWidgetReference() const;
 
 private:
 
 	FWidgetReference Item;
 	FName SlotName;
-
-	TWeakPtr<FWidgetBlueprintEditor> BlueprintEditor;
 };
 
 class FHierarchyWidget : public FHierarchyModel
@@ -168,7 +169,6 @@ public:
 
 	virtual TOptional<EItemDropZone> HandleCanAcceptDrop(const FDragDropEvent& DragDropEvent, EItemDropZone DropZone) override;
 
-	virtual FReply HandleDragDetected(const FGeometry& MyGeometry, const FPointerEvent& MouseEvent) override;
 	virtual void HandleDragLeave(const FDragDropEvent& DragDropEvent) override;
 	virtual FReply HandleAcceptDrop(FDragDropEvent const& DragDropEvent, EItemDropZone DropZone) override;
 
@@ -224,11 +224,10 @@ public:
 protected:
 	virtual void GetChildren(TArray< TSharedPtr<FHierarchyModel> >& Children) override;
 	virtual void UpdateSelection() override;
+	virtual FWidgetReference AsDraggedWidgetReference() const { return Item; }
 
 private:
 	FWidgetReference Item;
-
-	TWeakPtr<FWidgetBlueprintEditor> BlueprintEditor;
 };
 
 /**
