UMG - Little bit of cleanup in the navigation customization removing unused code.

diff --git a/Engine/Source/Editor/UMGEditor/Private/Customizations/WidgetNavigationCustomization.cpp b/Engine/Source/Editor/UMGEditor/Private/Customizations/WidgetNavigationCustomization.cpp
--- a/Engine/Source/Editor/UMGEditor/Private/Customizations/WidgetNavigationCustomization.cpp
+++ b/Engine/Source/Editor/UMGEditor/Private/Customizations/WidgetNavigationCustomization.cpp
@@ -23,6 +23,8 @@ void FWidgetNavigationCustomization::CustomizeChildren(TSharedRef<IPropertyHandl
 {
 	TWeakPtr<IPropertyHandle> PropertyHandlePtr(PropertyHandle);
 
+	//IDetailCategoryBuilder& PropertyCategory = DetailLayout.EditCategory("Events", LOCTEXT("Events", "Events"), ECategoryPriority::Uncommon);
+
 	MakeNavRow(PropertyHandlePtr, ChildBuilder, EUINavigation::Left, LOCTEXT("LeftNavigation", "Left"));
 	MakeNavRow(PropertyHandlePtr, ChildBuilder, EUINavigation::Right, LOCTEXT("RightNavigation", "Right"));
 	MakeNavRow(PropertyHandlePtr, ChildBuilder, EUINavigation::Up, LOCTEXT("UpNavigation", "Up"));
@@ -31,53 +33,6 @@ void FWidgetNavigationCustomization::CustomizeChildren(TSharedRef<IPropertyHandl
 	MakeNavRow(PropertyHandlePtr, ChildBuilder, EUINavigation::Previous, LOCTEXT("PreviousNavigation", "Previous"));
 }
 
-void FWidgetNavigationCustomization::FillOutChildren(TSharedRef<IPropertyHandle> PropertyHandle, IDetailChildrenBuilder& ChildBuilder, IPropertyTypeCustomizationUtils& CustomizationUtils)
-{
-	// Generate all the other children
-	uint32 NumChildren;
-	PropertyHandle->GetNumChildren(NumChildren);
-
-	for ( uint32 ChildIndex = 0; ChildIndex < NumChildren; ChildIndex++ )
-	{
-		TSharedRef<IPropertyHandle> ChildHandle = PropertyHandle->GetChildHandle(ChildIndex).ToSharedRef();
-		if ( ChildHandle->IsCustomized() )
-		{
-			continue;
-		}
-
-		if ( ChildHandle->GetProperty() == NULL )
-		{
-			FillOutChildren(ChildHandle, ChildBuilder, CustomizationUtils);
-		}
-		else
-		{
-			ChildBuilder.AddChildProperty(ChildHandle);
-		}
-	}
-}
-
-FReply FWidgetNavigationCustomization::OnCustomizeNavigation(TWeakPtr<IPropertyHandle> PropertyHandle)
-{
-	TArray<UObject*> OuterObjects;
-	PropertyHandle.Pin()->GetOuterObjects(OuterObjects);
-
-	const FScopedTransaction Transaction(LOCTEXT("UpdateNavigation", "Update Navigation"));
-
-	for ( UObject* OuterObject : OuterObjects )
-	{
-		if ( UWidget* Widget = Cast<UWidget>(OuterObject) )
-		{
-			if ( !Widget->Navigation )
-			{
-				Widget->Modify();
-				Widget->Navigation = NewObject<UWidgetNavigation>(Widget);
-			}
-		}
-	}
-
-	return FReply::Handled();
-}
-
 EUINavigationRule FWidgetNavigationCustomization::GetNavigationRule(TWeakPtr<IPropertyHandle> PropertyHandle, EUINavigation Nav) const
 {
 	TArray<UObject*> OuterObjects;
diff --git a/Engine/Source/Editor/UMGEditor/Private/Customizations/WidgetNavigationCustomization.h b/Engine/Source/Editor/UMGEditor/Private/Customizations/WidgetNavigationCustomization.h
--- a/Engine/Source/Editor/UMGEditor/Private/Customizations/WidgetNavigationCustomization.h
+++ b/Engine/Source/Editor/UMGEditor/Private/Customizations/WidgetNavigationCustomization.h
@@ -23,10 +23,6 @@ public:
 	virtual void CustomizeChildren(TSharedRef<IPropertyHandle> PropertyHandle, class IDetailChildrenBuilder& ChildBuilder, IPropertyTypeCustomizationUtils& CustomizationUtils) override;
 	
 private:
-
-	void FillOutChildren(TSharedRef<IPropertyHandle> PropertyHandle, class IDetailChildrenBuilder& ChildBuilder, IPropertyTypeCustomizationUtils& CustomizationUtils);
-
-	FReply OnCustomizeNavigation(TWeakPtr<IPropertyHandle> PropertyHandle);
 	
 	EUINavigationRule GetNavigationRule(TWeakPtr<IPropertyHandle> PropertyHandle, EUINavigation Nav) const;
 	FText GetNavigationText(TWeakPtr<IPropertyHandle> PropertyHandle, EUINavigation Nav) const;
