Changed tooltips when moving widgets in the UMG hierarchy to match World Outliner.
[UE-13349] - Add Green Checkmark and Red (/) to Hierarchy UI

diff --git a/Engine/Source/Editor/UMGEditor/Private/Hierarchy/SHierarchyViewItem.cpp b/Engine/Source/Editor/UMGEditor/Private/Hierarchy/SHierarchyViewItem.cpp
--- a/Engine/Source/Editor/UMGEditor/Private/Hierarchy/SHierarchyViewItem.cpp
+++ b/Engine/Source/Editor/UMGEditor/Private/Hierarchy/SHierarchyViewItem.cpp
@@ -158,7 +158,6 @@ TOptional<EItemDropZone> ProcessHierarchyDragDrop(const FDragDropEvent& DragDrop
 	if ( TemplateDragDropOp.IsValid() )
 	{
 		TemplateDragDropOp->ResetToDefaultToolTip();
-		TemplateDragDropOp->SetCursorOverride(TOptional<EMouseCursor::Type>());
 
 		// Are we adding to the root?
 		if ( !TargetItem.IsValid() && Blueprint->WidgetTree->RootWidget == nullptr )
@@ -170,6 +169,7 @@ TOptional<EItemDropZone> ProcessHierarchyDragDrop(const FDragDropEvent& DragDrop
 				FBlueprintEditorUtils::MarkBlueprintAsStructurallyModified(Blueprint);
 			}
 
+			TemplateDragDropOp->CurrentIconBrush = FEditorStyle::GetBrush(TEXT("Graph.ConnectorFeedback.OK"));
 			return EItemDropZone::OntoItem;
 		}
 		// Are we adding to a panel?
@@ -200,6 +200,7 @@ TOptional<EItemDropZone> ProcessHierarchyDragDrop(const FDragDropEvent& DragDrop
 					FBlueprintEditorUtils::MarkBlueprintAsStructurallyModified(Blueprint);
 				}
 
+				TemplateDragDropOp->CurrentIconBrush = FEditorStyle::GetBrush(TEXT("Graph.ConnectorFeedback.OK"));
 				return EItemDropZone::OntoItem;
 			}
 		}
@@ -208,14 +209,13 @@ TOptional<EItemDropZone> ProcessHierarchyDragDrop(const FDragDropEvent& DragDrop
 			TemplateDragDropOp->CurrentHoverText = LOCTEXT("CantHaveChildren", "Widget can't have children.");
 		}
 
-		TemplateDragDropOp->SetCursorOverride(EMouseCursor::SlashedCircle);
+		TemplateDragDropOp->CurrentIconBrush = FEditorStyle::GetBrush(TEXT("Graph.ConnectorFeedback.Error"));
 		return TOptional<EItemDropZone>();
 	}
 
 	TSharedPtr<FHierarchyWidgetDragDropOp> HierarchyDragDropOp = DragDropEvent.GetOperationAs<FHierarchyWidgetDragDropOp>();
 	if ( HierarchyDragDropOp.IsValid() )
 	{
-		HierarchyDragDropOp->SetCursorOverride(TOptional<EMouseCursor::Type>());
 		HierarchyDragDropOp->ResetToDefaultToolTip();
 
 		// If the target item is valid we're dealing with a normal widget in the hierarchy, otherwise we should assume it's
@@ -229,21 +229,21 @@ TOptional<EItemDropZone> ProcessHierarchyDragDrop(const FDragDropEvent& DragDrop
 
 			if ( bIsDraggedObject )
 			{
-				HierarchyDragDropOp->SetCursorOverride(EMouseCursor::SlashedCircle);
+				HierarchyDragDropOp->CurrentIconBrush = FEditorStyle::GetBrush(TEXT("Graph.ConnectorFeedback.Error"));
 				return TOptional<EItemDropZone>();
 			}
 
 			UPanelWidget* NewParent = Cast<UPanelWidget>(TargetItem.GetTemplate());
 			if ( !NewParent )
 			{
-				HierarchyDragDropOp->SetCursorOverride(EMouseCursor::SlashedCircle);
+				HierarchyDragDropOp->CurrentIconBrush = FEditorStyle::GetBrush(TEXT("Graph.ConnectorFeedback.Error"));
 				HierarchyDragDropOp->CurrentHoverText = LOCTEXT("CantHaveChildren", "Widget can't have children.");
 				return TOptional<EItemDropZone>();
 			}
 
 			if ( !NewParent->CanAddMoreChildren() )
 			{
-				HierarchyDragDropOp->SetCursorOverride(EMouseCursor::SlashedCircle);
+				HierarchyDragDropOp->CurrentIconBrush = FEditorStyle::GetBrush(TEXT("Graph.ConnectorFeedback.Error"));
 				HierarchyDragDropOp->CurrentHoverText = LOCTEXT("NoAdditionalChildren", "Widget can't accept additional children.");
 				return TOptional<EItemDropZone>();
 			}
@@ -264,7 +264,7 @@ TOptional<EItemDropZone> ProcessHierarchyDragDrop(const FDragDropEvent& DragDrop
 
 			if ( bFoundNewParentInChildSet )
 			{
-				HierarchyDragDropOp->SetCursorOverride(EMouseCursor::SlashedCircle);
+				HierarchyDragDropOp->CurrentIconBrush = FEditorStyle::GetBrush(TEXT("Graph.ConnectorFeedback.Error"));
 				HierarchyDragDropOp->CurrentHoverText = LOCTEXT("CantMakeWidgetChildOfChildren", "Can't make widget a child of its children.");
 				return TOptional<EItemDropZone>();
 			}
@@ -318,11 +318,13 @@ TOptional<EItemDropZone> ProcessHierarchyDragDrop(const FDragDropEvent& DragDrop
 				BlueprintEditor->SelectWidgets(SelectedTemplates, false);
 			}
 
+			HierarchyDragDropOp->CurrentIconBrush = FEditorStyle::GetBrush(TEXT("Graph.ConnectorFeedback.OK"));
 			return EItemDropZone::OntoItem;
 		}
 		else
 		{
-			HierarchyDragDropOp->SetCursorOverride(EMouseCursor::SlashedCircle);
+			HierarchyDragDropOp->CurrentIconBrush = FEditorStyle::GetBrush(TEXT("Graph.ConnectorFeedback.Error"));
+			HierarchyDragDropOp->CurrentHoverText = LOCTEXT("CantHaveChildren", "Widget can't have children.");
 		}
 
 		return TOptional<EItemDropZone>();
@@ -394,7 +396,6 @@ void FHierarchyModel::HandleDragLeave(const FDragDropEvent& DragDropEvent)
 	TSharedPtr<FDecoratedDragDropOp> DecoratedDragDropOp = DragDropEvent.GetOperationAs<FDecoratedDragDropOp>();
 	if ( DecoratedDragDropOp.IsValid() )
 	{
-		DecoratedDragDropOp->SetCursorOverride(TOptional<EMouseCursor::Type>());
 		DecoratedDragDropOp->ResetToDefaultToolTip();
 	}
 }
@@ -656,7 +657,7 @@ TOptional<EItemDropZone> FNamedSlotModel::HandleCanAcceptDrop(const FDragDropEve
 			if ( NamedSlotHost->GetContentForSlot(SlotName) != nullptr )
 			{
 				TSharedPtr<FDecoratedDragDropOp> DecoratedDragDropOp = DragDropEvent.GetOperationAs<FDecoratedDragDropOp>();
-				TemplateDragDropOp->SetCursorOverride(EMouseCursor::SlashedCircle);
+				TemplateDragDropOp->CurrentIconBrush = FEditorStyle::GetBrush(TEXT("Graph.ConnectorFeedback.Error"));
 
 				return TOptional<EItemDropZone>();
 			}
@@ -793,7 +794,6 @@ void FHierarchyWidget::HandleDragLeave(const FDragDropEvent& DragDropEvent)
 	TSharedPtr<FDecoratedDragDropOp> DecoratedDragDropOp = DragDropEvent.GetOperationAs<FDecoratedDragDropOp>();
 	if ( DecoratedDragDropOp.IsValid() )
 	{
-		DecoratedDragDropOp->SetCursorOverride(TOptional<EMouseCursor::Type>());
 		DecoratedDragDropOp->ResetToDefaultToolTip();
 	}
 }
