Added support for drag-moving multiple widgets with the mouse in UMG Designer
[UE-16855] - Multiple select dragging with the mouse in UMG Designer

diff --git a/Engine/Source/Editor/UMGEditor/Private/Designer/SDesignerView.cpp b/Engine/Source/Editor/UMGEditor/Private/Designer/SDesignerView.cpp
--- a/Engine/Source/Editor/UMGEditor/Private/Designer/SDesignerView.cpp
+++ b/Engine/Source/Editor/UMGEditor/Private/Designer/SDesignerView.cpp
@@ -71,39 +71,77 @@ class FSelectedWidgetDragDropOp : public FDecoratedDragDropOp
 public:
 	DRAG_DROP_OPERATOR_TYPE(FSelectedWidgetDragDropOp, FDecoratedDragDropOp)
 
-	TMap<FName, FString> ExportedSlotProperties;
+	struct FDraggingWidgetReference
+	{
+		FWidgetReference Widget;
+
+		FVector2D DraggedOffset;
+	};
+
+	struct FItem
+	{
+		/** The slot properties for the old slot the widget was in, is used to attempt to reapply the same layout information */
+		TMap<FName, FString> ExportedSlotProperties;
+
+		/** The widget being dragged */
+		UWidget* Template;
+
+		/** The preview widget being dragged */
+		UWidget* Preview;
 
-	UWidget* Template;
-	UWidget* Preview;
+		/** Can the drag drop change the widget's parent? */
+		bool bStayingInParent;
 
-	bool bStayingInParent;
-	FWidgetReference ParentWidget;
+		/** The original parent of the widget. */
+		FWidgetReference ParentWidget;
 
-	static TSharedRef<FSelectedWidgetDragDropOp> New(TSharedPtr<FWidgetBlueprintEditor> Editor, FWidgetReference InWidget);
+		FVector2D DraggedOffset;
+	};
+
+	TArray<FItem> DraggedWidgets;
+
+	static TSharedRef<FSelectedWidgetDragDropOp> New(TSharedPtr<FWidgetBlueprintEditor> Editor, const TArray<FDraggingWidgetReference>& InWidgets);
 };
 
-TSharedRef<FSelectedWidgetDragDropOp> FSelectedWidgetDragDropOp::New(TSharedPtr<FWidgetBlueprintEditor> Editor, FWidgetReference InWidget)
+TSharedRef<FSelectedWidgetDragDropOp> FSelectedWidgetDragDropOp::New(TSharedPtr<FWidgetBlueprintEditor> Editor, const TArray<FDraggingWidgetReference>& InWidgets)
 {
-	bool bStayInParent = false;
-	if ( UPanelWidget* PanelTemplate = InWidget.GetTemplate()->GetParent() )
+	TSharedRef<FSelectedWidgetDragDropOp> Operation = MakeShareable(new FSelectedWidgetDragDropOp());
+
+	for (const auto& InDraggedWidget : InWidgets)
 	{
-		bStayInParent = PanelTemplate->LockToPanelOnDrag() || GetDefault<UWidgetDesignerSettings>()->bLockToPanelOnDragByDefault;
-	}
+		FItem DraggedWidget;
+		DraggedWidget.bStayingInParent = false;
 
-	TSharedRef<FSelectedWidgetDragDropOp> Operation = MakeShareable(new FSelectedWidgetDragDropOp());
-	Operation->bStayingInParent = bStayInParent;
-	Operation->ParentWidget = Editor->GetReferenceFromTemplate(InWidget.GetTemplate()->GetParent());
-	Operation->DefaultHoverText = InWidget.GetTemplate()->GetLabelText();
-	Operation->CurrentHoverText = InWidget.GetTemplate()->GetLabelText();
-	Operation->Construct();
+		if (UPanelWidget* PanelTemplate = InDraggedWidget.Widget.GetTemplate()->GetParent())
+		{
+			DraggedWidget.ParentWidget = Editor->GetReferenceFromTemplate(PanelTemplate);
+			DraggedWidget.bStayingInParent = PanelTemplate->LockToPanelOnDrag() || GetDefault<UWidgetDesignerSettings>()->bLockToPanelOnDragByDefault;
+		}
 
-	// Cache the preview and template, it's not safe to query the preview/template while dragging the widget as it no longer
-	// exists in the tree.
-	Operation->Preview = InWidget.GetPreview();
-	Operation->Template = InWidget.GetTemplate();
+		// Cache the preview and template, it's not safe to query the preview/template while dragging the widget as it no longer
+		// exists in the tree.
+		DraggedWidget.Preview = InDraggedWidget.Widget.GetPreview();
+		DraggedWidget.Template = InDraggedWidget.Widget.GetTemplate();
 
-	FWidgetBlueprintEditorUtils::ExportPropertiesToText(InWidget.GetTemplate()->Slot, Operation->ExportedSlotProperties);
+		DraggedWidget.DraggedOffset = InDraggedWidget.DraggedOffset;
+
+		FWidgetBlueprintEditorUtils::ExportPropertiesToText(InDraggedWidget.Widget.GetTemplate()->Slot, DraggedWidget.ExportedSlotProperties);
+
+		Operation->DraggedWidgets.Add(DraggedWidget);
+	}
+
+	// Set the display text based on whether we're dragging a single or multiple widgets
+	if (InWidgets.Num() == 1)
+	{
+		Operation->DefaultHoverText = InWidgets[0].Widget.GetTemplate()->GetLabelText();
+		Operation->CurrentHoverText = InWidgets[0].Widget.GetTemplate()->GetLabelText();
+	}
+	else
+	{
+		Operation->CurrentHoverText = Operation->DefaultHoverText = LOCTEXT("DragMultipleWidgets", "Multiple Widgets");
+	}
 
+	Operation->Construct();
 	return Operation;
 }
 
@@ -167,8 +205,6 @@ void SDesignerView::Construct(const FArguments& InArgs, TSharedPtr<FWidgetBluepr
 	ScopedTransaction = nullptr;
 
 	PreviewWidget = nullptr;
-	DropPreviewWidget = nullptr;
-	DropPreviewParent = nullptr;
 	BlueprintEditor = InBlueprintEditor;
 
 	DesignerMessage = EDesignerMessage::None;
@@ -591,6 +627,18 @@ void SDesignerView::AddReferencedObjects(FReferenceCollector& Collector)
 	{
 		Collector.AddReferencedObject(PreviewWidget);
 	}
+
+	for (auto& DropPreview : DropPreviews)
+	{
+		if (DropPreview.Widget)
+		{
+			Collector.AddReferencedObject(DropPreview.Widget);
+		}
+		if (DropPreview.Parent)
+		{
+			Collector.AddReferencedObject(DropPreview.Parent);
+		}
+	}
 }
 
 void SDesignerView::SetTransformMode(ETransformMode::Type InTransformMode)
@@ -1033,6 +1081,21 @@ FVector2D SDesignerView::GetExtensionSize(TSharedRef<FDesignerSurfaceElement> Ex
 	return ExtensionElement->GetWidget()->GetDesiredSize();
 }
 
+void SDesignerView::ClearDropPreviews()
+{
+	for (const auto& DropPreview : DropPreviews)
+	{
+		if (DropPreview.Parent)
+		{
+			DropPreview.Parent->RemoveChild(DropPreview.Widget);
+		}
+
+		UWidgetBlueprint* BP = GetBlueprint();
+		BP->WidgetTree->RemoveWidget(DropPreview.Widget);
+	}
+	DropPreviews.Empty();
+}
+
 UWidgetBlueprint* SDesignerView::GetBlueprint() const
 {
 	if ( BlueprintEditor.IsValid() )
@@ -1094,8 +1157,13 @@ bool SDesignerView::FindWidgetUnderCursor(const FGeometry& MyGeometry, const FPo
 			FArrangedWidget& Child = Children.GetInternalArray()[ChildIndex];
 			WidgetUnderCursor = PreviewUserWidget->GetWidgetHandle(Child.Widget);
 			
-			// Ignore the drop preview widget when doing widget picking
-			if ( WidgetUnderCursor == DropPreviewWidget )
+			if (WidgetUnderCursor == nullptr)
+			{
+				continue;
+			}
+
+			// Ignore the drop preview widgets when doing widget picking
+			if (DropPreviews.ContainsByPredicate([WidgetUnderCursor](const FDropPreview& Preview){ return Preview.Widget == WidgetUnderCursor; }))
 			{
 				WidgetUnderCursor = nullptr;
 				continue;
@@ -1163,10 +1231,21 @@ FReply SDesignerView::OnMouseButtonDown(const FGeometry& MyGeometry, const FPoin
 
 		if ( MouseEvent.GetEffectingButton() == EKeys::LeftMouseButton )
 		{
-			const bool bResolvePendingSelectionImmediately =
-				!GetSelectedWidget().IsValid() ||
-				!PendingSelectedWidget.GetTemplate()->IsChildOf(GetSelectedWidget().GetTemplate()) ||
-				GetSelectedWidget().GetTemplate()->GetParent() == nullptr;
+			const TSet<FWidgetReference>& SelectedWidgets = GetSelectedWidgets();
+
+			bool bResolvePendingSelectionImmediately = true;
+
+			if (SelectedWidgets.Num() > 0)
+			{
+				for (const auto& SelectedWidget : SelectedWidgets)
+				{
+					if (PendingSelectedWidget == SelectedWidget || PendingSelectedWidget.GetTemplate()->IsChildOf(SelectedWidget.GetTemplate()))
+					{
+						bResolvePendingSelectionImmediately = false;
+						break;
+					}
+				}
+			}
 
 			// If the newly clicked item is a child of the active selection, add it to the pending set of selected 
 			// widgets, if they begin dragging we can just move the parent, but if it's not part of the parent set, 
@@ -1235,14 +1314,23 @@ FReply SDesignerView::OnMouseMove(const FGeometry& MyGeometry, const FPointerEve
 
 	if ( MouseEvent.IsMouseButtonDown(EKeys::LeftMouseButton) && HasMouseCapture() )
 	{
-		FWidgetReference SelectedWidget = GetSelectedWidget();
+		const TSet<FWidgetReference>& SelectedWidgets = GetSelectedWidgets();
 
-		if ( SelectedWidget.IsValid() && !bMovingExistingWidget )
+		if (SelectedWidgets.Num() > 0 && !bMovingExistingWidget)
 		{
 			if ( TransformMode == ETransformMode::Layout )
 			{
-				const bool bIsRootWidget = SelectedWidget.GetTemplate()->GetParent() == nullptr;
-				if ( !bIsRootWidget )
+				bool bIsRootWidgetSelected = false;
+				for (const auto& SelectedWidget : SelectedWidgets)
+				{
+					if (SelectedWidget.GetTemplate()->GetParent() == nullptr)
+					{
+						bIsRootWidgetSelected = true;
+						break;
+					}
+				}
+
+				if (!bIsRootWidgetSelected)
 				{
 					bMovingExistingWidget = true;
 					//Drag selected widgets
@@ -1254,22 +1342,33 @@ FReply SDesignerView::OnMouseMove(const FGeometry& MyGeometry, const FPointerEve
 				checkSlow(TransformMode == ETransformMode::Render);
 				checkSlow(bMovingExistingWidget == false);
 
-				BeginTransaction(LOCTEXT("MoveWidgetRT", "Move Widget (Render Transform)"));
+				if (SelectedWidgets.Num() == 1)
+				{
+					BeginTransaction(LOCTEXT("MoveWidgetRT", "Move Widget (Render Transform)"));
+				}
+				else
+				{
+					check(SelectedWidgets.Num() > 1);
+					BeginTransaction(LOCTEXT("MoveWidgetsRT", "Move Widgets (Render Transform)"));
+				}
 
-				if ( UWidget* WidgetPreview = SelectedWidget.GetPreview() )
+				for (const auto& SelectedWidget : SelectedWidgets)
 				{
-					FGeometry ParentGeometry;
-					if ( GetWidgetParentGeometry(SelectedWidget, ParentGeometry) )
+					if (UWidget* WidgetPreview = SelectedWidget.GetPreview())
 					{
-						const FSlateRenderTransform& AbsoluteToLocalTransform = Inverse(ParentGeometry.GetAccumulatedRenderTransform());
+						FGeometry ParentGeometry;
+						if (GetWidgetParentGeometry(SelectedWidget, ParentGeometry))
+						{
+							const FSlateRenderTransform& AbsoluteToLocalTransform = Inverse(ParentGeometry.GetAccumulatedRenderTransform());
 
-						FWidgetTransform RenderTransform = WidgetPreview->RenderTransform;
-						RenderTransform.Translation += AbsoluteToLocalTransform.TransformVector(MouseEvent.GetCursorDelta());
+							FWidgetTransform RenderTransform = WidgetPreview->RenderTransform;
+							RenderTransform.Translation += AbsoluteToLocalTransform.TransformVector(MouseEvent.GetCursorDelta());
 
-						static const FName RenderTransformName(TEXT("RenderTransform"));
+							static const FName RenderTransformName(TEXT("RenderTransform"));
 
-						FObjectEditorUtils::SetPropertyValue<UWidget, FWidgetTransform>(WidgetPreview, RenderTransformName, RenderTransform);
-						FObjectEditorUtils::SetPropertyValue<UWidget, FWidgetTransform>(SelectedWidget.GetTemplate(), RenderTransformName, RenderTransform);
+							FObjectEditorUtils::SetPropertyValue<UWidget, FWidgetTransform>(WidgetPreview, RenderTransformName, RenderTransform);
+							FObjectEditorUtils::SetPropertyValue<UWidget, FWidgetTransform>(SelectedWidget.GetTemplate(), RenderTransformName, RenderTransform);
+						}
 					}
 				}
 			}
@@ -1655,23 +1754,57 @@ void SDesignerView::Tick( const FGeometry& AllottedGeometry, const double InCurr
 
 FReply SDesignerView::OnDragDetected(const FGeometry& MyGeometry, const FPointerEvent& MouseEvent)
 {
+	typedef FSelectedWidgetDragDropOp::FDraggingWidgetReference FDragWidget;
+
 	SDesignSurface::OnDragDetected(MyGeometry, MouseEvent);
 
-	FWidgetReference SelectedWidget = GetSelectedWidget();
+	const TSet<FWidgetReference>& SelectedWidgets = GetSelectedWidgets();
 
-	if ( SelectedWidget.IsValid() )
+	if (SelectedWidgets.Num() > 0)
 	{
+		TArray<FDragWidget> DraggingWidgetCandidates;
+
 		// Clear any pending selected widgets, the user has already decided what widget they want.
 		PendingSelectedWidget = FWidgetReference();
 
-		// Determine The offset to keep the widget from the mouse while dragging
-		FArrangedWidget ArrangedWidget(SNullWidget::NullWidget, FGeometry());
-		FDesignTimeUtils::GetArrangedWidget(SelectedWidget.GetPreview()->GetCachedWidget().ToSharedRef(), ArrangedWidget);
-		SelectedWidgetContextMenuLocation = ArrangedWidget.Geometry.AbsoluteToLocal(DraggingStartPositionScreenSpace);
+		for (const auto& SelectedWidget : SelectedWidgets)
+		{
+			// Determine The offset to keep the widget from the mouse while dragging
+			FArrangedWidget ArrangedWidget(SNullWidget::NullWidget, FGeometry());
+			FDesignTimeUtils::GetArrangedWidget(SelectedWidget.GetPreview()->GetCachedWidget().ToSharedRef(), ArrangedWidget);
+			SelectedWidgetContextMenuLocation = ArrangedWidget.Geometry.AbsoluteToLocal(DraggingStartPositionScreenSpace);
+
+			FDragWidget DraggingWidget;
+			DraggingWidget.Widget = SelectedWidget;
+			DraggingWidget.DraggedOffset = SelectedWidgetContextMenuLocation;
+
+			DraggingWidgetCandidates.Add(DraggingWidget);
+		}
+
+		TArray<FDragWidget> DraggingWidgets;
+
+		for (const auto& Candidate : DraggingWidgetCandidates)
+		{
+			// check the parent chain of each dragged widget and ignore those that are children of other dragged widgets
+			bool bIsChild = false;
+			for (auto CursorPtr = Candidate.Widget.GetTemplate()->GetParent(); CursorPtr != nullptr; CursorPtr = CursorPtr->GetParent())
+			{
+				if (DraggingWidgetCandidates.ContainsByPredicate([CursorPtr](const FDragWidget& W){ return W.Widget.GetTemplate() == CursorPtr; }))
+				{
+					bIsChild = true;
+					break;
+				}
+			}
+
+			if (!bIsChild)
+			{
+				DraggingWidgets.Add(Candidate);
+			}
+		}
 
 		ClearExtensionWidgets();
 
-		return FReply::Handled().BeginDragDrop(FSelectedWidgetDragDropOp::New(BlueprintEditor.Pin(), SelectedWidget));
+		return FReply::Handled().BeginDragDrop(FSelectedWidgetDragDropOp::New(BlueprintEditor.Pin(), DraggingWidgets));
 	}
 
 	return FReply::Unhandled();
@@ -1699,39 +1832,19 @@ void SDesignerView::OnDragLeave(const FDragDropEvent& DragDropEvent)
 		DecoratedDragDropOp->ResetToDefaultToolTip();
 	}
 	
-	if (DropPreviewWidget)
-	{
-		if ( DropPreviewParent )
-		{
-			DropPreviewParent->RemoveChild(DropPreviewWidget);
-		}
-
-		UWidgetBlueprint* BP = GetBlueprint();
-		BP->WidgetTree->RemoveWidget(DropPreviewWidget);
-		DropPreviewWidget = nullptr;
-	}
+	ClearDropPreviews();
 }
 
 FReply SDesignerView::OnDragOver(const FGeometry& MyGeometry, const FDragDropEvent& DragDropEvent)
 {
 	SDesignSurface::OnDragOver(MyGeometry, DragDropEvent);
 
-	UWidgetBlueprint* BP = GetBlueprint();
-	
-	if (DropPreviewWidget)
-	{
-		if (DropPreviewParent)
-		{
-			DropPreviewParent->RemoveChild(DropPreviewWidget);
-		}
-		
-		BP->WidgetTree->RemoveWidget(DropPreviewWidget);
-		DropPreviewWidget = nullptr;
-	}
+	ClearDropPreviews();
 	
 	const bool bIsPreview = true;
-	DropPreviewWidget = ProcessDropAndAddWidget(MyGeometry, DragDropEvent, bIsPreview);
-	if ( DropPreviewWidget )
+	ProcessDropAndAddWidget(MyGeometry, DragDropEvent, bIsPreview);
+
+	if ( DropPreviews.Num() > 0 )
 	{
 		//@TODO UMG Drop Feedback
 		return FReply::Handled();
@@ -1740,29 +1853,25 @@ FReply SDesignerView::OnDragOver(const FGeometry& MyGeometry, const FDragDropEve
 	return FReply::Unhandled();
 }
 
-UWidget* SDesignerView::ProcessDropAndAddWidget(const FGeometry& MyGeometry, const FDragDropEvent& DragDropEvent, bool bIsPreview)
+void SDesignerView::ProcessDropAndAddWidget(const FGeometry& MyGeometry, const FDragDropEvent& DragDropEvent, const bool bIsPreview)
 {
-	// In order to prevent the GetWidgetAtCursor code from picking the widget we're about to move, we need to mark it
-	// as the drop preview widget before any other code can run.
+	// In order to prevent the GetWidgetAtCursor code from picking the widgets we're about to move, we need to mark them
+	// as the drop preview widgets before any other code can run.
 	TSharedPtr<FSelectedWidgetDragDropOp> SelectedDragDropOp = DragDropEvent.GetOperationAs<FSelectedWidgetDragDropOp>();
 	if ( SelectedDragDropOp.IsValid() )
 	{
-		DropPreviewWidget = SelectedDragDropOp->Preview;
-	}
-
-	UWidgetBlueprint* BP = GetBlueprint();
-
-	if ( DropPreviewWidget )
-	{
-		if ( DropPreviewParent )
+		DropPreviews.Empty();
+		
+		for (const auto& DraggedWidget : SelectedDragDropOp->DraggedWidgets)
 		{
-			DropPreviewParent->RemoveChild(DropPreviewWidget);
+			FDropPreview DropPreview;
+			DropPreview.Parent = nullptr;
+			DropPreview.Widget = DraggedWidget.Preview;
+			DropPreviews.Add(DropPreview);
 		}
-
-		BP->WidgetTree->RemoveWidget(DropPreviewWidget);
-		DropPreviewWidget = nullptr;
 	}
 
+	ClearDropPreviews();
 
 	UWidget* Target = nullptr;
 
@@ -1773,6 +1882,7 @@ UWidget* SDesignerView::ProcessDropAndAddWidget(const FGeometry& MyGeometry, con
 	}
 
 	FGeometry WidgetUnderCursorGeometry = HitResult.WidgetArranged.Geometry;
+	UWidgetBlueprint* BP = GetBlueprint();
 
 	TSharedPtr<FWidgetTemplateDragDropOp> TemplateDragDropOp = DragDropEvent.GetOperationAs<FWidgetTemplateDragDropOp>();
 	if ( TemplateDragDropOp.IsValid() )
@@ -1800,14 +1910,17 @@ UWidget* SDesignerView::ProcessDropAndAddWidget(const FGeometry& MyGeometry, con
 
 			BP->WidgetTree->RootWidget = Widget;
 
-			DropPreviewParent = nullptr;
+			FDropPreview DropPreview;
+			DropPreview.Widget = Widget;
+			DropPreview.Parent = nullptr;
+			DropPreviews.Add(DropPreview);
 
 			if ( bIsPreview )
 			{
 				Transaction.Cancel();
 			}
 
-			return Widget;
+			return;
 		}
 		// If there's already a root widget we need to try and place our widget into a parent widget that we've picked against
 		else if ( Target && Target->IsA(UPanelWidget::StaticClass()) )
@@ -1856,14 +1969,17 @@ UWidget* SDesignerView::ProcessDropAndAddWidget(const FGeometry& MyGeometry, con
 					CanvasSlot->SetSize(LocalSize);
 				}
 
-				DropPreviewParent = Parent;
+				FDropPreview DropPreview;
+				DropPreview.Widget = Widget;
+				DropPreview.Parent = Parent;
+				DropPreviews.Add(DropPreview);
 
 				if ( bIsPreview )
 				{
 					Transaction.Cancel();
 				}
 
-				return Widget;
+				return;
 			}
 			else
 			{
@@ -1885,173 +2001,182 @@ UWidget* SDesignerView::ProcessDropAndAddWidget(const FGeometry& MyGeometry, con
 	}
 
 	// Attempt to deal with moving widgets from a drag operation.
-	if ( SelectedDragDropOp.IsValid() )
+	if (SelectedDragDropOp.IsValid() && SelectedDragDropOp->DraggedWidgets.Num() > 0)
 	{
 		SelectedDragDropOp->SetCursorOverride(TOptional<EMouseCursor::Type>());
 
-		// If they've pressed alt, and we were staying in the parent, disable that
-		// and adjust the designer message to no longer warn.
-		if ( DragDropEvent.IsAltDown() && SelectedDragDropOp->bStayingInParent )
+		FText TransactionText;
+		if (SelectedDragDropOp->DraggedWidgets.Num() == 1)
 		{
-			SelectedDragDropOp->bStayingInParent = false;
-			DesignerMessage = EDesignerMessage::None;
+			TransactionText = LOCTEXT("Designer_MoveWidget", "Move Widget");
 		}
-
-		// If we're staying in the parent we started in, replace the parent found under the cursor with
-		// the original one, also update the arranged widget data so that our layout calculations are accurate.
-		if ( SelectedDragDropOp->bStayingInParent )
+		else
 		{
-			DesignerMessage = EDesignerMessage::MoveFromParent;
-
-			WidgetUnderCursorGeometry = GetDesignerGeometry();
-			if ( GetWidgetGeometry(SelectedDragDropOp->ParentWidget, WidgetUnderCursorGeometry) )
-			{
-				Target = bIsPreview ? SelectedDragDropOp->ParentWidget.GetPreview() : SelectedDragDropOp->ParentWidget.GetTemplate();
-			}
+			TransactionText = LOCTEXT("Designer_MoveWidgets", "Move Widgets");
 		}
 
-		FWidgetReference TargetReference = bIsPreview ? BlueprintEditor.Pin()->GetReferenceFromPreview(Target) : BlueprintEditor.Pin()->GetReferenceFromTemplate(Target);
-		BlueprintEditor.Pin()->SetHoveredWidget(TargetReference);
+		FScopedTransaction Transaction(TransactionText);
+		bool bWidgetMoved = false;
 
-		// If the widget being hovered over is a panel, attempt to place it into that panel.
-		if ( Target && Target->IsA(UPanelWidget::StaticClass()) )
+		for (auto& DraggedWidget : SelectedDragDropOp->DraggedWidgets)
 		{
-			UPanelWidget* NewParent = Cast<UPanelWidget>(Target);
-
-			FScopedTransaction Transaction(LOCTEXT("Designer_MoveWidget", "Move Widget"));
+			// If they've pressed alt, and we were staying in the parent, disable that
+			// and adjust the designer message to no longer warn.
+			if (DragDropEvent.IsAltDown() && DraggedWidget.bStayingInParent)
+			{
+				DraggedWidget.bStayingInParent = false;
+				DesignerMessage = EDesignerMessage::None;
+			}
 
-			// If this isn't a preview operation we need to modify a few things to properly undo the operation.
-			if ( !bIsPreview )
+			// If we're staying in the parent we started in, replace the parent found under the cursor with
+			// the original one, also update the arranged widget data so that our layout calculations are accurate.
+			if (DraggedWidget.bStayingInParent)
 			{
-				NewParent->SetFlags(RF_Transactional);
-				NewParent->Modify();
+				DesignerMessage = EDesignerMessage::MoveFromParent;
 
-				BP->WidgetTree->SetFlags(RF_Transactional);
-				BP->WidgetTree->Modify();
+				WidgetUnderCursorGeometry = GetDesignerGeometry();
+				if (GetWidgetGeometry(DraggedWidget.ParentWidget, WidgetUnderCursorGeometry))
+				{
+					Target = bIsPreview ? DraggedWidget.ParentWidget.GetPreview() : DraggedWidget.ParentWidget.GetTemplate();
+				}
 			}
 
-			UWidget* Widget = bIsPreview ? SelectedDragDropOp->Preview : SelectedDragDropOp->Template;
+			FWidgetReference TargetReference = bIsPreview ? BlueprintEditor.Pin()->GetReferenceFromPreview(Target) : BlueprintEditor.Pin()->GetReferenceFromTemplate(Target);
+			BlueprintEditor.Pin()->SetHoveredWidget(TargetReference);
 
-			if ( ensure(Widget) )
+			// If the widget being hovered over is a panel, attempt to place it into that panel.
+			if (Target && Target->IsA(UPanelWidget::StaticClass()))
 			{
-				if ( !bIsPreview )
-				{
-					Widget->Modify();
-				}
+				bWidgetMoved = true;
+				UPanelWidget* NewParent = Cast<UPanelWidget>(Target);
 
-				if ( Widget->GetParent() )
+				// If this isn't a preview operation we need to modify a few things to properly undo the operation.
+				if (!bIsPreview)
 				{
-					if ( !bIsPreview )
-					{
-						Widget->GetParent()->Modify();
-					}
+					NewParent->SetFlags(RF_Transactional);
+					NewParent->Modify();
 
-					Widget->GetParent()->RemoveChild(Widget);
+					BP->WidgetTree->SetFlags(RF_Transactional);
+					BP->WidgetTree->Modify();
 				}
 
-				FVector2D ScreenSpacePosition = DragDropEvent.GetScreenSpacePosition();
+				UWidget* Widget = bIsPreview ? DraggedWidget.Preview : DraggedWidget.Template;
 
-				const UWidgetDesignerSettings* DesignerSettings = GetDefault<UWidgetDesignerSettings>();
-				bool bGridSnapX, bGridSnapY;
-				bGridSnapX = bGridSnapY = DesignerSettings->GridSnapEnabled;
-
-				// As long as shift is pressed and we're staying in the same parent,
-				// allow the user to lock the movement to a specific axis.
-				const bool bLockToAxis =
-					FSlateApplication::Get().GetModifierKeys().IsShiftDown() &&
-					SelectedDragDropOp->bStayingInParent;
-
-				if ( bLockToAxis )
+				if (ensure(Widget))
 				{
-					// Choose the largest axis of movement as the primary axis to lock to.
-					FVector2D DragDelta = ScreenSpacePosition - DraggingStartPositionScreenSpace;
-					if ( FMath::Abs(DragDelta.X) > FMath::Abs(DragDelta.Y) )
+					if (!bIsPreview)
 					{
-						// Lock to X Axis
-						ScreenSpacePosition.Y = DraggingStartPositionScreenSpace.Y;
-						bGridSnapY = false;
+						Widget->Modify();
 					}
-					else
+
+					if (Widget->GetParent())
 					{
-						// Lock To Y Axis
-						ScreenSpacePosition.X = DraggingStartPositionScreenSpace.X;
-						bGridSnapX = false;
+						if (!bIsPreview)
+						{
+							Widget->GetParent()->Modify();
+						}
+
+						Widget->GetParent()->RemoveChild(Widget);
 					}
-				}
 
-				FVector2D LocalPosition = WidgetUnderCursorGeometry.AbsoluteToLocal(ScreenSpacePosition);
-				if ( UPanelSlot* Slot = NewParent->AddChild(Widget) )
-				{
-					FVector2D NewPosition = LocalPosition - SelectedWidgetContextMenuLocation;
+					FVector2D ScreenSpacePosition = DragDropEvent.GetScreenSpacePosition();
 
-					// Perform grid snapping on X and Y if we need to.
-					if ( bGridSnapX )
-					{
-						NewPosition.X = ( (int32)NewPosition.X ) - ( ( (int32)NewPosition.X ) % DesignerSettings->GridSnapSize );
-					}
+					const UWidgetDesignerSettings* DesignerSettings = GetDefault<UWidgetDesignerSettings>();
+					bool bGridSnapX, bGridSnapY;
+					bGridSnapX = bGridSnapY = DesignerSettings->GridSnapEnabled;
+
+					// As long as shift is pressed and we're staying in the same parent,
+					// allow the user to lock the movement to a specific axis.
+					const bool bLockToAxis =
+						FSlateApplication::Get().GetModifierKeys().IsShiftDown() &&
+						DraggedWidget.bStayingInParent;
 
-					if ( bGridSnapY )
+					if (bLockToAxis)
 					{
-						NewPosition.Y = ( (int32)NewPosition.Y ) - ( ( (int32)NewPosition.Y ) % DesignerSettings->GridSnapSize );
+						// Choose the largest axis of movement as the primary axis to lock to.
+						FVector2D DragDelta = ScreenSpacePosition - DraggingStartPositionScreenSpace;
+						if (FMath::Abs(DragDelta.X) > FMath::Abs(DragDelta.Y))
+						{
+							// Lock to X Axis
+							ScreenSpacePosition.Y = DraggingStartPositionScreenSpace.Y;
+							bGridSnapY = false;
+						}
+						else
+						{
+							// Lock To Y Axis
+							ScreenSpacePosition.X = DraggingStartPositionScreenSpace.X;
+							bGridSnapX = false;
+						}
 					}
 
-					// HACK UMG: In order to correctly drop items into the canvas that have a non-zero anchor,
-					// we need to know the layout information after slate has performed a pre-pass.  So we have
-					// to rebase the layout and reinterpret the new position based on anchor point layout data.
-					// This should be pulled out into an extension of some kind so that this can be fixed for
-					// other widgets as well that may need to do work like this.
-					if ( UCanvasPanelSlot* CanvasSlot = Cast<UCanvasPanelSlot>(Slot) )
+					FVector2D LocalPosition = WidgetUnderCursorGeometry.AbsoluteToLocal(ScreenSpacePosition);
+					if (UPanelSlot* Slot = NewParent->AddChild(Widget))
 					{
-						if ( bIsPreview )
+						FVector2D NewPosition = LocalPosition - DraggedWidget.DraggedOffset;
+
+						// Perform grid snapping on X and Y if we need to.
+						if (bGridSnapX)
 						{
-							FWidgetBlueprintEditorUtils::ImportPropertiesFromText(Slot, SelectedDragDropOp->ExportedSlotProperties);
+							NewPosition.X = ((int32)NewPosition.X) - (((int32)NewPosition.X) % DesignerSettings->GridSnapSize);
+						}
 
-							CanvasSlot->SaveBaseLayout();
-							CanvasSlot->SetDesiredPosition(NewPosition);
-							CanvasSlot->RebaseLayout();
+						if (bGridSnapY)
+						{
+							NewPosition.Y = ((int32)NewPosition.Y) - (((int32)NewPosition.Y) % DesignerSettings->GridSnapSize);
+						}
 
-							FWidgetBlueprintEditorUtils::ExportPropertiesToText(Slot, SelectedDragDropOp->ExportedSlotProperties);
+						// HACK UMG: In order to correctly drop items into the canvas that have a non-zero anchor,
+						// we need to know the layout information after slate has performed a pre-pass.  So we have
+						// to rebase the layout and reinterpret the new position based on anchor point layout data.
+						// This should be pulled out into an extension of some kind so that this can be fixed for
+						// other widgets as well that may need to do work like this.
+						if (UCanvasPanelSlot* CanvasSlot = Cast<UCanvasPanelSlot>(Slot))
+						{
+							if (bIsPreview)
+							{
+								FWidgetBlueprintEditorUtils::ImportPropertiesFromText(Slot, DraggedWidget.ExportedSlotProperties);
+
+								CanvasSlot->SaveBaseLayout();
+								CanvasSlot->SetDesiredPosition(NewPosition);
+								CanvasSlot->RebaseLayout();
+
+								FWidgetBlueprintEditorUtils::ExportPropertiesToText(Slot, DraggedWidget.ExportedSlotProperties);
+							}
+							else
+							{
+								FWidgetBlueprintEditorUtils::ImportPropertiesFromText(Slot, DraggedWidget.ExportedSlotProperties);
+							}
 						}
 						else
 						{
-							FWidgetBlueprintEditorUtils::ImportPropertiesFromText(Slot, SelectedDragDropOp->ExportedSlotProperties);
+							FWidgetBlueprintEditorUtils::ImportPropertiesFromText(Slot, DraggedWidget.ExportedSlotProperties);
 						}
+
+						FDropPreview DropPreview;
+						DropPreview.Widget = Widget;
+						DropPreview.Parent = NewParent;
+						DropPreviews.Add(DropPreview);
 					}
 					else
 					{
-						FWidgetBlueprintEditorUtils::ImportPropertiesFromText(Slot, SelectedDragDropOp->ExportedSlotProperties);
-					}
-
-					DropPreviewParent = NewParent;
+						SelectedDragDropOp->SetCursorOverride(EMouseCursor::SlashedCircle);
 
-					if ( bIsPreview )
-					{
-						Transaction.Cancel();
+						// TODO UMG ERROR Slot can not be created because maybe the max children has been reached.
+						//          Maybe we can traverse the hierarchy and add it to the first parent that will accept it?
 					}
-
-					return Widget;
-				}
-				else
-				{
-					SelectedDragDropOp->SetCursorOverride(EMouseCursor::SlashedCircle);
-
-					// TODO UMG ERROR Slot can not be created because maybe the max children has been reached.
-					//          Maybe we can traverse the hierarchy and add it to the first parent that will accept it?
 				}
 			}
-
-			if ( bIsPreview )
+			else
 			{
-				Transaction.Cancel();
+				SelectedDragDropOp->SetCursorOverride(EMouseCursor::SlashedCircle);
 			}
 		}
-		else
+
+		if (bIsPreview || !bWidgetMoved)
 		{
-			SelectedDragDropOp->SetCursorOverride(EMouseCursor::SlashedCircle);
+			Transaction.Cancel();
 		}
 	}
-	
-	return nullptr;
 }
 
 FReply SDesignerView::OnDrop(const FGeometry& MyGeometry, const FDragDropEvent& DragDropEvent)
@@ -2060,28 +2185,23 @@ FReply SDesignerView::OnDrop(const FGeometry& MyGeometry, const FDragDropEvent&
 
 	bMovingExistingWidget = false;
 
-	UWidgetBlueprint* BP = GetBlueprint();
-	
-	if (DropPreviewWidget)
-	{
-		if (DropPreviewParent)
-		{
-			DropPreviewParent->RemoveChild(DropPreviewWidget);
-		}
-		
-		BP->WidgetTree->RemoveWidget(DropPreviewWidget);
-		DropPreviewWidget = nullptr;
-	}
+	ClearDropPreviews();
 	
 	const bool bIsPreview = false;
-	UWidget* NewTemplateWidget = ProcessDropAndAddWidget(MyGeometry, DragDropEvent, bIsPreview);
-	if ( NewTemplateWidget )
+	ProcessDropAndAddWidget(MyGeometry, DragDropEvent, bIsPreview);
+
+	if (DropPreviews.Num() > 0)
 	{
+		UWidgetBlueprint* BP = GetBlueprint();
 		FBlueprintEditorUtils::MarkBlueprintAsStructurallyModified(BP);
 
 		TSet<FWidgetReference> SelectedTemplates;
-		SelectedTemplates.Add(BlueprintEditor.Pin()->GetReferenceFromTemplate(NewTemplateWidget));
-		 
+
+		for (const auto& DropPreview : DropPreviews)
+		{
+			SelectedTemplates.Add(BlueprintEditor.Pin()->GetReferenceFromTemplate(DropPreview.Widget));
+		}
+
 		BlueprintEditor.Pin()->SelectWidgets(SelectedTemplates, false);
 
 		// Regenerate extension widgets now that we've finished moving or placing the widget.
@@ -2089,6 +2209,7 @@ FReply SDesignerView::OnDrop(const FGeometry& MyGeometry, const FDragDropEvent&
 
 		DesignerMessage = EDesignerMessage::None;
 
+		DropPreviews.Empty();
 		return FReply::Handled();
 	}
 
diff --git a/Engine/Source/Editor/UMGEditor/Private/Designer/SDesignerView.h b/Engine/Source/Editor/UMGEditor/Private/Designer/SDesignerView.h
--- a/Engine/Source/Editor/UMGEditor/Private/Designer/SDesignerView.h
+++ b/Engine/Source/Editor/UMGEditor/Private/Designer/SDesignerView.h
@@ -214,12 +214,14 @@ private:
 	bool CanSetTransformMode(ETransformMode::Type InTransformMode) const;
 	bool IsTransformModeActive(ETransformMode::Type InTransformMode) const;
 
-	UWidget* ProcessDropAndAddWidget(const FGeometry& MyGeometry, const FDragDropEvent& DragDropEvent, bool bIsPreview);
+	void ProcessDropAndAddWidget(const FGeometry& MyGeometry, const FDragDropEvent& DragDropEvent, const bool bIsPreview);
 
 	FVector2D GetExtensionPosition(TSharedRef<FDesignerSurfaceElement> ExtensionElement) const;
 
 	FVector2D GetExtensionSize(TSharedRef<FDesignerSurfaceElement> ExtensionElement) const;
 	
+	void ClearDropPreviews();
+
 private:
 	/** A reference to the BP Editor that owns this designer */
 	TWeakPtr<FWidgetBlueprintEditor> BlueprintEditor;
@@ -236,8 +238,13 @@ private:
 	/** The current preview widget's slate widget */
 	TWeakPtr<SWidget> PreviewSlateWidget;
 	
-	UWidget* DropPreviewWidget;
-	UPanelWidget* DropPreviewParent;
+	struct FDropPreview
+	{
+		UWidget* Widget;
+		UPanelWidget* Parent;
+	};
+
+	TArray<FDropPreview> DropPreviews;
 
 	TSharedPtr<class SZoomPan> PreviewHitTestRoot;
 	TSharedPtr<SBox> PreviewAreaConstraint;
