MovieScene: Scaffolding for UMB bindings refactor

diff --git a/Engine/Source/Editor/UMGEditor/Private/Animation/UMGBindingManager.cpp b/Engine/Source/Editor/UMGEditor/Private/Animation/UMGBindingManager.cpp
--- a/Engine/Source/Editor/UMGEditor/Private/Animation/UMGBindingManager.cpp
+++ b/Engine/Source/Editor/UMGEditor/Private/Animation/UMGBindingManager.cpp
@@ -4,6 +4,12 @@
 #include "MovieScene.h"
 #include "MovieSceneInstance.h"
 #include "UMGBindingManager.h"
+#include "Animation/WidgetAnimation.h"
+#include "WidgetBlueprint.h"
+#include "WidgetBlueprintEditor.h"
+
+
+#define LOCTEXT_NAMESPACE "UUMGBindingManager"
 
 
 /* UUMGBindingManager structors
@@ -11,9 +17,74 @@
 
 UUMGBindingManager::UUMGBindingManager(const FObjectInitializer& ObjectInitializer)
 	: Super(ObjectInitializer)
+	, WidgetAnimation(nullptr)
+	, WidgetBlueprintEditor(nullptr)
 { }
 
 
+/* UUMGBindingManager interface
+ *****************************************************************************/
+
+void UUMGBindingManager::InitPreviewObjects()
+{
+	if ((WidgetAnimation == nullptr) || (WidgetBlueprintEditor == nullptr))
+	{
+		return;
+	}
+
+	// Populate preview object to guid map
+	UUserWidget* PreviewWidget = WidgetBlueprintEditor->GetPreview();
+
+	if (PreviewWidget == nullptr)
+	{
+		return;
+	}
+
+	UWidgetTree* WidgetTree = PreviewWidget->WidgetTree;
+
+	for (const FWidgetAnimationBinding& Binding : WidgetAnimation->AnimationBindings)
+	{
+		UObject* FoundObject = Binding.FindRuntimeObject(*WidgetTree);
+
+		if (FoundObject == nullptr)
+		{
+			continue;
+		}
+
+		UPanelSlot* FoundSlot = Cast<UPanelSlot>(FoundObject);
+
+		if (FoundSlot == nullptr)
+		{
+			ObjectIdToPreviewObjects.Add(Binding.AnimationId, FoundObject);
+			PreviewObjectToId.Add(FoundObject, Binding.AnimationId);
+		}
+		else
+		{
+			ObjectIdToSlotContentPreviewObjects.Add(Binding.AnimationId, FoundSlot->Content);
+			SlotContentPreviewObjectToId.Add(FoundSlot->Content, Binding.AnimationId);
+		}
+	}
+}
+
+
+void UUMGBindingManager::BindWidgetBlueprintEditor(FWidgetBlueprintEditor& InWidgetBlueprintEditor, UWidgetAnimation& InAnimation)
+{
+	WidgetAnimation = &InAnimation;
+	WidgetBlueprintEditor = &InWidgetBlueprintEditor;
+	WidgetBlueprintEditor->GetOnWidgetPreviewUpdated().AddUObject(this, &UUMGBindingManager::HandleWidgetPreviewUpdated);
+}
+
+
+void UUMGBindingManager::UnbindWidgetBlueprintEditor()
+{
+	if (WidgetBlueprintEditor != nullptr)
+	{
+		WidgetBlueprintEditor->GetOnWidgetPreviewUpdated().RemoveAll(this);
+	}
+}
+
+
+
 /* IMovieSceneBindingManager interface
  *****************************************************************************/
 
@@ -25,53 +96,188 @@ bool UUMGBindingManager::AllowsSpawnableObjects() const
 
 void UUMGBindingManager::BindPossessableObject(const FMovieSceneObjectId& ObjectId, UObject& PossessedObject)
 {
-	// @todo gmp: sequencer: bind possessable object
+	if (WidgetAnimation == nullptr)
+	{
+		return;
+	}
+
+	UPanelSlot* PossessedSlot = Cast<UPanelSlot>(&PossessedObject);
+
+	if ((PossessedSlot != nullptr) && (PossessedSlot->Content != nullptr))
+	{
+		SlotContentPreviewObjectToId.Add(PossessedSlot->Content, ObjectId);
+		ObjectIdToSlotContentPreviewObjects.Add(ObjectId, PossessedSlot->Content);
+
+		// Save the name of the widget containing the slots. This is the object
+		// to look up that contains the slot itself (the thing we are animating).
+		FWidgetAnimationBinding NewBinding;
+		{
+			NewBinding.AnimationId = ObjectId;
+			NewBinding.SlotWidgetName = PossessedSlot->GetFName();
+			NewBinding.WidgetName = PossessedSlot->Content->GetFName();
+		}
+
+		WidgetAnimation->AnimationBindings.Add(NewBinding);
+	}
+	else if (PossessedSlot == nullptr)
+	{
+		PreviewObjectToId.Add(&PossessedObject, ObjectId);
+		ObjectIdToPreviewObjects.Add(ObjectId, &PossessedObject);
+
+		FWidgetAnimationBinding NewBinding;
+		{
+			NewBinding.AnimationId = ObjectId;
+			NewBinding.WidgetName = PossessedObject.GetFName();
+		}
+
+		WidgetAnimation->AnimationBindings.Add(NewBinding);
+	}
 }
 
 
 bool UUMGBindingManager::CanPossessObject(UObject& Object) const
 {
-	return false;
-}
+	if (WidgetBlueprintEditor == nullptr)
+	{
+		return false;
+	}
 
+	UPanelSlot* Slot = Cast<UPanelSlot>(&Object);
 
-void UUMGBindingManager::DestroyAllSpawnedObjects()
-{
+	if ((Slot != nullptr) && (Slot->Content == nullptr))
+	{
+		// can't possess empty slots.
+		return false;
+	}
+
+	UWidgetBlueprint* WidgetBlueprint = WidgetBlueprintEditor->GetWidgetBlueprintObj();
+	UUserWidget* PreviewWidget = WidgetBlueprintEditor->GetPreview();
 
+	return (Object.IsA<UVisual>() && Object.IsIn(PreviewWidget));
 }
 
 
 FMovieSceneObjectId UUMGBindingManager::FindIdForObject(const UMovieScene& MovieScene, UObject& Object) const
 {
-	return FMovieSceneObjectId();
+	UPanelSlot* Slot = Cast<UPanelSlot>(&Object);
+
+	if (Slot != nullptr)
+	{
+		// slot guids are tracked by their content.
+		return PreviewObjectToId.FindRef(Slot->Content);
+	}
+
+	return PreviewObjectToId.FindRef(&Object);
 }
 
 
 void UUMGBindingManager::GetRuntimeObjects(const TSharedRef<FMovieSceneInstance>& MovieSceneInstance, const FMovieSceneObjectId& ObjectId, TArray<UObject*>& OutRuntimeObjects) const
 {
+	TArray<TWeakObjectPtr<UObject>> PreviewObjects;
+	ObjectIdToPreviewObjects.MultiFind(ObjectId, PreviewObjects);
+
+	for (TWeakObjectPtr<UObject>& PreviewObject : PreviewObjects)
+	{
+		OutRuntimeObjects.Add(PreviewObject.Get());
+	}
 
+	TArray<TWeakObjectPtr<UObject>> SlotContentPreviewObjects;
+	ObjectIdToSlotContentPreviewObjects.MultiFind(ObjectId, SlotContentPreviewObjects);
+
+	for (TWeakObjectPtr<UObject>& SlotContentPreviewObject : SlotContentPreviewObjects)
+	{
+		UWidget* ContentWidget = Cast<UWidget>(SlotContentPreviewObject.Get());
+		OutRuntimeObjects.Add(ContentWidget->Slot);
+	}
 }
 
 
-void UUMGBindingManager::RemoveMovieSceneInstance(const TSharedRef<FMovieSceneInstance>& MovieSceneInstance)
+bool UUMGBindingManager::TryGetObjectBindingDisplayName(const TSharedRef<FMovieSceneInstance>& MovieSceneInstance, const FMovieSceneObjectId& ObjectId, FText& OutDisplayName) const
 {
+	// TODO: This gets called every frame for every bound object and could be a potential performance issue for a really complicated animation.
+	TArray<TWeakObjectPtr<UObject>> BindingObjects;
+	ObjectIdToPreviewObjects.MultiFind(ObjectId, BindingObjects);
+	
+	TArray<TWeakObjectPtr<UObject>> SlotContentBindingObjects;
+	ObjectIdToSlotContentPreviewObjects.MultiFind(ObjectId, SlotContentBindingObjects);
+	
+	if ((BindingObjects.Num() == 0) && (SlotContentBindingObjects.Num() == 0))
+	{
+		OutDisplayName = LOCTEXT("NoBoundObjects", "No bound objects");
+	}
+	else if (BindingObjects.Num() + SlotContentBindingObjects.Num() > 1)
+	{
+		OutDisplayName = LOCTEXT("Multiple bound objects", "Multilple bound objects");
+	}
+	else if (BindingObjects.Num() == 1)
+	{
+		OutDisplayName = FText::FromString(BindingObjects[0].Get()->GetName());
+	}
+	else // SlotContentBindingObjects.Num() == 1
+	{
+		UWidget* SlotContent = Cast<UWidget>(SlotContentBindingObjects[0].Get());
+		FText PanelName = SlotContent->Slot != nullptr && SlotContent->Slot->Parent != nullptr
+			? FText::FromString(SlotContent->Slot->Parent->GetName())
+			: LOCTEXT("InvalidPanel", "Invalid Panel");
+		FText ContentName = FText::FromString(SlotContent->GetName());
+		OutDisplayName = FText::Format(LOCTEXT("SlotObject", "{0} ({1} Slot)"), ContentName, PanelName);
+	}
 
+	return true;
 }
 
 
-void UUMGBindingManager::SpawnOrDestroyObjectsForInstance(TSharedRef<FMovieSceneInstance> MovieSceneInstance, bool bDestroyAll)
+void UUMGBindingManager::UnbindPossessableObjects(const FMovieSceneObjectId& ObjectId)
 {
+	TArray<TWeakObjectPtr<UObject>> PreviewObjects;
+	ObjectIdToPreviewObjects.MultiFind(ObjectId, PreviewObjects);
 
-}
+	for (TWeakObjectPtr<UObject>& PreviewObject : PreviewObjects)
+	{
+		PreviewObjectToId.Remove(PreviewObject);
+	}
 
+	ObjectIdToPreviewObjects.Remove(ObjectId);
 
-bool UUMGBindingManager::TryGetObjectBindingDisplayName(const TSharedRef<FMovieSceneInstance>& MovieSceneInstance, const FMovieSceneObjectId& ObjectId, FText& OutDisplayName) const
-{
-	return false;
+	TArray<TWeakObjectPtr<UObject>> SlotContentPreviewObjects;
+	ObjectIdToSlotContentPreviewObjects.MultiFind(ObjectId, SlotContentPreviewObjects);
+
+	for (TWeakObjectPtr<UObject>& SlotContentPreviewObject : SlotContentPreviewObjects)
+	{
+		SlotContentPreviewObjectToId.Remove(SlotContentPreviewObject);
+	}
+
+	ObjectIdToSlotContentPreviewObjects.Remove(ObjectId);
+
+	if ((WidgetAnimation != nullptr) && (WidgetBlueprintEditor != nullptr))
+	{
+		UWidgetBlueprint* WidgetBlueprint = WidgetBlueprintEditor->GetWidgetBlueprintObj();
+
+		WidgetAnimation->Modify();
+		WidgetAnimation->AnimationBindings.RemoveAll([&](const FWidgetAnimationBinding& Binding) {
+			return Binding.AnimationId == ObjectId;
+		});
+	}
 }
 
 
-void UUMGBindingManager::UnbindPossessableObjects(const FMovieSceneObjectId& ObjectId)
+/* UUMGBindingManager event handlers
+ *****************************************************************************/
+
+void UUMGBindingManager::HandleWidgetPreviewUpdated()
 {
-	// @todo gmp: sequencer: bind possessable object
+	PreviewObjectToId.Empty();
+	ObjectIdToPreviewObjects.Empty();
+	ObjectIdToSlotContentPreviewObjects.Empty();
+	SlotContentPreviewObjectToId.Empty();
+
+	InitPreviewObjects();
+
+	if (WidgetBlueprintEditor != nullptr)
+	{
+		WidgetBlueprintEditor->GetSequencer()->UpdateRuntimeInstances();
+	}
 }
+
+
+#undef LOCTEXT_NAMESPACE
diff --git a/Engine/Source/Editor/UMGEditor/Private/Animation/UMGBindingManager.h b/Engine/Source/Editor/UMGEditor/Private/Animation/UMGBindingManager.h
--- a/Engine/Source/Editor/UMGEditor/Private/Animation/UMGBindingManager.h
+++ b/Engine/Source/Editor/UMGEditor/Private/Animation/UMGBindingManager.h
@@ -6,6 +6,11 @@
 #include "UMGBindingManager.generated.h"
 
 
+class FWidgetBlueprintEditor;
+class UObject;
+class UWidgetAnimation;
+
+
 /**
  * Movie scene binding manager used by Sequencer.
  */
@@ -16,6 +21,42 @@ class UUMGBindingManager
 {
 	GENERATED_UCLASS_BODY()
 
+public:
+
+	/**
+	 * Get the current widget animation.
+	 *
+	 * @return Widget animation object, or nullptr if not valid.
+	 * @see HasValidWidgetAnimation
+	 */
+	UWidgetAnimation* GetWidgetAnimation()
+	{
+		return WidgetAnimation;
+	}
+
+	/** Check whether the current animation is valid.
+	 *
+	 * @return true if the animation is valid, false otherwise.
+	 * @see GetWidgetAnimation
+	 */
+	bool HasValidWidgetAnimation() const;
+
+	/** Rebuild the mappings to live preview objects. */
+	void InitPreviewObjects();
+
+	/**
+	 * Bind the binding manager the given widget BP editor.
+	 *
+	 * @param InWidgetBlueprintEditor The widget BP editor to bind to.
+	 * @param InAnimation The current animation to set.
+	 */
+	// @todo sequencer: gmp: call this from some appropriate place
+	void BindWidgetBlueprintEditor(FWidgetBlueprintEditor& InWidgetBlueprintEditor, UWidgetAnimation& InAnimation);
+
+	/** Unbind from the widget BP editor. */
+	// @todo sequencer: gmp: call this from some appropriate place
+	void UnbindWidgetBlueprintEditor();
+
 public:
 
 	// IMovieSceneBindingManager interface
@@ -23,11 +64,31 @@ public:
 	virtual bool AllowsSpawnableObjects() const override;
 	virtual void BindPossessableObject(const FMovieSceneObjectId& ObjectId, UObject& PossessedObject) override;
 	virtual bool CanPossessObject( UObject& Object ) const override;
-	virtual void DestroyAllSpawnedObjects() override;
+	virtual void DestroyAllSpawnedObjects() override { }
 	virtual FMovieSceneObjectId FindIdForObject(const UMovieScene& MovieScene, UObject& Object) const override;
 	virtual void GetRuntimeObjects(const TSharedRef<FMovieSceneInstance>& MovieSceneInstance, const FMovieSceneObjectId& ObjectId, TArray<UObject*>& OutRuntimeObjects) const override;
-	virtual void RemoveMovieSceneInstance(const TSharedRef<FMovieSceneInstance>& MovieSceneInstance) override;
-	virtual void SpawnOrDestroyObjectsForInstance(TSharedRef<FMovieSceneInstance> MovieSceneInstance, bool bDestroyAll) override;
+	virtual void RemoveMovieSceneInstance(const TSharedRef<FMovieSceneInstance>& MovieSceneInstance) override { }
+	virtual void SpawnOrDestroyObjectsForInstance(TSharedRef<FMovieSceneInstance> MovieSceneInstance, bool bDestroyAll) override { }
 	virtual bool TryGetObjectBindingDisplayName(const TSharedRef<FMovieSceneInstance>& MovieSceneInstance, const FMovieSceneObjectId& ObjectId, FText& OutDisplayName) const override;
 	virtual void UnbindPossessableObjects(const FMovieSceneObjectId& ObjectId) override;
+
+private:
+
+	/** Handles widget preview updates in the current widget BP editor. */
+	void HandleWidgetPreviewUpdated();
+
+private:
+
+	/** Mapping of preview objects to sequencer guids */
+	TMap<TWeakObjectPtr<UObject>, FMovieSceneObjectId> PreviewObjectToId;
+	TMultiMap<FMovieSceneObjectId, TWeakObjectPtr<UObject>> ObjectIdToPreviewObjects;
+	TMultiMap<FMovieSceneObjectId, TWeakObjectPtr<UObject>> ObjectIdToSlotContentPreviewObjects;
+	TMap<TWeakObjectPtr<UObject>, FMovieSceneObjectId> SlotContentPreviewObjectToId;
+
+	/** The current widget animation. */
+	UPROPERTY(transient)
+	UWidgetAnimation* WidgetAnimation;
+
+	/** Pointer to the widget BP editor. */
+	FWidgetBlueprintEditor* WidgetBlueprintEditor;
 };
diff --git a/Engine/Source/Runtime/MovieScene/Public/MovieSceneObjectId.h b/Engine/Source/Runtime/MovieScene/Public/MovieSceneObjectId.h
--- a/Engine/Source/Runtime/MovieScene/Public/MovieSceneObjectId.h
+++ b/Engine/Source/Runtime/MovieScene/Public/MovieSceneObjectId.h
@@ -46,7 +46,7 @@ public:
 	 * @param Y The second identifier to compare.
 	 * @return true if the identifiers are equal, false otherwise.
 	 */
-	friend bool operator==( const FMovieSceneObjectId& X, const FMovieSceneObjectId& Y )
+	friend bool operator==(const FMovieSceneObjectId& X, const FMovieSceneObjectId& Y)
 	{
 		return ((X.ObjectGuid == Y.ObjectGuid) && (X.ComponentName == X.ComponentName));
 	}
@@ -58,11 +58,22 @@ public:
 	 * @param Y The second identifier to compare.
 	 * @return true if the identifier are not equal, false otherwise.
 	 */
-	friend bool operator!=( const FMovieSceneObjectId& X, const FMovieSceneObjectId& Y )
+	friend bool operator!=(const FMovieSceneObjectId& X, const FMovieSceneObjectId& Y)
 	{
 		return ((X.ObjectGuid != Y.ObjectGuid) || (X.ComponentName != X.ComponentName));
 	}
 
+	/**
+	 * Calculates the hash for an object identifier.
+	 *
+	 * @param ObjectId The identifier to calculate the hash for.
+	 * @return The hash.
+	 */
+	friend uint32 GetTypeHash(const FMovieSceneObjectId& ObjectId)
+	{
+		return GetTypeHash(ObjectId.ObjectGuid) ^ GetTypeHash(ObjectId.ComponentName);
+	}
+
 	/**
 	 * Checks whether this object identifier is valid or not.
 	 *
diff --git a/Engine/Source/Runtime/UMG/Public/Animation/WidgetAnimation.h b/Engine/Source/Runtime/UMG/Public/Animation/WidgetAnimation.h
--- a/Engine/Source/Runtime/UMG/Public/Animation/WidgetAnimation.h
+++ b/Engine/Source/Runtime/UMG/Public/Animation/WidgetAnimation.h
@@ -3,6 +3,7 @@
 #pragma once
 
 #include "IMovieScenePlayer.h"
+#include "MovieSceneObjectId.h"
 #include "WidgetAnimation.generated.h"
 
 class UMovieScene;
@@ -23,6 +24,10 @@ struct FWidgetAnimationBinding
 	UPROPERTY()
 	FGuid AnimationGuid;
 
+	// @todo sequencer: gmp: UMG integration
+	//UPROPERTY()
+	FMovieSceneObjectId AnimationId;
+
 public:
 	/**
 	 * Locates a runtime object to animate from the provided tree of widgets
