Fixed undo/redo when wrapping widgets in the UMG Designer
[UE-18005] - Unable to undo a wrap action for widgets

diff --git a/Engine/Source/Editor/UMGEditor/Private/WidgetBlueprintEditorUtils.cpp b/Engine/Source/Editor/UMGEditor/Private/WidgetBlueprintEditorUtils.cpp
--- a/Engine/Source/Editor/UMGEditor/Private/WidgetBlueprintEditorUtils.cpp
+++ b/Engine/Source/Editor/UMGEditor/Private/WidgetBlueprintEditorUtils.cpp
@@ -251,7 +251,7 @@ void FWidgetBlueprintEditorUtils::CreateWidgetContextMenu(FMenuBuilder& MenuBuil
 		MenuBuilder.AddSubMenu(
 			LOCTEXT("WidgetTree_WrapWith", "Wrap With..."),
 			LOCTEXT("WidgetTree_WrapWithToolTip", "Wraps the currently selected widgets inside of another container widget"),
-			FNewMenuDelegate::CreateStatic(&FWidgetBlueprintEditorUtils::BuildWrapWithMenu, BP, Widgets)
+			FNewMenuDelegate::CreateStatic(&FWidgetBlueprintEditorUtils::BuildWrapWithMenu, BlueprintEditor, BP, Widgets)
 			);
 
 		if ( Widgets.Num() == 1 )
@@ -378,7 +378,7 @@ bool FWidgetBlueprintEditorUtils::FindAndRemoveNamedSlotContent(UWidget* WidgetT
 	return bSuccess;
 }
 
-void FWidgetBlueprintEditorUtils::BuildWrapWithMenu(FMenuBuilder& Menu, UWidgetBlueprint* BP, TSet<FWidgetReference> Widgets)
+void FWidgetBlueprintEditorUtils::BuildWrapWithMenu(FMenuBuilder& Menu, TSharedRef<FWidgetBlueprintEditor> BlueprintEditor, UWidgetBlueprint* BP, TSet<FWidgetReference> Widgets)
 {
 	Menu.BeginSection("WrapWith", LOCTEXT("WidgetTree_WrapWith", "Wrap With..."));
 	{
@@ -394,7 +394,7 @@ void FWidgetBlueprintEditorUtils::BuildWrapWithMenu(FMenuBuilder& Menu, UWidgetB
 						FText::GetEmpty(),
 						FSlateIcon(),
 						FUIAction(
-							FExecuteAction::CreateStatic(&FWidgetBlueprintEditorUtils::WrapWidgets, BP, Widgets, WidgetClass),
+						FExecuteAction::CreateStatic(&FWidgetBlueprintEditorUtils::WrapWidgets, BlueprintEditor, BP, Widgets, WidgetClass),
 							FCanExecuteAction()
 						));
 				}
@@ -404,7 +404,7 @@ void FWidgetBlueprintEditorUtils::BuildWrapWithMenu(FMenuBuilder& Menu, UWidgetB
 	Menu.EndSection();
 }
 
-void FWidgetBlueprintEditorUtils::WrapWidgets(UWidgetBlueprint* BP, TSet<FWidgetReference> Widgets, UClass* WidgetClass)
+void FWidgetBlueprintEditorUtils::WrapWidgets(TSharedRef<FWidgetBlueprintEditor> BlueprintEditor, UWidgetBlueprint* BP, TSet<FWidgetReference> Widgets, UClass* WidgetClass)
 {
 	const FScopedTransaction Transaction(LOCTEXT("WrapWidgets", "Wrap Widgets"));
 
@@ -413,35 +413,42 @@ void FWidgetBlueprintEditorUtils::WrapWidgets(UWidgetBlueprint* BP, TSet<FWidget
 	// Old Parent -> New Parent Map
 	TMap<UPanelWidget*, UPanelWidget*> OldParentToNewParent;
 
-	for ( FWidgetReference& Item : Widgets )
+	for (FWidgetReference& Item : Widgets)
 	{
 		int32 OutIndex;
-		UPanelWidget* CurrentParent = BP->WidgetTree->FindWidgetParent(Item.GetTemplate(), OutIndex);
+		UWidget* Widget = Item.GetTemplate();
+		UPanelWidget* CurrentParent = BP->WidgetTree->FindWidgetParent(Widget, OutIndex);
 
 		// If the widget doesn't currently have a parent, and isn't the root, ignore it.
-		if ( CurrentParent == nullptr && Item.GetTemplate() != BP->WidgetTree->RootWidget )
+		if (CurrentParent == nullptr && Widget != BP->WidgetTree->RootWidget)
 		{
 			continue;
 		}
 
+		Widget->Modify();
+
 		UPanelWidget*& NewWrapperWidget = OldParentToNewParent.FindOrAdd(CurrentParent);
-		if ( NewWrapperWidget == nullptr || !NewWrapperWidget->CanAddMoreChildren() )
+		if (NewWrapperWidget == nullptr || !NewWrapperWidget->CanAddMoreChildren())
 		{
 			NewWrapperWidget = CastChecked<UPanelWidget>(Template->Create(BP->WidgetTree));
+			NewWrapperWidget->SetDesignerFlags(BlueprintEditor->GetCurrentDesignerFlags());
 
-			if ( CurrentParent )
+			BP->WidgetTree->SetFlags(RF_Transactional);
+			BP->WidgetTree->Modify();
+
+			if (CurrentParent)
 			{
+				CurrentParent->SetFlags(RF_Transactional);
 				CurrentParent->Modify();
 				CurrentParent->ReplaceChildAt(OutIndex, NewWrapperWidget);
 			}
 			else // Root Widget
 			{
-				BP->WidgetTree->Modify();
 				BP->WidgetTree->RootWidget = NewWrapperWidget;
 			}
 		}
 
-		NewWrapperWidget->AddChild(Item.GetTemplate());
+		NewWrapperWidget->AddChild(Widget);
 	}
 
 	FBlueprintEditorUtils::MarkBlueprintAsStructurallyModified(BP);
diff --git a/Engine/Source/Editor/UMGEditor/Private/WidgetBlueprintEditorUtils.h b/Engine/Source/Editor/UMGEditor/Private/WidgetBlueprintEditorUtils.h
--- a/Engine/Source/Editor/UMGEditor/Private/WidgetBlueprintEditorUtils.h
+++ b/Engine/Source/Editor/UMGEditor/Private/WidgetBlueprintEditorUtils.h
@@ -48,9 +48,9 @@ private:
 
 	static bool FindAndRemoveNamedSlotContent(UWidget* WidgetTemplate, UWidgetTree* WidgetTree);
 
-	static void BuildWrapWithMenu(FMenuBuilder& Menu, UWidgetBlueprint* BP, TSet<FWidgetReference> Widgets);
+	static void BuildWrapWithMenu(FMenuBuilder& Menu, TSharedRef<FWidgetBlueprintEditor> BlueprintEditor, UWidgetBlueprint* BP, TSet<FWidgetReference> Widgets);
 
-	static void WrapWidgets(UWidgetBlueprint* BP, TSet<FWidgetReference> Widgets, UClass* WidgetClass);
+	static void WrapWidgets(TSharedRef<FWidgetBlueprintEditor> BlueprintEditor, UWidgetBlueprint* BP, TSet<FWidgetReference> Widgets, UClass* WidgetClass);
 
 	static void BuildReplaceWithMenu(FMenuBuilder& Menu, UWidgetBlueprint* BP, TSet<FWidgetReference> Widgets);
 
