Moving widgets in the UMG Designer no longer makes them move to the last slot in their parent.
[UE-13346] - Last moved widget in graph reorders to bottom of hierarchy

diff --git a/Engine/Source/Editor/UMGEditor/Private/Designer/SDesignerView.cpp b/Engine/Source/Editor/UMGEditor/Private/Designer/SDesignerView.cpp
--- a/Engine/Source/Editor/UMGEditor/Private/Designer/SDesignerView.cpp
+++ b/Engine/Source/Editor/UMGEditor/Private/Designer/SDesignerView.cpp
@@ -2050,26 +2050,30 @@ void SDesignerView::ProcessDropAndAddWidget(const FGeometry& MyGeometry, const F
 				bWidgetMoved = true;
 				UPanelWidget* NewParent = Cast<UPanelWidget>(Target);
 
-				// If this isn't a preview operation we need to modify a few things to properly undo the operation.
-				if (!bIsPreview)
-				{
-					NewParent->SetFlags(RF_Transactional);
-					NewParent->Modify();
-
-					BP->WidgetTree->SetFlags(RF_Transactional);
-					BP->WidgetTree->Modify();
-				}
-
 				UWidget* Widget = bIsPreview ? DraggedWidget.Preview : DraggedWidget.Template;
 
 				if (ensure(Widget))
 				{
+					bool bIsChangingParent = Widget->GetParent() != NewParent;
+
+					check(Widget->GetParent() != nullptr || bIsChangingParent);
+
+					// If this isn't a preview operation we need to modify a few things to properly undo the operation.
 					if (!bIsPreview)
 					{
+						if (bIsChangingParent)
+						{
+							NewParent->SetFlags(RF_Transactional);
+							NewParent->Modify();
+
+							BP->WidgetTree->SetFlags(RF_Transactional);
+							BP->WidgetTree->Modify();
+						}
+
 						Widget->Modify();
 					}
 
-					if (Widget->GetParent())
+					if (Widget->GetParent() && bIsChangingParent)
 					{
 						if (!bIsPreview)
 						{
@@ -2110,7 +2114,20 @@ void SDesignerView::ProcessDropAndAddWidget(const FGeometry& MyGeometry, const F
 					}
 
 					FVector2D LocalPosition = WidgetUnderCursorGeometry.AbsoluteToLocal(ScreenSpacePosition);
-					if (UPanelSlot* Slot = NewParent->AddChild(Widget))
+
+					UPanelSlot* Slot;
+					if (bIsChangingParent)
+					{
+						Slot = NewParent->AddChild(Widget);
+					}
+					else
+					{
+						check(Widget->GetParent()->GetChildIndex(Widget) != INDEX_NONE);
+
+						Slot = Widget->GetParent()->GetSlots()[Widget->GetParent()->GetChildIndex(Widget)];
+					}
+
+					if (Slot != nullptr)
 					{
 						FVector2D NewPosition = LocalPosition - DraggedWidget.DraggedOffset;
 
