forgot to check in

diff --git a/Engine/Source/Editor/UMGEditor/UMGEditor.Build.cs b/Engine/Source/Editor/UMGEditor/UMGEditor.Build.cs
--- a/Engine/Source/Editor/UMGEditor/UMGEditor.Build.cs
+++ b/Engine/Source/Editor/UMGEditor/UMGEditor.Build.cs
@@ -49,6 +49,7 @@ public class UMGEditor : ModuleRules
 				"SlateCore",
 				"MovieScene",
 				"MovieSceneTools",
+                "MovieSceneTracks",
 				"Sequencer",
 				"DetailCustomizations",
                 "Settings",
