UMG: Moved animation binding class into its own files

diff --git a/Engine/Source/Editor/UMGEditor/Private/Animation/UMGMovieSceneObjectManager.cpp b/Engine/Source/Editor/UMGEditor/Private/Animation/UMGMovieSceneObjectManager.cpp
--- a/Engine/Source/Editor/UMGEditor/Private/Animation/UMGMovieSceneObjectManager.cpp
+++ b/Engine/Source/Editor/UMGEditor/Private/Animation/UMGMovieSceneObjectManager.cpp
@@ -54,13 +54,13 @@ void UUMGMovieSceneObjectManager::InitPreviewObjects()
 
 		if (FoundSlot == nullptr)
 		{
-			IdToPreviewObjects.Add(Binding.AnimationId, FoundObject);
-			PreviewObjectToIds.Add(FoundObject, Binding.AnimationId);
+			IdToPreviewObjects.Add(Binding.AnimationGuid, FoundObject);
+			PreviewObjectToIds.Add(FoundObject, Binding.AnimationGuid);
 		}
 		else
 		{
-			IdToSlotContentPreviewObjects.Add(Binding.AnimationId, FoundSlot->Content);
-			SlotContentPreviewObjectToIds.Add(FoundSlot->Content, Binding.AnimationId);
+			IdToSlotContentPreviewObjects.Add(Binding.AnimationGuid, FoundSlot->Content);
+			SlotContentPreviewObjectToIds.Add(FoundSlot->Content, Binding.AnimationGuid);
 		}
 	}
 }
@@ -111,7 +111,7 @@ void UUMGMovieSceneObjectManager::BindPossessableObject(const FGuid& ObjectId, U
 		// to look up that contains the slot itself (the thing we are animating).
 		FWidgetAnimationBinding NewBinding;
 		{
-			NewBinding.AnimationId = ObjectId;
+			NewBinding.AnimationGuid = ObjectId;
 			NewBinding.SlotWidgetName = PossessedSlot->GetFName();
 			NewBinding.WidgetName = PossessedSlot->Content->GetFName();
 		}
@@ -125,7 +125,7 @@ void UUMGMovieSceneObjectManager::BindPossessableObject(const FGuid& ObjectId, U
 
 		FWidgetAnimationBinding NewBinding;
 		{
-			NewBinding.AnimationId = ObjectId;
+			NewBinding.AnimationGuid = ObjectId;
 			NewBinding.WidgetName = PossessedObject.GetFName();
 		}
 
@@ -252,7 +252,7 @@ void UUMGMovieSceneObjectManager::UnbindPossessableObjects(const FGuid& ObjectId
 
 		WidgetAnimation->Modify();
 		WidgetAnimation->AnimationBindings.RemoveAll([&](const FWidgetAnimationBinding& Binding) {
-			return Binding.AnimationId == ObjectId;
+			return Binding.AnimationGuid == ObjectId;
 		});
 	}
 }
diff --git a/Engine/Source/Runtime/UMG/Private/Animation/WidgetAnimation.cpp b/Engine/Source/Runtime/UMG/Private/Animation/WidgetAnimation.cpp
--- a/Engine/Source/Runtime/UMG/Private/Animation/WidgetAnimation.cpp
+++ b/Engine/Source/Runtime/UMG/Private/Animation/WidgetAnimation.cpp
@@ -4,37 +4,25 @@
 #include "MovieScene.h"
 #include "WidgetAnimation.h"
 
-UObject* FWidgetAnimationBinding::FindRuntimeObject( UWidgetTree& WidgetTree ) const
-{
-	UObject* FoundObject = FindObject<UObject>(&WidgetTree, *WidgetName.ToString());
-	if(FoundObject)
-	{
-		if(SlotWidgetName != NAME_None)
-		{
-			// if we were animating the slot, look up the slot that contains the widget 
-			UWidget* WidgetObject = Cast<UWidget>(FoundObject);
-			if(WidgetObject->Slot)
-			{
-				FoundObject = WidgetObject->Slot;
-			}
-		}
-	}
 
-	return FoundObject;
-}
+/* UWidgetAnimation structors
+ *****************************************************************************/
 
 UWidgetAnimation::UWidgetAnimation(const FObjectInitializer& ObjectInitializer)
 	: Super(ObjectInitializer)
-{
+{ }
 
-}
+
+/* UWidgetAnimation interface
+ *****************************************************************************/
 
 #if WITH_EDITOR
 
 UWidgetAnimation* UWidgetAnimation::GetNullAnimation()
 {
 	static UWidgetAnimation* NullAnimation = nullptr;
-	if( !NullAnimation )
+
+	if (!NullAnimation)
 	{
 		NullAnimation = NewObject<UWidgetAnimation>(GetTransientPackage(), NAME_None, RF_RootSet);
 		NullAnimation->MovieScene = NewObject<UMovieScene>(NullAnimation, FName("No Animation"), RF_RootSet);
@@ -45,11 +33,13 @@ UWidgetAnimation* UWidgetAnimation::GetNullAnimation()
 
 #endif
 
+
 float UWidgetAnimation::GetStartTime() const
 {
 	return MovieScene->GetTimeRange().GetLowerBoundValue();
 }
 
+
 float UWidgetAnimation::GetEndTime() const
 {
 	return MovieScene->GetTimeRange().GetUpperBoundValue();
diff --git a/Engine/Source/Runtime/UMG/Private/Animation/WidgetAnimationBinding.cpp b/Engine/Source/Runtime/UMG/Private/Animation/WidgetAnimationBinding.cpp
new file mode 100644
--- /dev/null
+++ b/Engine/Source/Runtime/UMG/Private/Animation/WidgetAnimationBinding.cpp
@@ -0,0 +1,26 @@
+// Copyright 1998-2015 Epic Games, Inc. All Rights Reserved.
+
+#include "UMGPrivatePCH.h"
+#include "WidgetAnimationBinding.h"
+
+
+/* FWidgetAnimationBinding interface
+ *****************************************************************************/
+
+UObject* FWidgetAnimationBinding::FindRuntimeObject( UWidgetTree& WidgetTree ) const
+{
+	UObject* FoundObject = FindObject<UObject>(&WidgetTree, *WidgetName.ToString());
+
+	if (FoundObject && (SlotWidgetName != NAME_None))
+	{
+		// if we were animating the slot, look up the slot that contains the widget 
+		UWidget* WidgetObject = Cast<UWidget>(FoundObject);
+
+		if (WidgetObject->Slot)
+		{
+			FoundObject = WidgetObject->Slot;
+		}
+	}
+
+	return FoundObject;
+}
diff --git a/Engine/Source/Runtime/UMG/Public/Animation/WidgetAnimation.h b/Engine/Source/Runtime/UMG/Public/Animation/WidgetAnimation.h
--- a/Engine/Source/Runtime/UMG/Public/Animation/WidgetAnimation.h
+++ b/Engine/Source/Runtime/UMG/Public/Animation/WidgetAnimation.h
@@ -3,54 +3,17 @@
 #pragma once
 
 #include "IMovieScenePlayer.h"
+#include "WidgetAnimationBinding.h"
 #include "WidgetAnimation.generated.h"
 
+
 class UMovieScene;
 class UWidgetTree;
 
-/** A single object bound to a umg sequence */
-USTRUCT()
-struct FWidgetAnimationBinding
-{
-	GENERATED_USTRUCT_BODY()
-
-	UPROPERTY()
-	FName WidgetName;
-
-	UPROPERTY()
-	FName SlotWidgetName;
-
-	UPROPERTY()
-	FGuid AnimationGuid;
-
-	// @todo sequencer: gmp: UMG integration
-	//UPROPERTY()
-	FGuid AnimationId;
-
-public:
-	/**
-	 * Locates a runtime object to animate from the provided tree of widgets
-	 * @return the runtime object to animate or null if not found 
-	 */
-	UMG_API UObject* FindRuntimeObject(UWidgetTree& WidgetTree) const;
-
-	bool operator==(const FWidgetAnimationBinding& Other) const
-	{
-		return WidgetName == Other.WidgetName && SlotWidgetName == Other.SlotWidgetName && AnimationGuid == Other.AnimationGuid;
-	}
 
-	friend FArchive& operator<<(FArchive& Ar, FWidgetAnimationBinding& Binding)
-	{
-		Ar << Binding.WidgetName;
-		Ar << Binding.SlotWidgetName;
-		Ar << Binding.AnimationGuid;
-		return Ar;
-	}
-};
-
-
-UCLASS( BlueprintType, MinimalAPI )
-class UWidgetAnimation : public UObject
+UCLASS(BlueprintType, MinimalAPI)
+class UWidgetAnimation
+	: public UObject
 {
 	GENERATED_UCLASS_BODY()
 
@@ -67,6 +30,7 @@ public:
 	UMG_API float GetEndTime() const;
 
 public:
+
 	UPROPERTY()
 	UMovieScene* MovieScene;
 
diff --git a/Engine/Source/Runtime/UMG/Public/Animation/WidgetAnimationBinding.h b/Engine/Source/Runtime/UMG/Public/Animation/WidgetAnimationBinding.h
new file mode 100644
--- /dev/null
+++ b/Engine/Source/Runtime/UMG/Public/Animation/WidgetAnimationBinding.h
@@ -0,0 +1,58 @@
+// Copyright 1998-2015 Epic Games, Inc. All Rights Reserved.
+
+#pragma once
+
+#include "WidgetAnimationBinding.generated.h"
+
+
+/**
+ * A single object bound to a UMG sequence.
+ */
+USTRUCT()
+struct FWidgetAnimationBinding
+{
+	GENERATED_USTRUCT_BODY()
+
+	UPROPERTY()
+	FName WidgetName;
+
+	UPROPERTY()
+	FName SlotWidgetName;
+
+	UPROPERTY()
+	FGuid AnimationGuid;
+
+public:
+
+	/**
+	 * Locates a runtime object to animate from the provided tree of widgets
+	 * @return the runtime object to animate or null if not found 
+	 */
+	UMG_API UObject* FindRuntimeObject(UWidgetTree& WidgetTree) const;
+
+	/**
+	 * Compares two widget animation bindings for equality.
+	 *
+	 * @param X The first binding to compare.
+	 * @param Y The second binding to compare.
+	 * @return true if the bindings are equal, false otherwise.
+	 */
+	friend bool operator==(const FWidgetAnimationBinding& X, const FWidgetAnimationBinding& Y)
+	{
+		return (X.WidgetName == Y.WidgetName) && (X.SlotWidgetName == Y.SlotWidgetName) && (X.AnimationGuid == Y.AnimationGuid);
+	}
+
+	/**
+	 * Serializes a widget animation binding from/to the given archive.
+	 *
+	 * @param Ar The archive to serialize to/from.
+	 * @param Binding the binding to serialize.
+	 */
+	friend FArchive& operator<<(FArchive& Ar, FWidgetAnimationBinding& Binding)
+	{
+		Ar << Binding.WidgetName;
+		Ar << Binding.SlotWidgetName;
+		Ar << Binding.AnimationGuid;
+		return Ar;
+	}
+};
