UMG - Fix sequencer related crash.

Jira UE-17971

diff --git a/Engine/Source/Editor/UMGEditor/Private/Animation/UMGSequencerObjectBindingManager.cpp b/Engine/Source/Editor/UMGEditor/Private/Animation/UMGSequencerObjectBindingManager.cpp
--- a/Engine/Source/Editor/UMGEditor/Private/Animation/UMGSequencerObjectBindingManager.cpp
+++ b/Engine/Source/Editor/UMGEditor/Private/Animation/UMGSequencerObjectBindingManager.cpp
@@ -149,16 +149,30 @@ bool FUMGSequencerObjectBindingManager::TryGetObjectBindingDisplayName(const TSh
 	}
 	else if (BindingObjects.Num() == 1)
 	{
+		if (BindingObjects[0].IsValid())
+		{
 		DisplayName = FText::FromString(BindingObjects[0].Get()->GetName());
+	}
+		else
+		{
+			DisplayName = LOCTEXT("InvalidObject", "Invalid Object");
+		}
 	}
 	else // SlotContentBindingObjects.Num() == 1
 	{
+		if ( SlotContentBindingObjects[0].IsValid())
+		{
 		UWidget* SlotContent = Cast<UWidget>(SlotContentBindingObjects[0].Get());
 		FText PanelName = SlotContent->Slot != nullptr && SlotContent->Slot->Parent != nullptr
 			? FText::FromString(SlotContent->Slot->Parent->GetName())
 			: LOCTEXT("InvalidPanel", "Invalid Panel");
 		FText ContentName = FText::FromString(SlotContent->GetName());
 		DisplayName = FText::Format(LOCTEXT("SlotObject", "{0} ({1} Slot)"), ContentName, PanelName);
+	}
+		else
+		{
+			DisplayName = LOCTEXT("InvalidSlotContent", "Invalid slot object");
+		}
 	}
 	return true;
 }
