UMG - Fixing the warnings being generated when loading a project with widget blueprints not resolving the outer for various widgets and animations.

diff --git a/Engine/Source/Editor/UMGEditor/Private/WidgetBlueprintCompiler.cpp b/Engine/Source/Editor/UMGEditor/Private/WidgetBlueprintCompiler.cpp
--- a/Engine/Source/Editor/UMGEditor/Private/WidgetBlueprintCompiler.cpp
+++ b/Engine/Source/Editor/UMGEditor/Private/WidgetBlueprintCompiler.cpp
@@ -147,16 +147,25 @@ void FWidgetBlueprintCompiler::CleanAndSanitizeClass(UBlueprintGeneratedClass* C
 	check(ClassToClean == NewClass);
 	NewWidgetBlueprintClass = CastChecked<UWidgetBlueprintGeneratedClass>((UObject*)NewClass);
 
-	//// Purge all subobjects (properties, functions, params) of the class, as they will be regenerated
-	//TArray<UObject*> ClassSubObjects;
-	//GetObjectsWithOuter(ClassToClean, ClassSubObjects, true);
+	// Don't leave behind widget trees that are outered to the same widget blueprint rename them away first.
+	if ( NewWidgetBlueprintClass->WidgetTree )
+	{
+		NewWidgetBlueprintClass->WidgetTree->Rename(nullptr, GetTransientPackage(), REN_ForceNoResetLoaders | REN_DontCreateRedirectors);
+		NewWidgetBlueprintClass->WidgetTree = nullptr;
+	}
+	if ( NewWidgetBlueprintClass->DesignerWidgetTree )
+	{
+		NewWidgetBlueprintClass->DesignerWidgetTree->Rename(nullptr, GetTransientPackage(), REN_ForceNoResetLoaders | REN_DontCreateRedirectors);
+		NewWidgetBlueprintClass->DesignerWidgetTree = nullptr;
+	}
 
-	//UWidgetBlueprint* Blueprint = WidgetBlueprint();
+	for ( UWidgetAnimation* Animation : NewWidgetBlueprintClass->Animations )
+	{
+		Animation->Rename(nullptr, GetTransientPackage(), REN_ForceNoResetLoaders | REN_DontCreateRedirectors);
+	}
+	NewWidgetBlueprintClass->Animations.Empty();
 
-	//if ( 0 != Blueprint->WidgetTree->WidgetTemplates.Num() )
-	//{
-	//	ClassSubObjects.RemoveAllSwap(FCullTemplateObjectsHelper<UWidget>(Blueprint->WidgetTree->WidgetTemplates));
-	//}
+	NewWidgetBlueprintClass->Bindings.Empty();
 }
 
 void FWidgetBlueprintCompiler::SaveSubObjectsFromCleanAndSanitizeClass(FSubobjectCollection& SubObjectsToSave, UBlueprintGeneratedClass* ClassToClean, UObject*& OldCDO)
@@ -167,9 +176,17 @@ void FWidgetBlueprintCompiler::SaveSubObjectsFromCleanAndSanitizeClass(FSubobjec
 
 	UWidgetBlueprint* Blueprint = WidgetBlueprint();
 
+	// We need to save the widget tree to survive the initial sub-object clean blitz, 
+	// otherwise they all get renamed, and it causes early loading errors.
 	SubObjectsToSave.AddObject(Blueprint->WidgetTree);
 
-	SubObjectsToSave.AddObject( NewWidgetBlueprintClass->WidgetTree );
+	// We need to save all the animations to survive the initial sub-object clean blitz, 
+	// otherwise they all get renamed, and it causes early loading errors.
+	SubObjectsToSave.AddObject(NewWidgetBlueprintClass->WidgetTree);
+	for ( UWidgetAnimation* Animation : NewWidgetBlueprintClass->Animations )
+	{
+		SubObjectsToSave.AddObject(Animation);
+	}
 }
 
 void FWidgetBlueprintCompiler::CreateClassVariablesFromBlueprint()
@@ -240,50 +257,45 @@ void FWidgetBlueprintCompiler::FinishCompilingClass(UClass* Class)
 {
 	UWidgetBlueprint* Blueprint = WidgetBlueprint();
 
-	UWidgetBlueprintGeneratedClass* BPGClass = CastChecked<UWidgetBlueprintGeneratedClass>(Class);
-
-	BPGClass->WidgetTree = DuplicateObject<UWidgetTree>(Blueprint->WidgetTree, BPGClass);
-
-	BPGClass->Animations.Empty();
-
-	// Dont duplicate animation data on the skeleton class
-	if( Blueprint->SkeletonGeneratedClass != Class )
+	// Don't do a bunch of extra work on the skeleton generated class
+	if ( Blueprint->SkeletonGeneratedClass != Class )
 	{
-		for(const UWidgetAnimation* Animation : Blueprint->Animations)
+		UWidgetBlueprintGeneratedClass* BPGClass = CastChecked<UWidgetBlueprintGeneratedClass>(Class);
+		BPGClass->WidgetTree = BPGClass->DesignerWidgetTree = DuplicateObject<UWidgetTree>(Blueprint->WidgetTree, BPGClass);
+
+		for ( const UWidgetAnimation* Animation : Blueprint->Animations )
 		{
-			UWidgetAnimation* ClonedAnimation = DuplicateObject<UWidgetAnimation>(Animation, BPGClass, *(Animation->GetName()+TEXT("_INST")));
+			UWidgetAnimation* ClonedAnimation = DuplicateObject<UWidgetAnimation>(Animation, BPGClass, *( Animation->GetName() + TEXT("_INST") ));
 
 			BPGClass->Animations.Add(ClonedAnimation);
 		}
-	}
-
-	BPGClass->Bindings.Reset();
 
-	// Only check bindings on a full compile.  Also don't check them if we're regenerating on load,
-	// that has a nasty tendency to fail because the other dependent classes that may also be blueprints
-	// might not be loaded yet.
-	const bool bIsLoading = Blueprint->bIsRegeneratingOnLoad;
-	if ( bIsFullCompile )
-	{
-		// Convert all editor time property bindings into a list of bindings
-		// that will be applied at runtime.  Ensure all bindings are still valid.
-		for ( const FDelegateEditorBinding& EditorBinding : Blueprint->Bindings )
+		// Only check bindings on a full compile.  Also don't check them if we're regenerating on load,
+		// that has a nasty tendency to fail because the other dependent classes that may also be blueprints
+		// might not be loaded yet.
+		const bool bIsLoading = Blueprint->bIsRegeneratingOnLoad;
+		if ( bIsFullCompile )
 		{
-			if ( bIsLoading || EditorBinding.IsBindingValid(Class, Blueprint, MessageLog) )
+			// Convert all editor time property bindings into a list of bindings
+			// that will be applied at runtime.  Ensure all bindings are still valid.
+			for ( const FDelegateEditorBinding& EditorBinding : Blueprint->Bindings )
 			{
-				BPGClass->Bindings.Add(EditorBinding.ToRuntimeBinding(Blueprint));
+				if ( bIsLoading || EditorBinding.IsBindingValid(Class, Blueprint, MessageLog) )
+				{
+					BPGClass->Bindings.Add(EditorBinding.ToRuntimeBinding(Blueprint));
+				}
 			}
 		}
-	}
 
-	// Add all the names of the named slot widgets to the slot names structure.
-	BPGClass->NamedSlots.Reset();
-	BPGClass->WidgetTree->ForEachWidget([&] (UWidget* Widget) {
-		if ( Widget && Widget->IsA<UNamedSlot>() )
-		{
-			BPGClass->NamedSlots.Add(Widget->GetFName());
-		}
-	});
+		// Add all the names of the named slot widgets to the slot names structure.
+		BPGClass->NamedSlots.Reset();
+		BPGClass->WidgetTree->ForEachWidget([&] (UWidget* Widget) {
+			if ( Widget && Widget->IsA<UNamedSlot>() )
+			{
+				BPGClass->NamedSlots.Add(Widget->GetFName());
+			}
+		});
+	}
 
 	Super::FinishCompilingClass(Class);
 }
diff --git a/Engine/Source/Editor/UMGEditor/Private/WidgetBlueprintEditor.cpp b/Engine/Source/Editor/UMGEditor/Private/WidgetBlueprintEditor.cpp
--- a/Engine/Source/Editor/UMGEditor/Private/WidgetBlueprintEditor.cpp
+++ b/Engine/Source/Editor/UMGEditor/Private/WidgetBlueprintEditor.cpp
@@ -635,7 +635,7 @@ void FWidgetBlueprintEditor::UpdatePreview(UBlueprint* InBlueprint, bool bInForc
 
 		// Update the generated class'es widget tree to match the blueprint tree.  That way the preview can update
 		// without needing to do a full recompile.
-		Cast<UWidgetBlueprintGeneratedClass>(PreviewBlueprint->GeneratedClass)->WidgetTree = DuplicateObject<UWidgetTree>(PreviewBlueprint->WidgetTree, PreviewBlueprint->GeneratedClass);
+		Cast<UWidgetBlueprintGeneratedClass>(PreviewBlueprint->GeneratedClass)->DesignerWidgetTree = DuplicateObject<UWidgetTree>(PreviewBlueprint->WidgetTree, PreviewBlueprint->GeneratedClass);
 
 		PreviewActor = CreateWidget<UUserWidget>(PreviewScene.GetWorld(), PreviewBlueprint->GeneratedClass);
 		PreviewActor->SetFlags(RF_Transactional);
diff --git a/Engine/Source/Runtime/UMG/Private/WidgetBlueprintGeneratedClass.cpp b/Engine/Source/Runtime/UMG/Private/WidgetBlueprintGeneratedClass.cpp
--- a/Engine/Source/Runtime/UMG/Private/WidgetBlueprintGeneratedClass.cpp
+++ b/Engine/Source/Runtime/UMG/Private/WidgetBlueprintGeneratedClass.cpp
@@ -17,7 +17,11 @@ UWidgetBlueprintGeneratedClass::UWidgetBlueprintGeneratedClass(const FObjectInit
 
 void UWidgetBlueprintGeneratedClass::InitializeWidget(UUserWidget* UserWidget) const
 {
-	UWidgetTree* ClonedTree = DuplicateObject<UWidgetTree>( WidgetTree, UserWidget );
+#if WITH_EDITORONLY_DATA
+	UWidgetTree* ClonedTree = DuplicateObject<UWidgetTree>(DesignerWidgetTree ? DesignerWidgetTree : WidgetTree, UserWidget);
+#else
+	UWidgetTree* ClonedTree = DuplicateObject<UWidgetTree>(WidgetTree, UserWidget);
+#endif
 
 #if WITH_EDITOR
 	UserWidget->WidgetGeneratedBy = ClassGeneratedBy;
diff --git a/Engine/Source/Runtime/UMG/Public/Blueprint/WidgetBlueprintGeneratedClass.h b/Engine/Source/Runtime/UMG/Public/Blueprint/WidgetBlueprintGeneratedClass.h
--- a/Engine/Source/Runtime/UMG/Public/Blueprint/WidgetBlueprintGeneratedClass.h
+++ b/Engine/Source/Runtime/UMG/Public/Blueprint/WidgetBlueprintGeneratedClass.h
@@ -63,6 +63,15 @@ public:
 	UPROPERTY()
 	class UWidgetTree* WidgetTree;
 
+#if WITH_EDITORONLY_DATA
+	/**
+	 * A tree of the widget templates to be created.  This is only used by the designer to inject a new 
+	 * tree into a previously compiled class.
+	 */
+	UPROPERTY(Transient)
+	class UWidgetTree* DesignerWidgetTree;
+#endif
+
 	UPROPERTY()
 	TArray< FDelegateRuntimeBinding > Bindings;
 
