MovieScene: Merged some more non-functional refactor changes to reduce diffs

diff --git a/Engine/Source/Editor/UMGEditor/Private/WidgetBlueprintEditor.cpp b/Engine/Source/Editor/UMGEditor/Private/WidgetBlueprintEditor.cpp
--- a/Engine/Source/Editor/UMGEditor/Private/WidgetBlueprintEditor.cpp
+++ b/Engine/Source/Editor/UMGEditor/Private/WidgetBlueprintEditor.cpp
@@ -557,10 +557,12 @@ TSharedPtr<ISequencer>& FWidgetBlueprintEditor::GetSequencer()
 {
 	if(!Sequencer.IsValid())
 	{
-		FSequencerViewParams ViewParams( TEXT( "UMGSequencerSettings" ) );
-		ViewParams.InitalViewRange = TRange<float>(-0.02f, 3.2f);
-		ViewParams.InitialScrubPosition = 0;
-		ViewParams.OnGetAddMenuContent = FOnGetAddMenuContent::CreateSP(this, &FWidgetBlueprintEditor::OnGetAnimationAddMenuContent);
+		FSequencerViewParams ViewParams(TEXT("UMGSequencerSettings"));
+		{
+			ViewParams.InitalViewRange = TRange<float>(-0.02f, 3.2f);
+			ViewParams.InitialScrubPosition = 0;
+			ViewParams.OnGetAddMenuContent = FOnGetAddMenuContent::CreateSP(this, &FWidgetBlueprintEditor::OnGetAnimationAddMenuContent);
+		}
 
 		SequencerObjectBindingManager = MakeShareable(new FUMGSequencerObjectBindingManager(*this, *UWidgetAnimation::GetNullAnimation()));
 
