Fixed UMG editor undo/redo by adding missing transaction code for adding/moving widgets
[UE-18904] - UMG Undo/Redo not working

diff --git a/Engine/Source/Editor/UMGEditor/Private/Extensions/HorizontalSlotExtension.cpp b/Engine/Source/Editor/UMGEditor/Private/Extensions/HorizontalSlotExtension.cpp
--- a/Engine/Source/Editor/UMGEditor/Private/Extensions/HorizontalSlotExtension.cpp
+++ b/Engine/Source/Editor/UMGEditor/Private/Extensions/HorizontalSlotExtension.cpp
@@ -78,6 +78,7 @@ void FHorizontalSlotExtension::ShiftHorizontal(UWidget* Widget, int32 ShiftAmoun
 {
 	UHorizontalBox* Parent = CastChecked<UHorizontalBox>(Widget->GetParent());
 
+	Parent->Modify();
 	int32 CurrentIndex = Parent->GetChildIndex(Widget);
 	Parent->ShiftChild(CurrentIndex + ShiftAmount, Widget);
 }
diff --git a/Engine/Source/Editor/UMGEditor/Private/Extensions/VerticalSlotExtension.cpp b/Engine/Source/Editor/UMGEditor/Private/Extensions/VerticalSlotExtension.cpp
--- a/Engine/Source/Editor/UMGEditor/Private/Extensions/VerticalSlotExtension.cpp
+++ b/Engine/Source/Editor/UMGEditor/Private/Extensions/VerticalSlotExtension.cpp
@@ -78,6 +78,7 @@ void FVerticalSlotExtension::ShiftVertical(UWidget* Widget, int32 ShiftAmount)
 {
 	UVerticalBox* Parent = CastChecked<UVerticalBox>(Widget->GetParent());
 
+	Parent->Modify();
 	int32 CurrentIndex = Parent->GetChildIndex(Widget);
 	Parent->ShiftChild(CurrentIndex + ShiftAmount, Widget);
 }
diff --git a/Engine/Source/Editor/UMGEditor/Private/Hierarchy/SHierarchyViewItem.cpp b/Engine/Source/Editor/UMGEditor/Private/Hierarchy/SHierarchyViewItem.cpp
--- a/Engine/Source/Editor/UMGEditor/Private/Hierarchy/SHierarchyViewItem.cpp
+++ b/Engine/Source/Editor/UMGEditor/Private/Hierarchy/SHierarchyViewItem.cpp
@@ -166,6 +166,11 @@ TOptional<EItemDropZone> ProcessHierarchyDragDrop(const FDragDropEvent& DragDrop
 			// TODO UMG Allow showing a preview of this.
 			if ( bIsDrop )
 			{
+				FScopedTransaction Transaction(LOCTEXT("AddWidgetFromTemplate", "Add Widget"));
+
+				Blueprint->WidgetTree->SetFlags(RF_Transactional);
+				Blueprint->WidgetTree->Modify();
+
 				Blueprint->WidgetTree->RootWidget = TemplateDragDropOp->Template->Create(Blueprint->WidgetTree);
 				FBlueprintEditorUtils::MarkBlueprintAsStructurallyModified(Blueprint);
 			}
@@ -185,6 +190,12 @@ TOptional<EItemDropZone> ProcessHierarchyDragDrop(const FDragDropEvent& DragDrop
 				// TODO UMG Allow showing a preview of this.
 				if (bIsDrop)
 				{
+					FScopedTransaction Transaction(LOCTEXT("AddWidgetFromTemplate", "Add Widget"));
+
+					Blueprint->WidgetTree->SetFlags(RF_Transactional);
+					Blueprint->WidgetTree->Modify();
+					Parent->Modify();
+
 					UWidget* Widget = TemplateDragDropOp->Template->Create(Blueprint->WidgetTree);
 
 					UPanelSlot* NewSlot = nullptr;
