Fixed dereference null crash in SUMGAnimationList::OnWidgetBlueprintTransaction()

diff --git a/Engine/Source/Editor/UMGEditor/Private/TabFactory/AnimationTabSummoner.cpp b/Engine/Source/Editor/UMGEditor/Private/TabFactory/AnimationTabSummoner.cpp
--- a/Engine/Source/Editor/UMGEditor/Private/TabFactory/AnimationTabSummoner.cpp
+++ b/Engine/Source/Editor/UMGEditor/Private/TabFactory/AnimationTabSummoner.cpp
@@ -297,9 +297,12 @@ private:
 
 		if(ViewedAnim)
 		{
-			TSharedPtr<FWidgetAnimationListItem> FoundListItem = *Animations.FindByPredicate( [&](const TSharedPtr<FWidgetAnimationListItem>& ListItem ) { return ListItem->Animation == ViewedAnim; } );
-			
-			AnimationListView->SetSelection(FoundListItem);
+			const TSharedPtr<FWidgetAnimationListItem>* FoundListItemPtr = Animations.FindByPredicate( [&](const TSharedPtr<FWidgetAnimationListItem>& ListItem ) { return ListItem->Animation == ViewedAnim; } );
+
+			if (FoundListItemPtr != nullptr)
+			{
+				AnimationListView->SetSelection(*FoundListItemPtr);
+			}
 		}
 
 	}
