UMG - Widget Blueprints now correctly note any child widget that is from a blueprint as being a dependecny.  Previously this only occured if it was referenced in the graph.  This should help solve problems with needing to compile the child blueprint, and then compiling all referencing parent blueprints.

diff --git a/Engine/Source/Editor/UMGEditor/Private/WidgetBlueprint.cpp b/Engine/Source/Editor/UMGEditor/Private/WidgetBlueprint.cpp
--- a/Engine/Source/Editor/UMGEditor/Private/WidgetBlueprint.cpp
+++ b/Engine/Source/Editor/UMGEditor/Private/WidgetBlueprint.cpp
@@ -572,6 +572,27 @@ bool UWidgetBlueprint::AllowsDynamicBinding() const
 	return true;
 }
 
+void UWidgetBlueprint::GatherDependencies(TSet<TWeakObjectPtr<UBlueprint>>& InDependencies) const
+{
+	Super::GatherDependencies(InDependencies);
+
+	if ( WidgetTree )
+	{
+		WidgetTree->ForEachWidget([&] (UWidget* Widget) {
+			if ( UBlueprint* WidgetBlueprint = UBlueprint::GetBlueprintFromClass(Widget->GetClass()) )
+			{
+				bool bWasAlreadyInSet;
+				InDependencies.Add(WidgetBlueprint, &bWasAlreadyInSet);
+
+				if ( !bWasAlreadyInSet )
+				{
+					WidgetBlueprint->GatherDependencies(InDependencies);
+				}
+			}
+		});
+	}
+}
+
 bool UWidgetBlueprint::ValidateGeneratedClass(const UClass* InClass)
 {
 	bool Result = Super::ValidateGeneratedClass(InClass);
diff --git a/Engine/Source/Editor/UMGEditor/Public/WidgetBlueprint.h b/Engine/Source/Editor/UMGEditor/Public/WidgetBlueprint.h
--- a/Engine/Source/Editor/UMGEditor/Public/WidgetBlueprint.h
+++ b/Engine/Source/Editor/UMGEditor/Public/WidgetBlueprint.h
@@ -225,5 +225,7 @@ public:
 	virtual bool AlwaysCompileOnLoad() const override { return true; }
 	// End of UBlueprint interface
 
+	virtual void GatherDependencies(TSet<TWeakObjectPtr<UBlueprint>>& InDependencies) const override;
+
 	static bool ValidateGeneratedClass(const UClass* InClass);
 };
diff --git a/Engine/Source/Editor/UnrealEd/Private/Kismet2/BlueprintEditorUtils.cpp b/Engine/Source/Editor/UnrealEd/Private/Kismet2/BlueprintEditorUtils.cpp
--- a/Engine/Source/Editor/UnrealEd/Private/Kismet2/BlueprintEditorUtils.cpp
+++ b/Engine/Source/Editor/UnrealEd/Private/Kismet2/BlueprintEditorUtils.cpp
@@ -2781,6 +2781,8 @@ void FBlueprintEditorUtils::GatherDependencies(const UBlueprint* InBlueprint, TS
 				{
 					return;
 				}
+
+				Blueprint->GatherDependencies(InDependencies);
 			}
 		}
 	};
@@ -2789,6 +2791,8 @@ void FBlueprintEditorUtils::GatherDependencies(const UBlueprint* InBlueprint, TS
 	Dependencies.Empty();
 	OutUDSDependencies.Empty();
 
+	InBlueprint->GatherDependencies(Dependencies);
+
 	FGatherDependenciesHelper::ProcessHierarchy(InBlueprint->ParentClass, Dependencies);
 
 	for (const auto& InterfaceDesc : InBlueprint->ImplementedInterfaces)
diff --git a/Engine/Source/Runtime/Engine/Classes/Engine/Blueprint.h b/Engine/Source/Runtime/Engine/Classes/Engine/Blueprint.h
--- a/Engine/Source/Runtime/Engine/Classes/Engine/Blueprint.h
+++ b/Engine/Source/Runtime/Engine/Classes/Engine/Blueprint.h
@@ -584,6 +584,9 @@ public:
 
 	UInheritableComponentHandler* GetInheritableComponentHandler(bool bCreateIfNecessary);
 
+	/** Collect blueprints that depend on this blueprint. */
+	virtual void GatherDependencies(TSet<TWeakObjectPtr<UBlueprint>>& InDependencies) const;
+
 #endif	//#if WITH_EDITOR
 
 	// Begin UObject interface
diff --git a/Engine/Source/Runtime/Engine/Private/Blueprint.cpp b/Engine/Source/Runtime/Engine/Private/Blueprint.cpp
--- a/Engine/Source/Runtime/Engine/Private/Blueprint.cpp
+++ b/Engine/Source/Runtime/Engine/Private/Blueprint.cpp
@@ -1287,6 +1287,11 @@ bool UBlueprint::Modify(bool bAlwaysMarkDirty)
 	return Super::Modify(bAlwaysMarkDirty);
 }
 
+void UBlueprint::GatherDependencies(TSet<TWeakObjectPtr<UBlueprint>>& InDependencies) const
+{
+
+}
+
 UInheritableComponentHandler* UBlueprint::GetInheritableComponentHandler(bool bCreateIfNecessary)
 {
 	static const FBoolConfigValueHelper EnableInheritableComponents(TEXT("Kismet"), TEXT("bEnableInheritableComponents"), GEngineIni);
