UMG - Optimization, if the user does not implement Tick / OnPaint, we now note that in the generated class and configure some flags on the instance to not attempt to ever tick or paint if those were not implemented.  This saves a lot with lots of UserWidgets because it avoids NumWidgets * 2 FindFunctionChecked and the copying of the paint args multiple times.

diff --git a/Engine/Source/Editor/UMGEditor/Private/WidgetBlueprintCompiler.cpp b/Engine/Source/Editor/UMGEditor/Private/WidgetBlueprintCompiler.cpp
--- a/Engine/Source/Editor/UMGEditor/Private/WidgetBlueprintCompiler.cpp
+++ b/Engine/Source/Editor/UMGEditor/Private/WidgetBlueprintCompiler.cpp
@@ -304,6 +304,26 @@ void FWidgetBlueprintCompiler::Compile()
 {
 	Super::Compile();
 
+	//TODO Once we handle multiple derived blueprint classes, we need to check parent versions of the class.
+	if ( const UFunction* ReceiveTickEvent = FKismetCompilerUtilities::FindOverriddenImplementableEvent(GET_FUNCTION_NAME_CHECKED(UUserWidget, Tick), NewWidgetBlueprintClass) )
+	{
+		NewWidgetBlueprintClass->bCanEverTick = true;
+	}
+	else
+	{
+		NewWidgetBlueprintClass->bCanEverTick = false;
+	}
+	
+	//TODO Once we handle multiple derived blueprint classes, we need to check parent versions of the class.
+	if ( const UFunction* ReceivePaintEvent = FKismetCompilerUtilities::FindOverriddenImplementableEvent(GET_FUNCTION_NAME_CHECKED(UUserWidget, OnPaint), NewWidgetBlueprintClass) )
+	{
+		NewWidgetBlueprintClass->bCanEverPaint = true;
+	}
+	else
+	{
+		NewWidgetBlueprintClass->bCanEverPaint = false;
+	}
+
 	WidgetToMemberVariableMap.Empty();
 }
 
diff --git a/Engine/Source/Runtime/UMG/Private/UserWidget.cpp b/Engine/Source/Runtime/UMG/Private/UserWidget.cpp
--- a/Engine/Source/Runtime/UMG/Private/UserWidget.cpp
+++ b/Engine/Source/Runtime/UMG/Private/UserWidget.cpp
@@ -16,6 +16,8 @@
 
 UUserWidget::UUserWidget(const FObjectInitializer& ObjectInitializer)
 	: Super(ObjectInitializer)
+	, bCanEverTick(true)
+	, bCanEverPaint(true)
 {
 	ViewportAnchors = FAnchors(0, 0, 1, 1);
 	Visiblity_DEPRECATED = Visibility = ESlateVisibility::SelfHitTestInvisible;
@@ -729,7 +731,10 @@ void UUserWidget::NativeTick(const FGeometry& MyGeometry, float InDeltaTime)
 
 	TickActionsAndAnimation(MyGeometry, InDeltaTime);
 
-	Tick(MyGeometry, InDeltaTime);
+	if ( bCanEverTick )
+	{
+		Tick(MyGeometry, InDeltaTime);
+	}
 }
 
 void UUserWidget::TickActionsAndAnimation(const FGeometry& MyGeometry, float InDeltaTime)
@@ -776,7 +781,10 @@ void UUserWidget::TickActionsAndAnimation(const FGeometry& MyGeometry, float InD
 
 void UUserWidget::NativePaint( FPaintContext& InContext ) const 
 {
-	OnPaint( InContext );
+	if ( bCanEverPaint )
+	{
+		OnPaint(InContext);
+	}
 }
 
 bool UUserWidget::NativeIsInteractable() const
diff --git a/Engine/Source/Runtime/UMG/Private/WidgetBlueprintGeneratedClass.cpp b/Engine/Source/Runtime/UMG/Private/WidgetBlueprintGeneratedClass.cpp
--- a/Engine/Source/Runtime/UMG/Private/WidgetBlueprintGeneratedClass.cpp
+++ b/Engine/Source/Runtime/UMG/Private/WidgetBlueprintGeneratedClass.cpp
@@ -17,6 +17,9 @@ UWidgetBlueprintGeneratedClass::UWidgetBlueprintGeneratedClass(const FObjectInit
 
 void UWidgetBlueprintGeneratedClass::InitializeWidget(UUserWidget* UserWidget) const
 {
+	UserWidget->bCanEverTick = bCanEverTick;
+	UserWidget->bCanEverPaint = bCanEverPaint;
+
 #if WITH_EDITORONLY_DATA
 	UWidgetTree* ClonedTree = DuplicateObject<UWidgetTree>(DesignerWidgetTree ? DesignerWidgetTree : WidgetTree, UserWidget);
 #else
diff --git a/Engine/Source/Runtime/UMG/Public/Blueprint/WidgetBlueprintGeneratedClass.h b/Engine/Source/Runtime/UMG/Public/Blueprint/WidgetBlueprintGeneratedClass.h
--- a/Engine/Source/Runtime/UMG/Public/Blueprint/WidgetBlueprintGeneratedClass.h
+++ b/Engine/Source/Runtime/UMG/Public/Blueprint/WidgetBlueprintGeneratedClass.h
@@ -81,6 +81,12 @@ public:
 	UPROPERTY()
 	TArray< FName > NamedSlots;
 
+	UPROPERTY()
+	uint32 bCanEverTick : 1;
+
+	UPROPERTY()
+	uint32 bCanEverPaint : 1;
+
 public:
 
 	virtual void PostLoad() override;
