Added valid object name check to widget renaming code in UMG Designer.
[UE-19542] - Having a period in the name of a widget will cause the editor to crash after copying and pasting (no error/warning given)

diff --git a/Engine/Source/Editor/UMGEditor/Private/WidgetBlueprintEditorUtils.cpp b/Engine/Source/Editor/UMGEditor/Private/WidgetBlueprintEditorUtils.cpp
--- a/Engine/Source/Editor/UMGEditor/Private/WidgetBlueprintEditorUtils.cpp
+++ b/Engine/Source/Editor/UMGEditor/Private/WidgetBlueprintEditorUtils.cpp
@@ -103,6 +103,11 @@ bool FWidgetBlueprintEditorUtils::VerifyWidgetRename(TSharedRef<class FWidgetBlu
 	}
 	else
 	{
+		if (!FName(*NewNameString).IsValidObjectName(OutErrorMessage))
+		{
+			return false;
+		}
+
 		// Not an existing widget in the tree BUT it still mustn't create a UObject name clash
 		UWidget* WidgetPreview = Widget.GetPreview();
 		if (WidgetPreview)
