UMG: Fixed linker load warnings about WidgetTree outers in cooked builds
- LogLinker:Warning: CreateImport: Failed to load Outer for resource 'WidgetTree': WidgetBlueprint /Script/UMGEditor.Default__WidgetBlueprint
[INTEGRATE] Change 2668264 by [REDACTED_PERSON][REDACTED_HANDLE] on 2015/08/25 16:20:29

diff --git a/Engine/Source/Editor/UMGEditor/Public/WidgetBlueprint.h b/Engine/Source/Editor/UMGEditor/Public/WidgetBlueprint.h
--- a/Engine/Source/Editor/UMGEditor/Public/WidgetBlueprint.h
+++ b/Engine/Source/Editor/UMGEditor/Public/WidgetBlueprint.h
@@ -182,6 +182,7 @@ public:
 
 public:
 
+#if WITH_EDITORONLY_DATA
 	/** A tree of the widget templates to be created */
 	UPROPERTY()
 	class UWidgetTree* WidgetTree;
@@ -202,6 +203,7 @@ public:
 	 */
 	UPROPERTY(AssetRegistrySearchable)
 	FString PaletteCategory;
+#endif
 
 public:
 
