Fixed drag-drop to and from named slots in the UMG widget hierarchy view
[UE-18941] - Can't drag-drop an existing widget in the hierarchy to a named slot

diff --git a/Engine/Source/Editor/UMGEditor/Private/Hierarchy/SHierarchyViewItem.cpp b/Engine/Source/Editor/UMGEditor/Private/Hierarchy/SHierarchyViewItem.cpp
--- a/Engine/Source/Editor/UMGEditor/Private/Hierarchy/SHierarchyViewItem.cpp
+++ b/Engine/Source/Editor/UMGEditor/Private/Hierarchy/SHierarchyViewItem.cpp
@@ -246,14 +246,14 @@ TOptional<EItemDropZone> ProcessHierarchyDragDrop(const FDragDropEvent& DragDrop
 			}
 
 			UPanelWidget* NewParent = Cast<UPanelWidget>(TargetItem.GetTemplate());
-			if ( !NewParent )
+			if (!NewParent)
 			{
 				HierarchyDragDropOp->CurrentIconBrush = FEditorStyle::GetBrush(TEXT("Graph.ConnectorFeedback.Error"));
 				HierarchyDragDropOp->CurrentHoverText = LOCTEXT("CantHaveChildren", "Widget can't have children.");
 				return TOptional<EItemDropZone>();
 			}
 
-			if ( !NewParent->CanAddMoreChildren() )
+			if (!NewParent->CanAddMoreChildren())
 			{
 				HierarchyDragDropOp->CurrentIconBrush = FEditorStyle::GetBrush(TEXT("Graph.ConnectorFeedback.Error"));
 				HierarchyDragDropOp->CurrentHoverText = LOCTEXT("NoAdditionalChildren", "Widget can't accept additional children.");
@@ -274,7 +274,7 @@ TOptional<EItemDropZone> ProcessHierarchyDragDrop(const FDragDropEvent& DragDrop
 				});
 			}
 
-			if ( bFoundNewParentInChildSet )
+			if (bFoundNewParentInChildSet)
 			{
 				HierarchyDragDropOp->CurrentIconBrush = FEditorStyle::GetBrush(TEXT("Graph.ConnectorFeedback.Error"));
 				HierarchyDragDropOp->CurrentHoverText = LOCTEXT("CantMakeWidgetChildOfChildren", "Can't make widget a child of its children.");
@@ -306,6 +306,14 @@ TOptional<EItemDropZone> ProcessHierarchyDragDrop(const FDragDropEvent& DragDrop
 						}
 					}
 
+					// We don't know if this widget is being removed from a named slot and RemoveFromParent is not enough to take care of this
+					INamedSlotInterface* NamedSlotHost = FWidgetBlueprintEditorUtils::FindNamedSlotHostForContent(TemplateWidget, Blueprint->WidgetTree);
+					if (NamedSlotHost != nullptr)
+					{
+						NamedSlotHost->ModifySlots();
+						FWidgetBlueprintEditorUtils::RemoveNamedSlotHostContent(TemplateWidget, NamedSlotHost);
+					}
+
 					TemplateWidget->RemoveFromParent();
 
 					UPanelSlot* NewSlot = nullptr;
@@ -660,20 +668,62 @@ FWidgetReference FNamedSlotModel::AsDraggedWidgetReference() const
 
 TOptional<EItemDropZone> FNamedSlotModel::HandleCanAcceptDrop(const FDragDropEvent& DragDropEvent, EItemDropZone DropZone)
 {
+	UWidgetBlueprint* Blueprint = BlueprintEditor.Pin()->GetWidgetBlueprintObj();
+
 	TSharedPtr<FWidgetTemplateDragDropOp> TemplateDragDropOp = DragDropEvent.GetOperationAs<FWidgetTemplateDragDropOp>();
-	if ( TemplateDragDropOp.IsValid() )
+	if (TemplateDragDropOp.IsValid())
 	{
+		TemplateDragDropOp->ResetToDefaultToolTip();
+
 		if ( INamedSlotInterface* NamedSlotHost = Cast<INamedSlotInterface>(Item.GetTemplate()) )
 		{
 			// Only assign content to the named slot if it is null.
 			if ( NamedSlotHost->GetContentForSlot(SlotName) != nullptr )
 			{
-				TSharedPtr<FDecoratedDragDropOp> DecoratedDragDropOp = DragDropEvent.GetOperationAs<FDecoratedDragDropOp>();
 				TemplateDragDropOp->CurrentIconBrush = FEditorStyle::GetBrush(TEXT("Graph.ConnectorFeedback.Error"));
+				TemplateDragDropOp->CurrentHoverText = LOCTEXT("NamedSlotAlreadyFull", "Named Slot already has a child.");
+				return TOptional<EItemDropZone>();
+			}
 
+			TemplateDragDropOp->CurrentIconBrush = FEditorStyle::GetBrush(TEXT("Graph.ConnectorFeedback.OK"));
+			return EItemDropZone::OntoItem;
+		}
+	}
+
+	TSharedPtr<FHierarchyWidgetDragDropOp> HierarchyDragDropOp = DragDropEvent.GetOperationAs<FHierarchyWidgetDragDropOp>();
+	if (HierarchyDragDropOp.IsValid() && HierarchyDragDropOp->DraggedWidgets.Num() == 1)
+	{
+		HierarchyDragDropOp->ResetToDefaultToolTip();
+
+		if (INamedSlotInterface* NamedSlotHost = Cast<INamedSlotInterface>(Item.GetTemplate()))
+		{
+			// Only assign content to the named slot if it is null.
+			if (NamedSlotHost->GetContentForSlot(SlotName) != nullptr)
+			{
+				HierarchyDragDropOp->CurrentIconBrush = FEditorStyle::GetBrush(TEXT("Graph.ConnectorFeedback.Error"));
+				HierarchyDragDropOp->CurrentHoverText = LOCTEXT("NamedSlotAlreadyFull", "Named Slot already has a child.");
+				return TOptional<EItemDropZone>();
+			}
+
+			bool bFoundNewParentInChildSet = false;
+			UWidget* TemplateWidget = HierarchyDragDropOp->DraggedWidgets[0].Widget.GetTemplate();
+
+			// Verify that the new location we're placing the widget is not inside of its existing children.
+			Blueprint->WidgetTree->ForWidgetAndChildren(TemplateWidget, [&](UWidget* Widget) {
+				if (Item.GetTemplate() == Widget)
+				{
+					bFoundNewParentInChildSet = true;
+				}
+			});
+
+			if (bFoundNewParentInChildSet)
+			{
+				HierarchyDragDropOp->CurrentIconBrush = FEditorStyle::GetBrush(TEXT("Graph.ConnectorFeedback.Error"));
+				HierarchyDragDropOp->CurrentHoverText = LOCTEXT("CantMakeWidgetChildOfChildren", "Can't make widget a child of its children.");
 				return TOptional<EItemDropZone>();
 			}
 
+			HierarchyDragDropOp->CurrentIconBrush = FEditorStyle::GetBrush(TEXT("Graph.ConnectorFeedback.OK"));
 			return EItemDropZone::OntoItem;
 		}
 	}
@@ -683,34 +733,78 @@ TOptional<EItemDropZone> FNamedSlotModel::HandleCanAcceptDrop(const FDragDropEve
 
 FReply FNamedSlotModel::HandleAcceptDrop(FDragDropEvent const& DragDropEvent, EItemDropZone DropZone)
 {
-	UWidgetBlueprint* Blueprint = BlueprintEditor.Pin()->GetWidgetBlueprintObj();
+	UWidget* SlotHostWidget = Item.GetTemplate();
+	INamedSlotInterface* NamedSlotHost = Cast<INamedSlotInterface>(SlotHostWidget);
+	if (NamedSlotHost == nullptr)
+	{
+		return FReply::Unhandled();
+	}
+	if (NamedSlotHost->GetContentForSlot(SlotName) != nullptr)
+	{
+		return FReply::Unhandled();
+	}
 
-	// Is this a drag/drop op to create a new widget in the tree?
 	TSharedPtr<FWidgetTemplateDragDropOp> TemplateDragDropOp = DragDropEvent.GetOperationAs<FWidgetTemplateDragDropOp>();
-	if ( TemplateDragDropOp.IsValid() )
+	if (TemplateDragDropOp.IsValid())
 	{
-		if ( INamedSlotInterface* NamedSlotHost = Cast<INamedSlotInterface>(Item.GetTemplate()) )
-		{
-			// Only assign content to the named slot if it is null.
-			if ( NamedSlotHost->GetContentForSlot(SlotName) == nullptr )
-			{
-				UWidget* Widget = TemplateDragDropOp->Template->Create(Blueprint->WidgetTree);
-				NamedSlotHost->SetContentForSlot(SlotName, Widget);
-				FBlueprintEditorUtils::MarkBlueprintAsStructurallyModified(Blueprint);
+		FScopedTransaction Transaction(LOCTEXT("AddWidgetFromTemplate", "Add Widget"));
 
-				TSet<FWidgetReference> SelectedTemplates;
-				SelectedTemplates.Add(BlueprintEditor.Pin()->GetReferenceFromTemplate(Widget));
+		UWidgetBlueprint* Blueprint = BlueprintEditor.Pin()->GetWidgetBlueprintObj();
+		Blueprint->WidgetTree->SetFlags(RF_Transactional);
+		Blueprint->WidgetTree->Modify();
 
-				BlueprintEditor.Pin()->SelectWidgets(SelectedTemplates, false);
-			}
+		UWidget* DroppingWidget = TemplateDragDropOp->Template->Create(Blueprint->WidgetTree);
+
+		DoDrop(SlotHostWidget, DroppingWidget);
+
+		return FReply::Handled();
+	}
 
-			return FReply::Handled();
+	TSharedPtr<FHierarchyWidgetDragDropOp> HierarchyDragDropOp = DragDropEvent.GetOperationAs<FHierarchyWidgetDragDropOp>();
+	if (HierarchyDragDropOp.IsValid() && HierarchyDragDropOp->DraggedWidgets.Num() == 1)
+	{
+		UWidgetBlueprint* Blueprint = BlueprintEditor.Pin()->GetWidgetBlueprintObj();
+		Blueprint->WidgetTree->SetFlags(RF_Transactional);
+		Blueprint->WidgetTree->Modify();
+
+		UWidget* DroppingWidget = HierarchyDragDropOp->DraggedWidgets[0].Widget.GetTemplate();
+
+		// We don't know if this widget is being removed from a named slot and RemoveFromParent is not enough to take care of this
+		INamedSlotInterface* SourceNamedSlotHost = FWidgetBlueprintEditorUtils::FindNamedSlotHostForContent(DroppingWidget, Blueprint->WidgetTree);
+		if (SourceNamedSlotHost != nullptr)
+		{
+			SourceNamedSlotHost->ModifySlots();
+			FWidgetBlueprintEditorUtils::RemoveNamedSlotHostContent(DroppingWidget, SourceNamedSlotHost);
 		}
+
+		DroppingWidget->RemoveFromParent();
+
+		DoDrop(SlotHostWidget, DroppingWidget);
+
+		return FReply::Handled();
 	}
 
 	return FReply::Unhandled();
 }
 
+void FNamedSlotModel::DoDrop(UWidget* NamedSlotHostWidget, UWidget* DroppingWidget)
+{
+	UWidgetBlueprint* Blueprint = BlueprintEditor.Pin()->GetWidgetBlueprintObj();
+
+	NamedSlotHostWidget->SetFlags(RF_Transactional);
+	NamedSlotHostWidget->Modify();
+
+	INamedSlotInterface* NamedSlotInterface = Cast<INamedSlotInterface>(NamedSlotHostWidget);
+	NamedSlotInterface->SetContentForSlot(SlotName, DroppingWidget);
+
+	FBlueprintEditorUtils::MarkBlueprintAsStructurallyModified(Blueprint);
+
+	TSet<FWidgetReference> SelectedTemplates;
+	SelectedTemplates.Add(BlueprintEditor.Pin()->GetReferenceFromTemplate(DroppingWidget));
+
+	BlueprintEditor.Pin()->SelectWidgets(SelectedTemplates, false);
+}
+
 //////////////////////////////////////////////////////////////////////////
 
 FHierarchyWidget::FHierarchyWidget(FWidgetReference InItem, TSharedPtr<FWidgetBlueprintEditor> InBlueprintEditor)
diff --git a/Engine/Source/Editor/UMGEditor/Private/Hierarchy/SHierarchyViewItem.h b/Engine/Source/Editor/UMGEditor/Private/Hierarchy/SHierarchyViewItem.h
--- a/Engine/Source/Editor/UMGEditor/Private/Hierarchy/SHierarchyViewItem.h
+++ b/Engine/Source/Editor/UMGEditor/Private/Hierarchy/SHierarchyViewItem.h
@@ -143,6 +143,8 @@ protected:
 	virtual void UpdateSelection() override;
 	virtual FWidgetReference AsDraggedWidgetReference() const;
 
+	void DoDrop(UWidget* NamedSlotHostWidget, UWidget* DroppingWidget);
+
 private:
 
 	FWidgetReference Item;
diff --git a/Engine/Source/Editor/UMGEditor/Private/WidgetBlueprintEditorUtils.cpp b/Engine/Source/Editor/UMGEditor/Private/WidgetBlueprintEditorUtils.cpp
--- a/Engine/Source/Editor/UMGEditor/Private/WidgetBlueprintEditorUtils.cpp
+++ b/Engine/Source/Editor/UMGEditor/Private/WidgetBlueprintEditorUtils.cpp
@@ -350,37 +350,67 @@ void FWidgetBlueprintEditorUtils::DeleteWidgets(UWidgetBlueprint* BP, TSet<FWidg
 	}
 }
 
-bool FWidgetBlueprintEditorUtils::FindAndRemoveNamedSlotContent(UWidget* WidgetTemplate, UWidgetTree* WidgetTree)
+INamedSlotInterface* FWidgetBlueprintEditorUtils::FindNamedSlotHostForContent(UWidget* WidgetTemplate, UWidgetTree* WidgetTree)
 {
-	bool bSuccess = false;
+	INamedSlotInterface* Host = nullptr;
 
-	WidgetTree->ForEachWidget([&] (UWidget* Widget) {
+	WidgetTree->ForEachWidget([&](UWidget* Widget) {
 
-		if ( bSuccess )
+		if (Host != nullptr)
 		{
 			return;
 		}
 
-		if ( INamedSlotInterface* NamedSlotHost = Cast<INamedSlotInterface>(Widget) )
+		if (INamedSlotInterface* NamedSlotHost = Cast<INamedSlotInterface>(Widget))
 		{
 			TArray<FName> SlotNames;
 			NamedSlotHost->GetSlotNames(SlotNames);
 
-			for ( FName SlotName : SlotNames )
+			for (FName SlotName : SlotNames)
 			{
-				if ( UWidget* SlotContent = NamedSlotHost->GetContentForSlot(SlotName) )
+				if (UWidget* SlotContent = NamedSlotHost->GetContentForSlot(SlotName))
 				{
-					if ( SlotContent == WidgetTemplate )
+					if (SlotContent == WidgetTemplate)
 					{
-						NamedSlotHost->SetContentForSlot(SlotName, nullptr);
-						bSuccess = true;
+						Host = NamedSlotHost;
 					}
 				}
 			}
 		}
 	});
 
-	return bSuccess;
+	return Host;
+}
+
+bool FWidgetBlueprintEditorUtils::RemoveNamedSlotHostContent(UWidget* WidgetTemplate, INamedSlotInterface* NamedSlotHost)
+{
+	TArray<FName> SlotNames;
+	NamedSlotHost->GetSlotNames(SlotNames);
+
+	for (FName SlotName : SlotNames)
+	{
+		if (UWidget* SlotContent = NamedSlotHost->GetContentForSlot(SlotName))
+		{
+			if (SlotContent == WidgetTemplate)
+			{
+				NamedSlotHost->SetContentForSlot(SlotName, nullptr);
+				return true;
+			}
+		}
+	}
+
+	return false;
+}
+
+bool FWidgetBlueprintEditorUtils::FindAndRemoveNamedSlotContent(UWidget* WidgetTemplate, UWidgetTree* WidgetTree)
+{
+	INamedSlotInterface* NamedSlotHost = FindNamedSlotHostForContent(WidgetTemplate, WidgetTree);
+	if (NamedSlotHost != nullptr)
+	{
+		return RemoveNamedSlotHostContent(WidgetTemplate, NamedSlotHost);
+	}
+
+	return false;
 }
 
 void FWidgetBlueprintEditorUtils::BuildWrapWithMenu(FMenuBuilder& Menu, TSharedRef<FWidgetBlueprintEditor> BlueprintEditor, UWidgetBlueprint* BP, TSet<FWidgetReference> Widgets)
diff --git a/Engine/Source/Editor/UMGEditor/Private/WidgetBlueprintEditorUtils.h b/Engine/Source/Editor/UMGEditor/Private/WidgetBlueprintEditorUtils.h
--- a/Engine/Source/Editor/UMGEditor/Private/WidgetBlueprintEditorUtils.h
+++ b/Engine/Source/Editor/UMGEditor/Private/WidgetBlueprintEditorUtils.h
@@ -40,14 +40,18 @@ public:
 	/** Attempts to import any property in the map and apply it to a property with the same name on the object. */
 	static void ImportPropertiesFromText(UObject* Object, const TMap<FName, FString>& ExportedProperties);
 
+	static INamedSlotInterface* FindNamedSlotHostForContent(UWidget* WidgetTemplate, UWidgetTree* WidgetTree);
+
+	static bool RemoveNamedSlotHostContent(UWidget* WidgetTemplate, INamedSlotInterface* NamedSlotHost);
+
 private:
 
 	static void ExecuteOpenSelectedWidgetsForEdit( TSet<FWidgetReference> SelectedWidgets );
 
-	static bool CanOpenSelectedWidgetsForEdit( TSet<FWidgetReference> SelectedWidgets );
-
 	static bool FindAndRemoveNamedSlotContent(UWidget* WidgetTemplate, UWidgetTree* WidgetTree);
 
+	static bool CanOpenSelectedWidgetsForEdit( TSet<FWidgetReference> SelectedWidgets );
+
 	static void BuildWrapWithMenu(FMenuBuilder& Menu, TSharedRef<FWidgetBlueprintEditor> BlueprintEditor, UWidgetBlueprint* BP, TSet<FWidgetReference> Widgets);
 
 	static void WrapWidgets(TSharedRef<FWidgetBlueprintEditor> BlueprintEditor, UWidgetBlueprint* BP, TSet<FWidgetReference> Widgets, UClass* WidgetClass);
diff --git a/Engine/Source/Runtime/UMG/Private/Components/ExpandableArea.cpp b/Engine/Source/Runtime/UMG/Private/Components/ExpandableArea.cpp
--- a/Engine/Source/Runtime/UMG/Private/Components/ExpandableArea.cpp
+++ b/Engine/Source/Runtime/UMG/Private/Components/ExpandableArea.cpp
@@ -82,6 +82,12 @@ void UExpandableArea::SetContentForSlot(FName SlotName, UWidget* Content)
 	}
 }
 
+void UExpandableArea::ModifySlots()
+{
+	SetFlags(RF_Transactional);
+	Modify();
+}
+
 TSharedRef<SWidget> UExpandableArea::RebuildWidget()
 {
 	TSharedRef<SWidget> HeaderWidget = HeaderContent ? HeaderContent->TakeWidget() : SNullWidget::NullWidget;
diff --git a/Engine/Source/Runtime/UMG/Private/UserWidget.cpp b/Engine/Source/Runtime/UMG/Private/UserWidget.cpp
--- a/Engine/Source/Runtime/UMG/Private/UserWidget.cpp
+++ b/Engine/Source/Runtime/UMG/Private/UserWidget.cpp
@@ -466,6 +466,12 @@ void UUserWidget::SetContentForSlot(FName SlotName, UWidget* Content)
 	}
 }
 
+void UUserWidget::ModifySlots()
+{
+	SetFlags(RF_Transactional);
+	Modify();
+}
+
 TSharedRef<SWidget> UUserWidget::MakeViewportWidget(TSharedPtr<SWidget>& UserSlateWidget)
 {
 	UserSlateWidget = TakeWidget();
diff --git a/Engine/Source/Runtime/UMG/Public/Blueprint/UserWidget.h b/Engine/Source/Runtime/UMG/Public/Blueprint/UserWidget.h
--- a/Engine/Source/Runtime/UMG/Public/Blueprint/UserWidget.h
+++ b/Engine/Source/Runtime/UMG/Public/Blueprint/UserWidget.h
@@ -171,6 +171,7 @@ public:
 	virtual void GetSlotNames(TArray<FName>& SlotNames) const override;
 	virtual UWidget* GetContentForSlot(FName SlotName) const override;
 	virtual void SetContentForSlot(FName SlotName, UWidget* Content) override;
+	virtual void ModifySlots() override;
 	// UNamedSlotInterface End
 
 	/**
diff --git a/Engine/Source/Runtime/UMG/Public/Components/ExpandableArea.h b/Engine/Source/Runtime/UMG/Public/Components/ExpandableArea.h
--- a/Engine/Source/Runtime/UMG/Public/Components/ExpandableArea.h
+++ b/Engine/Source/Runtime/UMG/Public/Components/ExpandableArea.h
@@ -53,6 +53,7 @@ public:
 	virtual void GetSlotNames(TArray<FName>& SlotNames) const override;
 	virtual UWidget* GetContentForSlot(FName SlotName) const override;
 	virtual void SetContentForSlot(FName SlotName, UWidget* Content) override;
+	virtual void ModifySlots() override;
 	// End INamedSlotInterface
 
 public:
diff --git a/Engine/Source/Runtime/UMG/Public/Components/NamedSlotInterface.h b/Engine/Source/Runtime/UMG/Public/Components/NamedSlotInterface.h
--- a/Engine/Source/Runtime/UMG/Public/Components/NamedSlotInterface.h
+++ b/Engine/Source/Runtime/UMG/Public/Components/NamedSlotInterface.h
@@ -28,4 +28,6 @@ class UMG_API INamedSlotInterface
 
 	/**  */
 	bool ContainsContent(UWidget* Content) const;
+
+	virtual void ModifySlots() = 0;
 };
