Fixed drag-drop to and from named slots in the UMG widget hierarchy view
[UE-18941] - Can't drag-drop an existing widget in the hierarchy to a named slot

Second submission following codereview feedback. Found a way to avoid modifying INameSlotInterface.

diff --git a/Engine/Source/Editor/UMGEditor/Private/Hierarchy/SHierarchyViewItem.cpp b/Engine/Source/Editor/UMGEditor/Private/Hierarchy/SHierarchyViewItem.cpp
--- a/Engine/Source/Editor/UMGEditor/Private/Hierarchy/SHierarchyViewItem.cpp
+++ b/Engine/Source/Editor/UMGEditor/Private/Hierarchy/SHierarchyViewItem.cpp
@@ -307,11 +307,16 @@ TOptional<EItemDropZone> ProcessHierarchyDragDrop(const FDragDropEvent& DragDrop
 					}
 
 					// We don't know if this widget is being removed from a named slot and RemoveFromParent is not enough to take care of this
-					INamedSlotInterface* NamedSlotHost = FWidgetBlueprintEditorUtils::FindNamedSlotHostForContent(TemplateWidget, Blueprint->WidgetTree);
-					if (NamedSlotHost != nullptr)
+					UWidget* NamedSlotHostWidget = FWidgetBlueprintEditorUtils::FindNamedSlotHostWidgetForContent(TemplateWidget, Blueprint->WidgetTree);
+					if (NamedSlotHostWidget != nullptr)
 					{
-						NamedSlotHost->ModifySlots();
-						FWidgetBlueprintEditorUtils::RemoveNamedSlotHostContent(TemplateWidget, NamedSlotHost);
+						INamedSlotInterface* NamedSlotHost = Cast<INamedSlotInterface>(NamedSlotHostWidget);
+						if (NamedSlotHost != nullptr)
+						{
+							NamedSlotHostWidget->SetFlags(RF_Transactional);
+							NamedSlotHostWidget->Modify();
+							FWidgetBlueprintEditorUtils::RemoveNamedSlotHostContent(TemplateWidget, NamedSlotHost);
+						}
 					}
 
 					TemplateWidget->RemoveFromParent();
@@ -770,11 +775,16 @@ FReply FNamedSlotModel::HandleAcceptDrop(FDragDropEvent const& DragDropEvent, EI
 		UWidget* DroppingWidget = HierarchyDragDropOp->DraggedWidgets[0].Widget.GetTemplate();
 
 		// We don't know if this widget is being removed from a named slot and RemoveFromParent is not enough to take care of this
-		INamedSlotInterface* SourceNamedSlotHost = FWidgetBlueprintEditorUtils::FindNamedSlotHostForContent(DroppingWidget, Blueprint->WidgetTree);
-		if (SourceNamedSlotHost != nullptr)
+		UWidget* SourceNamedSlotHostWidget = FWidgetBlueprintEditorUtils::FindNamedSlotHostWidgetForContent(DroppingWidget, Blueprint->WidgetTree);
+		if (SourceNamedSlotHostWidget != nullptr)
 		{
-			SourceNamedSlotHost->ModifySlots();
-			FWidgetBlueprintEditorUtils::RemoveNamedSlotHostContent(DroppingWidget, SourceNamedSlotHost);
+			INamedSlotInterface* SourceNamedSlotHost = Cast<INamedSlotInterface>(SourceNamedSlotHostWidget);
+			if (SourceNamedSlotHost != nullptr)
+			{
+				SourceNamedSlotHostWidget->SetFlags(RF_Transactional);
+				SourceNamedSlotHostWidget->Modify();
+				FWidgetBlueprintEditorUtils::RemoveNamedSlotHostContent(DroppingWidget, SourceNamedSlotHost);
+			}
 		}
 
 		DroppingWidget->RemoveFromParent();
diff --git a/Engine/Source/Editor/UMGEditor/Private/WidgetBlueprintEditorUtils.cpp b/Engine/Source/Editor/UMGEditor/Private/WidgetBlueprintEditorUtils.cpp
--- a/Engine/Source/Editor/UMGEditor/Private/WidgetBlueprintEditorUtils.cpp
+++ b/Engine/Source/Editor/UMGEditor/Private/WidgetBlueprintEditorUtils.cpp
@@ -352,11 +352,16 @@ void FWidgetBlueprintEditorUtils::DeleteWidgets(UWidgetBlueprint* BP, TSet<FWidg
 
 INamedSlotInterface* FWidgetBlueprintEditorUtils::FindNamedSlotHostForContent(UWidget* WidgetTemplate, UWidgetTree* WidgetTree)
 {
-	INamedSlotInterface* Host = nullptr;
+	return Cast<INamedSlotInterface>(FindNamedSlotHostWidgetForContent(WidgetTemplate, WidgetTree));
+}
+
+UWidget* FWidgetBlueprintEditorUtils::FindNamedSlotHostWidgetForContent(UWidget* WidgetTemplate, UWidgetTree* WidgetTree)
+{
+	UWidget* HostWidget = nullptr;
 
 	WidgetTree->ForEachWidget([&](UWidget* Widget) {
 
-		if (Host != nullptr)
+		if (HostWidget != nullptr)
 		{
 			return;
 		}
@@ -372,14 +377,14 @@ INamedSlotInterface* FWidgetBlueprintEditorUtils::FindNamedSlotHostForContent(UW
 				{
 					if (SlotContent == WidgetTemplate)
 					{
-						Host = NamedSlotHost;
+						HostWidget = Widget;
 					}
 				}
 			}
 		}
 	});
 
-	return Host;
+	return HostWidget;
 }
 
 bool FWidgetBlueprintEditorUtils::RemoveNamedSlotHostContent(UWidget* WidgetTemplate, INamedSlotInterface* NamedSlotHost)
diff --git a/Engine/Source/Editor/UMGEditor/Private/WidgetBlueprintEditorUtils.h b/Engine/Source/Editor/UMGEditor/Private/WidgetBlueprintEditorUtils.h
--- a/Engine/Source/Editor/UMGEditor/Private/WidgetBlueprintEditorUtils.h
+++ b/Engine/Source/Editor/UMGEditor/Private/WidgetBlueprintEditorUtils.h
@@ -42,6 +42,8 @@ public:
 
 	static INamedSlotInterface* FindNamedSlotHostForContent(UWidget* WidgetTemplate, UWidgetTree* WidgetTree);
 
+	static UWidget* FindNamedSlotHostWidgetForContent(UWidget* WidgetTemplate, UWidgetTree* WidgetTree);
+
 	static bool RemoveNamedSlotHostContent(UWidget* WidgetTemplate, INamedSlotInterface* NamedSlotHost);
 
 private:
diff --git a/Engine/Source/Runtime/UMG/Private/Components/ExpandableArea.cpp b/Engine/Source/Runtime/UMG/Private/Components/ExpandableArea.cpp
--- a/Engine/Source/Runtime/UMG/Private/Components/ExpandableArea.cpp
+++ b/Engine/Source/Runtime/UMG/Private/Components/ExpandableArea.cpp
@@ -82,12 +82,6 @@ void UExpandableArea::SetContentForSlot(FName SlotName, UWidget* Content)
 	}
 }
 
-void UExpandableArea::ModifySlots()
-{
-	SetFlags(RF_Transactional);
-	Modify();
-}
-
 TSharedRef<SWidget> UExpandableArea::RebuildWidget()
 {
 	TSharedRef<SWidget> HeaderWidget = HeaderContent ? HeaderContent->TakeWidget() : SNullWidget::NullWidget;
diff --git a/Engine/Source/Runtime/UMG/Private/UserWidget.cpp b/Engine/Source/Runtime/UMG/Private/UserWidget.cpp
--- a/Engine/Source/Runtime/UMG/Private/UserWidget.cpp
+++ b/Engine/Source/Runtime/UMG/Private/UserWidget.cpp
@@ -466,12 +466,6 @@ void UUserWidget::SetContentForSlot(FName SlotName, UWidget* Content)
 	}
 }
 
-void UUserWidget::ModifySlots()
-{
-	SetFlags(RF_Transactional);
-	Modify();
-}
-
 TSharedRef<SWidget> UUserWidget::MakeViewportWidget(TSharedPtr<SWidget>& UserSlateWidget)
 {
 	UserSlateWidget = TakeWidget();
diff --git a/Engine/Source/Runtime/UMG/Public/Blueprint/UserWidget.h b/Engine/Source/Runtime/UMG/Public/Blueprint/UserWidget.h
--- a/Engine/Source/Runtime/UMG/Public/Blueprint/UserWidget.h
+++ b/Engine/Source/Runtime/UMG/Public/Blueprint/UserWidget.h
@@ -171,7 +171,6 @@ public:
 	virtual void GetSlotNames(TArray<FName>& SlotNames) const override;
 	virtual UWidget* GetContentForSlot(FName SlotName) const override;
 	virtual void SetContentForSlot(FName SlotName, UWidget* Content) override;
-	virtual void ModifySlots() override;
 	// UNamedSlotInterface End
 
 	/**
diff --git a/Engine/Source/Runtime/UMG/Public/Components/ExpandableArea.h b/Engine/Source/Runtime/UMG/Public/Components/ExpandableArea.h
--- a/Engine/Source/Runtime/UMG/Public/Components/ExpandableArea.h
+++ b/Engine/Source/Runtime/UMG/Public/Components/ExpandableArea.h
@@ -53,7 +53,6 @@ public:
 	virtual void GetSlotNames(TArray<FName>& SlotNames) const override;
 	virtual UWidget* GetContentForSlot(FName SlotName) const override;
 	virtual void SetContentForSlot(FName SlotName, UWidget* Content) override;
-	virtual void ModifySlots() override;
 	// End INamedSlotInterface
 
 public:
diff --git a/Engine/Source/Runtime/UMG/Public/Components/NamedSlotInterface.h b/Engine/Source/Runtime/UMG/Public/Components/NamedSlotInterface.h
--- a/Engine/Source/Runtime/UMG/Public/Components/NamedSlotInterface.h
+++ b/Engine/Source/Runtime/UMG/Public/Components/NamedSlotInterface.h
@@ -28,6 +28,4 @@ class UMG_API INamedSlotInterface
 
 	/**  */
 	bool ContainsContent(UWidget* Content) const;
-
-	virtual void ModifySlots() = 0;
 };
