Possible fix for linux compile error

diff --git a/Engine/Source/Editor/UMGEditor/Public/SlateVectorArtDataFactory.h b/Engine/Source/Editor/UMGEditor/Public/SlateVectorArtDataFactory.h
--- a/Engine/Source/Editor/UMGEditor/Public/SlateVectorArtDataFactory.h
+++ b/Engine/Source/Editor/UMGEditor/Public/SlateVectorArtDataFactory.h
@@ -1,6 +1,8 @@
 // Copyright 1998-2015 Epic Games, Inc. All Rights Reserved.
 
 #pragma once
+
+#include "Factories/Factory.h"
 #include "SlateVectorArtDataFactory.generated.h"
 
 UCLASS(hidecategories = Object, collapsecategories, MinimalAPI)
