Sequencer: Fix autokey enabled check which got inverted when adding autokey modes.

diff --git a/Engine/Source/Editor/UMGEditor/Private/Details/SWidgetDetailsView.cpp b/Engine/Source/Editor/UMGEditor/Private/Details/SWidgetDetailsView.cpp
--- a/Engine/Source/Editor/UMGEditor/Private/Details/SWidgetDetailsView.cpp
+++ b/Engine/Source/Editor/UMGEditor/Private/Details/SWidgetDetailsView.cpp
@@ -467,7 +467,7 @@ void SWidgetDetailsView::HandleIsVariableChanged(ECheckBoxState CheckState)
 void SWidgetDetailsView::NotifyPreChange(FEditPropertyChain* PropertyAboutToChange)
 {
 	// During auto-key do not migrate values
-	if( BlueprintEditor.Pin()->GetSequencer()->GetAutoKeyMode() != EAutoKeyMode::KeyNone )
+	if( BlueprintEditor.Pin()->GetSequencer()->GetAutoKeyMode() == EAutoKeyMode::KeyNone )
 	{
 		TSharedPtr<FWidgetBlueprintEditor> Editor = BlueprintEditor.Pin();
 
@@ -480,7 +480,7 @@ void SWidgetDetailsView::NotifyPostChange(const FPropertyChangedEvent& PropertyC
 {
 	const static FName DesignerRebuildName("DesignerRebuild");
 
-	if ( PropertyChangedEvent.ChangeType != EPropertyChangeType::Interactive && BlueprintEditor.Pin()->GetSequencer()->GetAutoKeyMode() != EAutoKeyMode::KeyNone )
+	if ( PropertyChangedEvent.ChangeType != EPropertyChangeType::Interactive && BlueprintEditor.Pin()->GetSequencer()->GetAutoKeyMode() == EAutoKeyMode::KeyNone )
 	{
 		TSharedPtr<FWidgetBlueprintEditor> Editor = BlueprintEditor.Pin();
 
