Fixed access Violation crash deleting a Widget Blueprint with an unsaved Animation track

Should have been fixed in dev-editor before merge, was incorrectly labled as not a regression

#rnx

diff --git a/Engine/Source/Editor/UMGEditor/Private/WidgetBlueprintEditor.cpp b/Engine/Source/Editor/UMGEditor/Private/WidgetBlueprintEditor.cpp
--- a/Engine/Source/Editor/UMGEditor/Private/WidgetBlueprintEditor.cpp
+++ b/Engine/Source/Editor/UMGEditor/Private/WidgetBlueprintEditor.cpp
@@ -1046,14 +1046,14 @@ void FWidgetBlueprintEditor::OnGetAnimationAddMenuContent(FMenuBuilder& MenuBuil
 		MenuBuilder.AddSubMenu(
 			LOCTEXT("AllNamedWidgets", "All Named Widgets"),
 			LOCTEXT("AllNamedWidgetsTooltip", "Select a widget or slot to create an animation track for"),
-			FNewMenuDelegate::CreateRaw(this, &FWidgetBlueprintEditor::OnGetAnimationAddMenuContentAllWidgets, InSequencer),
+			FNewMenuDelegate::CreateRaw(this, &FWidgetBlueprintEditor::OnGetAnimationAddMenuContentAllWidgets),
 			false,
 			FSlateIcon()
 		);
 	}
 }
 
-void FWidgetBlueprintEditor::OnGetAnimationAddMenuContentAllWidgets(FMenuBuilder& MenuBuilder, TSharedRef<ISequencer> InSequencer)
+void FWidgetBlueprintEditor::OnGetAnimationAddMenuContentAllWidgets(FMenuBuilder& MenuBuilder)
 {
 	MenuBuilder.AddSearchWidget();
 
diff --git a/Engine/Source/Editor/UMGEditor/Private/WidgetBlueprintEditor.h b/Engine/Source/Editor/UMGEditor/Private/WidgetBlueprintEditor.h
--- a/Engine/Source/Editor/UMGEditor/Private/WidgetBlueprintEditor.h
+++ b/Engine/Source/Editor/UMGEditor/Private/WidgetBlueprintEditor.h
@@ -235,7 +235,7 @@ private:
 	void OnGetAnimationAddMenuContent(FMenuBuilder& MenuBuilder, TSharedRef<ISequencer> Sequencer);
 
 	/** Populates the sequencer add submenu for the big list of widgets. */
-	void OnGetAnimationAddMenuContentAllWidgets(FMenuBuilder& MenuBuilder, TSharedRef<ISequencer> Sequencer);
+	void OnGetAnimationAddMenuContentAllWidgets(FMenuBuilder& MenuBuilder);
 
 	/** Adds the supplied UObject to the current animation. */
 	void AddObjectToAnimation(UObject* ObjectToAnimate);
