Removing legacy uses of GetEditorIcon from UMG

diff --git a/Engine/Source/Editor/UMGEditor/Private/Details/SWidgetDetailsView.cpp b/Engine/Source/Editor/UMGEditor/Private/Details/SWidgetDetailsView.cpp
--- a/Engine/Source/Editor/UMGEditor/Private/Details/SWidgetDetailsView.cpp
+++ b/Engine/Source/Editor/UMGEditor/Private/Details/SWidgetDetailsView.cpp
@@ -30,6 +30,7 @@
 #include "Blueprint/WidgetTree.h"
 #include "WidgetBlueprintEditorUtils.h"
 #include "ScopedTransaction.h"
+#include "Styling/SlateIconFinder.h"
 
 #define LOCTEXT_NAMESPACE "UMG"
 
@@ -372,19 +373,16 @@ FText SWidgetDetailsView::GetCategoryText() const
 	return FText::GetEmpty();
 }
 
-const FSlateBrush* GetEditorIcon_Deprecated(UWidget* Widget);
-
 const FSlateBrush* SWidgetDetailsView::GetNameIcon() const
 {
 	if ( SelectedObjects.Num() == 1 )
 	{
-		UWidget* Widget = Cast<UWidget>(SelectedObjects[0].Get());
-		if ( Widget )
+		UClass* WidgetClass = SelectedObjects[0].Get()->GetClass();
+		if (WidgetClass)
 		{
-			// @todo UMG: remove after 4.12
-			return GetEditorIcon_Deprecated(Widget);
-			// return FClassIconFinder::FindIconForClass(Widget->GetClass());
+			return FSlateIconFinder::FindIconBrushForClass(WidgetClass);
 		}
+		return nullptr;
 	}
 
 	return nullptr;
diff --git a/Engine/Source/Editor/UMGEditor/Private/Hierarchy/SHierarchyViewItem.cpp b/Engine/Source/Editor/UMGEditor/Private/Hierarchy/SHierarchyViewItem.cpp
--- a/Engine/Source/Editor/UMGEditor/Private/Hierarchy/SHierarchyViewItem.cpp
+++ b/Engine/Source/Editor/UMGEditor/Private/Hierarchy/SHierarchyViewItem.cpp
@@ -32,6 +32,7 @@
 #include "Blueprint/WidgetTree.h"
 #include "WidgetBlueprintEditorUtils.h"
 #include "ScopedTransaction.h"
+#include "Styling/SlateIconFinder.h"
 
 #define LOCTEXT_NAMESPACE "UMG"
 
@@ -926,13 +927,13 @@ FText FHierarchyWidget::GetLabelToolTipText() const
 	return FText::GetEmpty();
 }
 
-const FSlateBrush* GetEditorIcon_Deprecated(UWidget* Widget);
-
 const FSlateBrush* FHierarchyWidget::GetImage() const
 {
-	// @todo UMG: remove after 4.12
-	return GetEditorIcon_Deprecated(Item.GetTemplate());
-	// return FClassIconFinder::FindIconForClass(Item.GetTemplate()->GetClass());
+	if (Item.GetTemplate())
+	{
+		return FSlateIconFinder::FindIconBrushForClass(Item.GetTemplate()->GetClass());
+	}
+	return nullptr;
 }
 
 FSlateFontInfo FHierarchyWidget::GetFont() const
diff --git a/Engine/Source/Editor/UMGEditor/Private/Templates/WidgetTemplateBlueprintClass.cpp b/Engine/Source/Editor/UMGEditor/Private/Templates/WidgetTemplateBlueprintClass.cpp
--- a/Engine/Source/Editor/UMGEditor/Private/Templates/WidgetTemplateBlueprintClass.cpp
+++ b/Engine/Source/Editor/UMGEditor/Private/Templates/WidgetTemplateBlueprintClass.cpp
@@ -8,6 +8,7 @@
 #include "Toolkits/AssetEditorManager.h"
 
 #include "Kismet2/BlueprintEditorUtils.h"
+#include "Styling/SlateIconFinder.h"
 
 #define LOCTEXT_NAMESPACE "UMGEditor"
 
@@ -65,14 +66,10 @@ UWidget* FWidgetTemplateBlueprintClass::Create(UWidgetTree* Tree)
 	return FWidgetTemplateClass::CreateNamed(Tree, FName(*FBlueprintEditorUtils::GetClassNameWithoutSuffix(WidgetClass.Get())));
 }
 
-const FSlateBrush* GetEditorIcon_Deprecated(UWidget* Widget);
-
 const FSlateBrush* FWidgetTemplateBlueprintClass::GetIcon() const
 {
-	// @todo UMG: remove after 4.12
-	auto DefaultUserWidget = UUserWidget::StaticClass()->GetDefaultObject<UUserWidget>();
-	return GetEditorIcon_Deprecated(DefaultUserWidget);
-	// return FClassIconFinder::FindIconForClass(UUserWidget::StaticClass());
+	return FSlateIconFinder::FindIconBrushForClass(UUserWidget::StaticClass());
+	
 }
 
 TSharedRef<IToolTip> FWidgetTemplateBlueprintClass::GetToolTip() const
diff --git a/Engine/Source/Editor/UMGEditor/Private/Templates/WidgetTemplateClass.cpp b/Engine/Source/Editor/UMGEditor/Private/Templates/WidgetTemplateClass.cpp
--- a/Engine/Source/Editor/UMGEditor/Private/Templates/WidgetTemplateClass.cpp
+++ b/Engine/Source/Editor/UMGEditor/Private/Templates/WidgetTemplateClass.cpp
@@ -9,6 +9,7 @@
 #include "IDocumentation.h"
 
 #include "Blueprint/WidgetTree.h"
+#include "Styling/SlateIconFinder.h"
 
 #define LOCTEXT_NAMESPACE "UMGEditor"
 
@@ -44,14 +45,13 @@ UWidget* FWidgetTemplateClass::Create(UWidgetTree* Tree)
 	return CreateNamed(Tree, NAME_None);
 }
 
-const FSlateBrush* GetEditorIcon_Deprecated(UWidget* Widget);
-
 const FSlateBrush* FWidgetTemplateClass::GetIcon() const
 {
-	// @todo UMG: remove after 4.12
-	UWidget* DefaultWidget = WidgetClass->GetDefaultObject<UWidget>();
-	return GetEditorIcon_Deprecated(DefaultWidget);
-	// return FClassIconFinder::FindIconForClass(WidgetClass);
+	if (WidgetClass.IsValid())
+	{
+		return FSlateIconFinder::FindIconBrushForClass(WidgetClass.Get());
+	}
+	return nullptr;
 }
 
 TSharedRef<IToolTip> FWidgetTemplateClass::GetToolTip() const
