Fix for crash in widget details view accessing invalid object

diff --git a/Engine/Source/Editor/UMGEditor/Private/Details/SWidgetDetailsView.cpp b/Engine/Source/Editor/UMGEditor/Private/Details/SWidgetDetailsView.cpp
--- a/Engine/Source/Editor/UMGEditor/Private/Details/SWidgetDetailsView.cpp
+++ b/Engine/Source/Editor/UMGEditor/Private/Details/SWidgetDetailsView.cpp
@@ -375,7 +375,7 @@ FText SWidgetDetailsView::GetCategoryText() const
 
 const FSlateBrush* SWidgetDetailsView::GetNameIcon() const
 {
-	if ( SelectedObjects.Num() == 1 )
+	if ( SelectedObjects.Num() == 1 && SelectedObjects[0].IsValid())
 	{
 		UClass* WidgetClass = SelectedObjects[0].Get()->GetClass();
 		if (WidgetClass)
