Replicating the Blueprint Editor's method of add/view bound delegate widgets to avoid the crash on Outer != Graph

diff --git a/Engine/Source/Editor/UMGEditor/Private/Customizations/UMGDetailCustomizations.cpp b/Engine/Source/Editor/UMGEditor/Private/Customizations/UMGDetailCustomizations.cpp
--- a/Engine/Source/Editor/UMGEditor/Private/Customizations/UMGDetailCustomizations.cpp
+++ b/Engine/Source/Editor/UMGEditor/Private/Customizations/UMGDetailCustomizations.cpp
@@ -24,6 +24,7 @@
 #include "Kismet2/BlueprintEditorUtils.h"
 #include "Components/PanelSlot.h"
 #include "Details/SPropertyBinding.h"
+#include "Widgets/Layout/SWidgetSwitcher.h"
 
 #define LOCTEXT_NAMESPACE "UMG"
 
@@ -158,6 +159,45 @@ void FBlueprintWidgetCustomization::ResetToDefault_RemoveBinding(TSharedPtr<IPro
 	FBlueprintEditorUtils::MarkBlueprintAsStructurallyModified(Blueprint);
 }
 
+
+FReply FBlueprintWidgetCustomization::HandleAddOrViewEventForVariable(const FName EventName, FName PropertyName, TWeakObjectPtr<UClass> PropertyClass)
+{
+	UBlueprint* BlueprintObj = Blueprint;
+
+	// Find the corresponding variable property in the Blueprint
+	UObjectProperty* VariableProperty = FindField<UObjectProperty>(BlueprintObj->SkeletonGeneratedClass, PropertyName);
+
+	if (VariableProperty)
+	{
+		if (!FKismetEditorUtilities::FindBoundEventForComponent(BlueprintObj, EventName, VariableProperty->GetFName()))
+		{
+			FKismetEditorUtilities::CreateNewBoundEventForClass(PropertyClass.Get(), EventName, BlueprintObj, VariableProperty);
+		}
+		else
+		{
+			const UK2Node_ComponentBoundEvent* ExistingNode = FKismetEditorUtilities::FindBoundEventForComponent(BlueprintObj, EventName, VariableProperty->GetFName());
+			if (ExistingNode)
+			{
+				FKismetEditorUtilities::BringKismetToFocusAttentionOnObject(ExistingNode);
+			}
+		}
+	}
+
+	return FReply::Handled();
+}
+
+int32 FBlueprintWidgetCustomization::HandleAddOrViewIndexForButton(const FName EventName, FName PropertyName) const
+{
+	UBlueprint* BlueprintObj = Blueprint;
+
+	if (FKismetEditorUtilities::FindBoundEventForComponent(BlueprintObj, EventName, PropertyName))
+	{
+		return 0; // View
+	}
+
+	return 1; // Add
+}
+
 void FBlueprintWidgetCustomization::CreateMulticastEventCustomization(IDetailLayoutBuilder& DetailLayout, FName ThisComponentName, UClass* PropertyClass, UMulticastDelegateProperty* DelegateProperty)
 {
 	const FString AddString = FString(TEXT("Add "));
@@ -176,9 +216,6 @@ void FBlueprintWidgetCustomization::CreateMulticastEventCustomization(IDetailLay
 		PropertyTooltip = FText::FromString(DelegateProperty->GetName());
 	}
 
-	// Add on category for delegate property
-	const FText EventCategory = FObjectEditorUtils::GetCategoryText(DelegateProperty);
-
 	UObjectProperty* ComponentProperty = FindField<UObjectProperty>(Blueprint->SkeletonGeneratedClass, ThisComponentName);
 
 	if ( !ComponentProperty )
@@ -186,72 +223,59 @@ void FBlueprintWidgetCustomization::CreateMulticastEventCustomization(IDetailLay
 		return;
 	}
 
-	const UK2Node* EventNode = FKismetEditorUtilities::FindBoundEventForComponent(Blueprint, DelegateProperty->GetFName(), ComponentProperty->GetFName());
-
-	TSharedPtr<FEdGraphSchemaAction> ClickAction;
-	TSharedPtr<SWidget> ButtonContent;
-
-	if ( EventNode )
-	{
-		TSharedPtr<FEdGraphSchemaAction_K2ViewNode> NewDelegateNode = TSharedPtr<FEdGraphSchemaAction_K2ViewNode>(new FEdGraphSchemaAction_K2ViewNode(EventCategory, FText::FromString(ViewString + DelegateProperty->GetName()), PropertyTooltip, K2Schema->AG_LevelReference));
-		NewDelegateNode->NodePtr = EventNode;
-
-		ClickAction = NewDelegateNode;
-
-		ButtonContent = SNew(STextBlock)
-			.Text(LOCTEXT("ViewEvent", "View"));
-	}
-	else
-	{
-		TSharedPtr<FEdGraphSchemaAction_K2NewNode> NewDelegateNode = TSharedPtr<FEdGraphSchemaAction_K2NewNode>(new FEdGraphSchemaAction_K2NewNode(EventCategory, FText::FromString(AddString + DelegateProperty->GetName()), PropertyTooltip, K2Schema->AG_LevelReference));
-
-		UK2Node_ComponentBoundEvent* NewComponentEvent = NewObject<UK2Node_ComponentBoundEvent>(Blueprint, UK2Node_ComponentBoundEvent::StaticClass());
-		NewComponentEvent->InitializeComponentBoundEventParams(ComponentProperty, DelegateProperty);
-		NewDelegateNode->NodeTemplate = NewComponentEvent;
-		NewDelegateNode->bGotoNode = true;
+	FName PropertyName = ComponentProperty->GetFName();
+	FName EventName = DelegateProperty->GetFName();
+	FText EventText = DelegateProperty->GetDisplayNameText();
 
-		ClickAction = NewDelegateNode;
+	IDetailCategoryBuilder& EventCategory = DetailLayout.EditCategory(TEXT("Events"), LOCTEXT("Events", "Events"), ECategoryPriority::Uncommon);
 
-		ButtonContent = SNew(SImage)
-			.Image(FEditorStyle::GetBrush("Plus"));
-	}
-
-	TSharedRef<IPropertyHandle> DelegatePropertyHandle = DetailLayout.GetProperty(DelegateProperty->GetFName(), CastChecked<UClass>(DelegateProperty->GetOuter()));
-
-	IDetailCategoryBuilder& PropertyCategory = DetailLayout.EditCategory("Events", LOCTEXT("Events", "Events"), ECategoryPriority::Uncommon);
-
-	FText DelegatePropertyName = DelegateProperty->GetDisplayNameText();
-	PropertyCategory.AddCustomRow(DelegatePropertyName)
-	.NameContent()
-	[
-		SNew(SHorizontalBox)
-
-		+ SHorizontalBox::Slot()
-		.AutoWidth()
-		.VAlign(VAlign_Center)
-		.Padding(0,0,5,0)
+	EventCategory.AddCustomRow(EventText)
+		.NameContent()
 		[
-			SNew(SImage)
-			.Image(FEditorStyle::GetBrush("GraphEditor.Event_16x"))
-		]
+			SNew(SHorizontalBox)
+			.ToolTipText(DelegateProperty->GetToolTipText())
+			+ SHorizontalBox::Slot()
+			.AutoWidth()
+			.VAlign(VAlign_Center)
+			.Padding(0, 0, 5, 0)
+			[
+				SNew(SImage)
+				.Image(FEditorStyle::GetBrush("GraphEditor.Event_16x"))
+			]
 
-		+ SHorizontalBox::Slot()
-		.VAlign(VAlign_Center)
-		[
-			SNew(STextBlock)
-			.Font(IDetailLayoutBuilder::GetDetailFont())
-			.Text(DelegatePropertyName)
+			+ SHorizontalBox::Slot()
+			.VAlign(VAlign_Center)
+			[
+				SNew(STextBlock)
+				.Font(IDetailLayoutBuilder::GetDetailFont())
+				.Text(EventText)
+			]
 		]
-	]
-	.ValueContent()
-	.MinDesiredWidth(150)
-	.MaxDesiredWidth(200)
-	[
-		SNew(SGraphSchemaActionButton, Editor.Pin(), ClickAction)
+		.ValueContent()
+		.MinDesiredWidth(150)
+		.MaxDesiredWidth(200)
 		[
-			ButtonContent.ToSharedRef()
-		]
-	];
+			SNew(SButton)
+			.ButtonStyle(FEditorStyle::Get(), "FlatButton.Success")
+			.HAlign(HAlign_Center)
+			.OnClicked(this, &FBlueprintWidgetCustomization::HandleAddOrViewEventForVariable, EventName, PropertyName, MakeWeakObjectPtr(PropertyClass))
+			.ForegroundColor(FSlateColor::UseForeground())
+			[
+				SNew(SWidgetSwitcher)
+				.WidgetIndex(this, &FBlueprintWidgetCustomization::HandleAddOrViewIndexForButton, EventName, PropertyName)
+				+ SWidgetSwitcher::Slot()
+				[
+					SNew(STextBlock)
+					.Font(FEditorStyle::GetFontStyle(TEXT("BoldFont")))
+					.Text(LOCTEXT("ViewEvent", "View"))
+				]
+				+ SWidgetSwitcher::Slot()
+				[
+					SNew(SImage)
+					.Image(FEditorStyle::GetBrush("Plus"))
+				]
+			]
+		];
 }
 
 void FBlueprintWidgetCustomization::CustomizeDetails( IDetailLayoutBuilder& DetailLayout )
diff --git a/Engine/Source/Editor/UMGEditor/Private/Customizations/UMGDetailCustomizations.h b/Engine/Source/Editor/UMGEditor/Private/Customizations/UMGDetailCustomizations.h
--- a/Engine/Source/Editor/UMGEditor/Private/Customizations/UMGDetailCustomizations.h
+++ b/Engine/Source/Editor/UMGEditor/Private/Customizations/UMGDetailCustomizations.h
@@ -39,6 +39,10 @@ private:
 	void CreateMulticastEventCustomization(IDetailLayoutBuilder& DetailLayout, FName ThisComponentName, UClass* PropertyClass, UMulticastDelegateProperty* Property);
 
 	void ResetToDefault_RemoveBinding(TSharedPtr<IPropertyHandle> PropertyHandle);
+
+	FReply HandleAddOrViewEventForVariable(const FName EventName, FName PropertyName, TWeakObjectPtr<UClass> PropertyClass);
+
+	int32 HandleAddOrViewIndexForButton(const FName EventName, FName PropertyName) const;
 private:
 
 	TWeakPtr<FWidgetBlueprintEditor> Editor;
