AutoKey didn't work with position UMG widgets since we changes are done to a preview widget that is deleted after the drop happens. So to fix we just re-import the changes on the the new preview widget that we get from the widget reference. This will retrigger the auto-key onto the correct preview widget.

diff --git a/Engine/Source/Editor/UMGEditor/Private/Designer/SDesignerView.cpp b/Engine/Source/Editor/UMGEditor/Private/Designer/SDesignerView.cpp
--- a/Engine/Source/Editor/UMGEditor/Private/Designer/SDesignerView.cpp
+++ b/Engine/Source/Editor/UMGEditor/Private/Designer/SDesignerView.cpp
@@ -2895,17 +2895,35 @@ FReply SDesignerView::OnDrop(const FGeometry& MyGeometry, const FDragDropEvent&
 		FBlueprintEditorUtils::MarkBlueprintAsStructurallyModified(BP);
 
 		TSet<FWidgetReference> SelectedTemplates;
-
 		for (const auto& DropPreview : DropPreviews)
 		{
 			SelectedTemplates.Add(BlueprintEditor.Pin()->GetReferenceFromTemplate(DropPreview.Widget));
 		}
 
 		BlueprintEditor.Pin()->SelectWidgets(SelectedTemplates, false);
-
 		// Regenerate extension widgets now that we've finished moving or placing the widget.
 		CreateExtensionWidgetsForSelection();
 
+		UPanelSlot* Slot;
+		TSharedPtr<FSelectedWidgetDragDropOp> SelectedDragDropOp = DragDropEvent.GetOperationAs<FSelectedWidgetDragDropOp>();
+		UWidgetTree* WidgetTree = BP->WidgetTree;
+		if (WidgetTree && SelectedDragDropOp.IsValid())
+		{
+			for (auto& DraggedWidget : SelectedDragDropOp->DraggedWidgets)
+			{
+				FWidgetReference Reference = BlueprintEditor.Pin()->GetReferenceFromTemplate(DraggedWidget.Template);
+				UWidget* PreviewWidget = Reference.GetPreview();
+				if (PreviewWidget && PreviewWidget->GetParent()) 
+				{
+					Slot = PreviewWidget->GetParent()->GetSlots()[PreviewWidget->GetParent()->GetChildIndex(PreviewWidget)];
+					if (Slot != nullptr)
+					{
+						FWidgetBlueprintEditorUtils::ImportPropertiesFromText(Slot, DraggedWidget.ExportedSlotProperties);
+					}
+				}
+			}
+		}
+
 		DropPreviews.Empty();
 		return FReply::Handled().SetUserFocus(SharedThis(this));
 	}
