Fix for Win 64 warnings.
#trivial

diff --git a/Engine/Source/Editor/UMGEditor/Private/Designer/SDesignerView.cpp b/Engine/Source/Editor/UMGEditor/Private/Designer/SDesignerView.cpp
--- a/Engine/Source/Editor/UMGEditor/Private/Designer/SDesignerView.cpp
+++ b/Engine/Source/Editor/UMGEditor/Private/Designer/SDesignerView.cpp
@@ -2905,17 +2905,16 @@ FReply SDesignerView::OnDrop(const FGeometry& MyGeometry, const FDragDropEvent&
 		CreateExtensionWidgetsForSelection();
 
 		UPanelSlot* Slot;
-		TSharedPtr<FSelectedWidgetDragDropOp> SelectedDragDropOp = DragDropEvent.GetOperationAs<FSelectedWidgetDragDropOp>();
 		UWidgetTree* WidgetTree = BP->WidgetTree;
 		if (WidgetTree && SelectedDragDropOp.IsValid())
 		{
 			for (auto& DraggedWidget : SelectedDragDropOp->DraggedWidgets)
 			{
 				FWidgetReference Reference = BlueprintEditor.Pin()->GetReferenceFromTemplate(DraggedWidget.Template);
-				UWidget* PreviewWidget = Reference.GetPreview();
-				if (PreviewWidget && PreviewWidget->GetParent()) 
+				UWidget* LocalPreviewWidget = Reference.GetPreview();
+				if (LocalPreviewWidget && LocalPreviewWidget->GetParent())
 				{
-					Slot = PreviewWidget->GetParent()->GetSlots()[PreviewWidget->GetParent()->GetChildIndex(PreviewWidget)];
+					Slot = LocalPreviewWidget->GetParent()->GetSlots()[LocalPreviewWidget->GetParent()->GetChildIndex(LocalPreviewWidget)];
 					if (Slot != nullptr)
 					{
 						FWidgetBlueprintEditorUtils::ImportPropertiesFromText(Slot, DraggedWidget.ExportedSlotProperties);
