Fix support for toolbar extensions in the UMG editor

#rnx

diff --git a/Engine/Source/Editor/UMGEditor/Private/BlueprintModes/WidgetDesignerApplicationMode.cpp b/Engine/Source/Editor/UMGEditor/Private/BlueprintModes/WidgetDesignerApplicationMode.cpp
--- a/Engine/Source/Editor/UMGEditor/Private/BlueprintModes/WidgetDesignerApplicationMode.cpp
+++ b/Engine/Source/Editor/UMGEditor/Private/BlueprintModes/WidgetDesignerApplicationMode.cpp
@@ -7,6 +7,7 @@
 #include "BlueprintEditorSharedTabFactories.h"
 
 #include "WidgetBlueprintEditorToolbar.h"
+#include "UMGEditorModule.h"
 
 
 #include "TabFactory/PaletteTabSummoner.h"
@@ -116,7 +117,11 @@ FWidgetDesignerApplicationMode::FWidgetDesignerApplicationMode(TSharedPtr<FWidge
 
 	// setup toolbar - clear existing toolbar extender from the BP mode
 	//@TODO: Keep this in sync with BlueprintEditorModes.cpp
-	ToolbarExtender = MakeShareable(new FExtender);
+
+	//Make sure we start with our existing list of extenders instead of creating a new one
+	IUMGEditorModule& UMGEditorModule = FModuleManager::LoadModuleChecked<IUMGEditorModule>("UMGEditor");
+	ToolbarExtender = UMGEditorModule.GetToolBarExtensibilityManager()->GetAllExtenders();
+	
 	InWidgetEditor->GetWidgetToolbarBuilder()->AddWidgetBlueprintEditorModesToolbar(ToolbarExtender);
 	InWidgetEditor->GetWidgetToolbarBuilder()->AddWidgetReflector(ToolbarExtender);
 	InWidgetEditor->GetToolbarBuilder()->AddCompileToolbar(ToolbarExtender);
