PR #4663: UE-20103: Slate widgets retain their category name v2 (Contributed by [REDACTED_PERSON])

diff --git a/Engine/Source/Editor/UMGEditor/Private/WidgetBlueprint.cpp b/Engine/Source/Editor/UMGEditor/Private/WidgetBlueprint.cpp
--- a/Engine/Source/Editor/UMGEditor/Private/WidgetBlueprint.cpp
+++ b/Engine/Source/Editor/UMGEditor/Private/WidgetBlueprint.cpp
@@ -5,9 +5,7 @@
 #include "Blueprint/UserWidget.h"
 #include "MovieScene.h"
 
-#if WITH_EDITOR
-	#include "Engine/UserDefinedStruct.h"
-#endif // WITH_EDITOR
+#include "Engine/UserDefinedStruct.h"
 #include "EdGraph/EdGraph.h"
 #include "Blueprint/WidgetTree.h"
 #include "Animation/WidgetAnimation.h"
@@ -16,16 +14,11 @@
 
 #include "Kismet2/CompilerResultsLog.h"
 #include "Binding/PropertyBinding.h"
-#include "Blueprint/WidgetTree.h"
 #include "Blueprint/WidgetBlueprintGeneratedClass.h"
 #include "UObject/PropertyTag.h"
-#include "WidgetBlueprint.h"
 #include "WidgetBlueprintCompiler.h"
-#include "Binding/PropertyBinding.h"
-#include "Engine/UserDefinedStruct.h"
 #include "UObject/EditorObjectVersion.h"
 #include "WidgetGraphSchema.h"
-#include "WidgetBlueprintCompiler.h"
 #include "UMGEditorProjectSettings.h"
 
 #if WITH_EDITOR
@@ -528,8 +521,6 @@ UWidgetBlueprint::UWidgetBlueprint(const FObjectInitializer& ObjectInitializer)
 	, TickFrequency(EWidgetTickFrequency::Auto)
 	, SupportDynamicCreation(EWidgetSupportsDynamicCreation::Default)
 {
-	WidgetTree = CreateDefaultSubobject<UWidgetTree>(TEXT("WidgetTree"));
-	WidgetTree->SetFlags(RF_Transactional);
 }
 
 void UWidgetBlueprint::ReplaceDeprecatedNodes()
@@ -765,61 +756,11 @@ bool UWidgetBlueprint::IsWidgetFreeFromCircularReferences(UUserWidget* UserWidge
 	return true;
 }
 
-TArray<UWidget*> UWidgetBlueprint::GetAllSourceWidgets()
-{
-	TArray<UWidget*> Ret;
-	ForEachSourceWidgetImpl( [&Ret](UWidget* Inner) { Ret.Push(Inner); } );
-	return Ret;
-}
-
-TArray<const UWidget*> UWidgetBlueprint::GetAllSourceWidgets() const
-{
-	TArray<const UWidget*> Ret;
-	ForEachSourceWidgetImpl( [&Ret](UWidget* Inner) { Ret.Push(Inner); } );
-	return Ret;
-}
-
-void UWidgetBlueprint::ForEachSourceWidget(TFunctionRef<void(UWidget*)> Fn)
-{
-	ForEachSourceWidgetImpl(Fn);
-}
-
-void UWidgetBlueprint::ForEachSourceWidget(TFunctionRef<void(const UWidget*)> Fn) const
-{
-	ForEachSourceWidgetImpl(Fn);
-}
-
 UPackage* UWidgetBlueprint::GetWidgetTemplatePackage() const
 {
 	return GetOutermost();
 }
 
-void UWidgetBlueprint::ForEachSourceWidgetImpl (TFunctionRef<void(UWidget*)> Fn) const
-{
-	// This exists in order to facilitate working with collections of UWidgets wo/ 
-	// relying on user implemented UWidget virtual functions. During blueprint compilation
-	// it is bad practice to call those virtual functions until the class is fully formed
-	// and reinstancing has finished. For instance, GetDefaultObject() calls in those user
-	// functions may create a CDO before the class has been linked, or even before
-	// all member variables have been generated:
-	UWidgetTree* WidgetTreeForCapture = WidgetTree;
-	ForEachObjectWithOuter(
-		WidgetTree, 
-		[Fn, WidgetTreeForCapture](UObject* Inner)
-		{
-			if(UWidget* AsWidget = Cast<UWidget>(Inner))
-			{
-				// Widgets owned by another UWidgetBlueprint aren't really 'source' widgets, E.g. widgets
-				// created by the user *in this blueprint*
-				if(AsWidget->GetTypedOuter<UWidgetTree>() == WidgetTreeForCapture)
-				{
-					Fn(AsWidget);
-				}
-			}
-		}
-	);
-}
-
 static bool HasLatentActions(UEdGraph* Graph)
 {
 	for (const UEdGraphNode* Node : Graph->Nodes)
diff --git a/Engine/Source/Editor/UMGEditor/Private/WidgetBlueprintCompiler.cpp b/Engine/Source/Editor/UMGEditor/Private/WidgetBlueprintCompiler.cpp
--- a/Engine/Source/Editor/UMGEditor/Private/WidgetBlueprintCompiler.cpp
+++ b/Engine/Source/Editor/UMGEditor/Private/WidgetBlueprintCompiler.cpp
@@ -273,8 +273,10 @@ void FWidgetBlueprintCompiler::CreateClassVariablesFromBlueprint()
 			{
 				WidgetProperty->SetPropertyFlags(CPF_BlueprintVisible);
 
+				const FString& CategoryName = Widget->GetCategoryName();
+				
 				// Only include Category metadata for variables (i.e. a visible/editable property); otherwise, UHT will raise a warning if this Blueprint is nativized.
-				WidgetProperty->SetMetaData(TEXT("Category"), *WidgetBP->GetName());
+				WidgetProperty->SetMetaData(TEXT("Category"), *(CategoryName.IsEmpty() ? WidgetBP->GetName() : CategoryName));
 			}
 
 			WidgetProperty->SetPropertyFlags(CPF_Instanced);
diff --git a/Engine/Source/Editor/UMGEditor/Public/WidgetBlueprint.h b/Engine/Source/Editor/UMGEditor/Public/WidgetBlueprint.h
--- a/Engine/Source/Editor/UMGEditor/Public/WidgetBlueprint.h
+++ b/Engine/Source/Editor/UMGEditor/Public/WidgetBlueprint.h
@@ -6,7 +6,7 @@
 #include "UObject/ObjectMacros.h"
 #include "Misc/Guid.h"
 #include "UObject/Class.h"
-#include "Engine/Blueprint.h"
+#include "BaseWidgetBlueprint.h"
 #include "Binding/DynamicPropertyPath.h"
 #include "Blueprint/WidgetBlueprintGeneratedClass.h"
 #include "Animation/WidgetAnimationBinding.h"
@@ -206,10 +206,45 @@ enum class EWidgetCompileTimeTickPrediction : uint8
  * The widget blueprint enables extending UUserWidget the user extensible UWidget.
  */
 UCLASS(BlueprintType)
-class UMGEDITOR_API UWidgetBlueprint : public UBlueprint
+class UMGEDITOR_API UWidgetBlueprint : public UBaseWidgetBlueprint
 {
 	GENERATED_UCLASS_BODY()
 
+public:
+
+#if WITH_EDITORONLY_DATA
+	
+	UPROPERTY()
+	TArray< FDelegateEditorBinding > Bindings;
+
+	UPROPERTY()
+	TArray<FWidgetAnimation_DEPRECATED> AnimationData_DEPRECATED;
+
+	UPROPERTY()
+	TArray<UWidgetAnimation*> Animations;
+
+	/**
+	 * Don't directly modify this property to change the palette category.  The actual value is stored 
+	 * in the CDO of the UUserWidget, but a copy is stored here so that it's available in the serialized 
+	 * Tag data in the asset header for access in the FAssetData.
+	 */
+	UPROPERTY(AssetRegistrySearchable, AssetRegistrySearchable)
+	FString PaletteCategory;
+
+	UPROPERTY(EditAnywhere, AdvancedDisplay, Category=WidgetBlueprintOptions, AssetRegistrySearchable)
+	bool bForceSlowConstructionPath;
+
+private:
+	/**
+	 * Widgets by default all support calling CreateWidget for them, however for mobile games
+	 * you may want to disable this by default, or on a per widget basis as it can save several
+	 * MB on a large game from lots of widget templates being cooked ready to make dynamic
+	 * construction faster.
+	 */
+	UPROPERTY(EditAnywhere, AdvancedDisplay, Category=WidgetBlueprintOptions, AssetRegistrySearchable)
+	EWidgetSupportsDynamicCreation SupportDynamicCreation;
+#endif
+
 public:
 
 	/** UObject interface */
@@ -245,18 +280,6 @@ public:
 	/** Returns true if the supplied user widget will not create a circular reference when added to this blueprint */
 	bool IsWidgetFreeFromCircularReferences(UUserWidget* UserWidget) const;
 
-	/** 
-	 * Returns collection of widgets that represent the 'source' (user edited) widgets for this 
-	 * blueprint - avoids calling virtual functions on instances and is therefore safe to use 
-	 * throughout compilation.
-	 */
-	TArray<UWidget*> GetAllSourceWidgets();
-	TArray<const UWidget*> GetAllSourceWidgets() const;
-
-	/** Identical to GetAllSourceWidgets, but as an algorithm */
-	void ForEachSourceWidget(TFunctionRef<void(UWidget*)> Fn);
-	void ForEachSourceWidget(TFunctionRef<void(const UWidget*)> Fn) const;
-
 	bool WidgetSupportsDynamicCreation() const;
 
 	static bool ValidateGeneratedClass(const UClass* InClass);
@@ -266,79 +289,41 @@ public:
 	void UpdateTickabilityStats(bool& OutHasLatentActions, bool& OutHasAnimations, bool& OutClassRequiresNativeTick);
 
 private:
-	void ForEachSourceWidgetImpl(TFunctionRef<void(UWidget*)> Fn) const;
-
 #if WITH_EDITOR
 	virtual void LoadModulesRequiredForCompilation() override;
-#endif 
 
 public:
 
-#if WITH_EDITORONLY_DATA
-	/** A tree of the widget templates to be created */
-	UPROPERTY()
-	class UWidgetTree* WidgetTree;
-
-	UPROPERTY()
-	TArray<FDelegateEditorBinding> Bindings;
-
-	UPROPERTY()
-	TArray<FWidgetAnimation_DEPRECATED> AnimationData_DEPRECATED;
-
-	UPROPERTY()
-	TArray<UWidgetAnimation*> Animations;
-
 	/**
-	* Don't directly modify this property to change the palette category.  The actual value is stored
-	* in the CDO of the UUserWidget, but a copy is stored here so that it's available in the serialized
-	* Tag data in the asset header for access in the FAssetData.
+	* The total number of widgets this widget contains.  This is a good way to find the "largest" widgets.
 	*/
 	UPROPERTY(AssetRegistrySearchable)
-	FString PaletteCategory;
-
-	UPROPERTY(EditAnywhere, AdvancedDisplay, Category = WidgetBlueprintOptions, AssetRegistrySearchable)
-	bool bForceSlowConstructionPath;
-
-	/**
-	 * The total number of widgets this widget contains.  This is a good way to find the "largest" widgets.
-	 */
-	UPROPERTY(AssetRegistrySearchable)
 	int32 InclusiveWidgets;
 
 private:
 	/**
-	 * The desired tick frequency set by the user on the UserWidget's CDO.
-	 */
+	* The desired tick frequency set by the user on the UserWidget's CDO.
+	*/
 	UPROPERTY(AssetRegistrySearchable)
 	EWidgetTickFrequency TickFrequency;
 
 	/**
-	 * The computed frequency that the widget will need to be ticked at.  You can find the reasons for
-	 * this decision by looking at TickPredictionReason.
-	 */
+	* The computed frequency that the widget will need to be ticked at.  You can find the reasons for
+	* this decision by looking at TickPredictionReason.
+	*/
 	UPROPERTY(AssetRegistrySearchable)
 	EWidgetCompileTimeTickPrediction TickPrediction;
 
 	/**
-	 * The reasons we may need to tick this widget.
-	 */
+	* The reasons we may need to tick this widget.
+	*/
 	UPROPERTY(AssetRegistrySearchable)
 	FString TickPredictionReason;
 
 	/**
-	 * Widgets by default all support calling CreateWidget for them, however for mobile games
-	 * you may want to disable this by default, or on a per widget basis as it can save several
-	 * MB on a large game from lots of widget templates being cooked ready to make dynamic
-	 * construction faster.
-	 */
-	UPROPERTY(EditAnywhere, AdvancedDisplay, Category = WidgetBlueprintOptions, AssetRegistrySearchable)
-	EWidgetSupportsDynamicCreation SupportDynamicCreation;
-
-	/**
-	 * The total number of property bindings.  Consider this as a performance warning.
-	 */
+	* The total number of property bindings.  Consider this as a performance warning.
+	*/
 	UPROPERTY(AssetRegistrySearchable)
 	int32 PropertyBindings;
-
-#endif
+#endif 
 };
diff --git a/Engine/Source/Editor/UnrealEd/Private/BaseWidgetBlueprint.cpp b/Engine/Source/Editor/UnrealEd/Private/BaseWidgetBlueprint.cpp
new file mode 100644
--- /dev/null
+++ b/Engine/Source/Editor/UnrealEd/Private/BaseWidgetBlueprint.cpp
@@ -0,0 +1,63 @@
+// Copyright 1998-2018 Epic Games, Inc. All Rights Reserved.
+
+#include "BaseWidgetBlueprint.h"
+#include "Components/Widget.h"
+#include "Blueprint/WidgetTree.h"
+#include "UObject/UObjectHash.h"
+
+UBaseWidgetBlueprint::UBaseWidgetBlueprint(const FObjectInitializer& ObjectInitializer)
+	: Super(ObjectInitializer)
+{
+	WidgetTree = CreateDefaultSubobject<UWidgetTree>(TEXT("WidgetTree"));
+	WidgetTree->SetFlags(RF_Transactional);
+}
+
+TArray<UWidget*> UBaseWidgetBlueprint::GetAllSourceWidgets()
+{
+	TArray<UWidget*> Ret;
+	ForEachSourceWidgetImpl([&Ret](UWidget* Inner) { Ret.Push(Inner); });
+	return Ret;
+}
+
+TArray<const UWidget*> UBaseWidgetBlueprint::GetAllSourceWidgets() const
+{
+	TArray<const UWidget*> Ret;
+	ForEachSourceWidgetImpl([&Ret](UWidget* Inner) { Ret.Push(Inner); });
+	return Ret;
+}
+
+void UBaseWidgetBlueprint::ForEachSourceWidget(TFunctionRef<void(UWidget*)> Fn)
+{
+	ForEachSourceWidgetImpl(Fn);
+}
+
+void UBaseWidgetBlueprint::ForEachSourceWidget(TFunctionRef<void(const UWidget*)> Fn) const
+{
+	ForEachSourceWidgetImpl(Fn);
+}
+
+void UBaseWidgetBlueprint::ForEachSourceWidgetImpl(TFunctionRef<void(UWidget*)> Fn) const
+{
+	// This exists in order to facilitate working with collections of UWidgets wo/ 
+	// relying on user implemented UWidget virtual functions. During blueprint compilation
+	// it is bad practice to call those virtual functions until the class is fully formed
+	// and reinstancing has finished. For instance, GetDefaultObject() calls in those user
+	// functions may create a CDO before the class has been linked, or even before
+	// all member variables have been generated:
+	UWidgetTree* WidgetTreeForCapture = WidgetTree;
+	ForEachObjectWithOuter(
+		WidgetTree,
+		[Fn, WidgetTreeForCapture](UObject* Inner)
+	{
+		if (UWidget* AsWidget = Cast<UWidget>(Inner))
+		{
+			// Widgets owned by another UWidgetBlueprint aren't really 'source' widgets, E.g. widgets
+			// created by the user *in this blueprint*
+			if (AsWidget->GetTypedOuter<UWidgetTree>() == WidgetTreeForCapture)
+			{
+				Fn(AsWidget);
+			}
+		}
+	}
+	);
+}
diff --git a/Engine/Source/Editor/UnrealEd/Private/Kismet2/BlueprintEditorUtils.cpp b/Engine/Source/Editor/UnrealEd/Private/Kismet2/BlueprintEditorUtils.cpp
--- a/Engine/Source/Editor/UnrealEd/Private/Kismet2/BlueprintEditorUtils.cpp
+++ b/Engine/Source/Editor/UnrealEd/Private/Kismet2/BlueprintEditorUtils.cpp
@@ -48,7 +48,7 @@
 #include "Kismet2/CompilerResultsLog.h"
 #include "Algo/Transform.h"
 
-#include "Editor/KismetCompiler/Public/KismetCompilerModule.h"
+#include "KismetCompilerModule.h"
 #include "EdGraphSchema_K2_Actions.h"
 #include "K2Node_Event.h"
 #include "K2Node_ActorBoundEvent.h"
@@ -113,6 +113,8 @@
 
 #include "EditorCategoryUtils.h"
 #include "Styling/SlateIconFinder.h"
+#include "BaseWidgetBlueprint.h"
+#include "Components/Widget.h"
 
 extern COREUOBJECT_API bool GBlueprintUseCompilationManager;
 
@@ -3909,6 +3911,15 @@ void FBlueprintEditorUtils::SetBlueprintVariableCategory(UBlueprint* Blueprint,
 			{
 				Timeline->SetMetaData(TEXT("Category"), *SetCategory.ToString());
 			}
+			else if (UBaseWidgetBlueprint* WidgetBP = Cast<UBaseWidgetBlueprint>(Blueprint))
+			{
+				WidgetBP->ForEachSourceWidget([&](UWidget* InWidget) {
+					if (InWidget->GetFName() == VarName)
+					{
+						InWidget->SetCategoryName(SetCategory.ToString());
+					}
+				});
+			}
 			else
 			{
 				TargetProperty->SetMetaData(TEXT("Category"), *SetCategory.ToString());
diff --git a/Engine/Source/Editor/UnrealEd/Public/BaseWidgetBlueprint.h b/Engine/Source/Editor/UnrealEd/Public/BaseWidgetBlueprint.h
new file mode 100644
--- /dev/null
+++ b/Engine/Source/Editor/UnrealEd/Public/BaseWidgetBlueprint.h
@@ -0,0 +1,33 @@
+// Copyright 1998-2018 Epic Games, Inc. All Rights Reserved.
+
+#pragma once
+
+#include "Engine/Blueprint.h"
+#include "BaseWidgetBlueprint.generated.h"
+
+UCLASS(Abstract)
+class UNREALED_API UBaseWidgetBlueprint : public UBlueprint
+{
+	GENERATED_UCLASS_BODY()
+public:
+#if WITH_EDITORONLY_DATA
+	/** A tree of the widget templates to be created */
+	UPROPERTY()
+	class UWidgetTree* WidgetTree;
+#endif
+
+	/**
+	* Returns collection of widgets that represent the 'source' (user edited) widgets for this
+	* blueprint - avoids calling virtual functions on instances and is therefore safe to use
+	* throughout compilation.
+	*/
+	TArray<class UWidget*> GetAllSourceWidgets();
+	TArray<const class UWidget*> GetAllSourceWidgets() const;
+
+	/** Identical to GetAllSourceWidgets, but as an algorithm */
+	void ForEachSourceWidget(TFunctionRef<void(class UWidget*)> Fn);
+	void ForEachSourceWidget(TFunctionRef<void(const class UWidget*)> Fn) const;
+
+private:
+	void ForEachSourceWidgetImpl(TFunctionRef<void(class UWidget*)> Fn) const;
+};
diff --git a/Engine/Source/Runtime/Engine/Classes/Engine/Blueprint.h b/Engine/Source/Runtime/Engine/Classes/Engine/Blueprint.h
--- a/Engine/Source/Runtime/Engine/Classes/Engine/Blueprint.h
+++ b/Engine/Source/Runtime/Engine/Classes/Engine/Blueprint.h
@@ -706,9 +706,6 @@ public:
 	/** Find the object in the TemplateObjects array with the supplied name */
 	UActorComponent* FindTemplateByName(const FName& TemplateName) const;
 
-	/** Rename the component template in the TemplateObjects array with the supplied name */
-	bool RenameComponentTemplate(const FName& OldTemplateName, const FName& NewTemplateName);
-
 	/** Find a timeline by name */
 	class UTimelineTemplate* FindTimelineTemplateByVariableName(const FName& TimelineName);	
 
diff --git a/Engine/Source/Runtime/UMG/Private/Components/Widget.cpp b/Engine/Source/Runtime/UMG/Private/Components/Widget.cpp
--- a/Engine/Source/Runtime/UMG/Private/Components/Widget.cpp
+++ b/Engine/Source/Runtime/UMG/Private/Components/Widget.cpp
@@ -873,6 +873,16 @@ void UWidget::SetDisplayLabel(const FString& InDisplayLabel)
 	DisplayLabel = InDisplayLabel;
 }
 
+const FString& UWidget::GetCategoryName() const
+{
+	return CategoryName;
+}
+
+void UWidget::SetCategoryName(const FString& InValue)
+{
+	CategoryName = InValue;
+}
+
 bool UWidget::IsGeneratedName() const
 {
 	if (!DisplayLabel.IsEmpty())
diff --git a/Engine/Source/Runtime/UMG/Public/Components/Widget.h b/Engine/Source/Runtime/UMG/Public/Components/Widget.h
--- a/Engine/Source/Runtime/UMG/Public/Components/Widget.h
+++ b/Engine/Source/Runtime/UMG/Public/Components/Widget.h
@@ -695,6 +695,12 @@ public:
 	/** Sets the friendly name of the widget to display in the editor */
 	void SetDisplayLabel(const FString& DisplayLabel);
 
+	/** @return the category name of the widget */
+	const FString& GetCategoryName() const;
+
+	/** Sets the category name of the widget */
+	void SetCategoryName(const FString& InValue);
+
 	/**
 	 * Called at the end of Widget Blueprint compilation.
 	 * Allows UMG elements to evaluate their default states and determine whether they are acceptable.
@@ -876,6 +882,10 @@ private:
 
 	/** The underlying SWidget for the design time wrapper widget. */
 	TWeakPtr<SWidget> DesignWrapperWidget;
+
+	/** Category name used in the widget designer for sorting purpose */
+	UPROPERTY()
+	FString CategoryName;
 #endif
 
 #if !(UE_BUILD_SHIPPING || UE_BUILD_TEST)
