UMG Hierarchy now remembers expansion state when being destroyed and recreated (due to closing widget or switching to Graph view)

diff --git a/Engine/Source/Editor/UMGEditor/Private/Hierarchy/SHierarchyView.cpp b/Engine/Source/Editor/UMGEditor/Private/Hierarchy/SHierarchyView.cpp
--- a/Engine/Source/Editor/UMGEditor/Private/Hierarchy/SHierarchyView.cpp
+++ b/Engine/Source/Editor/UMGEditor/Private/Hierarchy/SHierarchyView.cpp
@@ -82,7 +82,6 @@ void SHierarchyView::Construct(const FArguments& InArgs, TSharedPtr<FWidgetBluep
 	BlueprintEditor.Pin()->OnSelectedWidgetsChanged.AddRaw(this, &SHierarchyView::OnEditorSelectionChanged);
 
 	bRefreshRequested = true;
-	bExpandAllNodes = true;
 }
 
 SHierarchyView::~SHierarchyView()
@@ -411,6 +410,11 @@ void SHierarchyView::RecursiveExpand(TSharedPtr<FHierarchyModel>& Model, EExpand
 	{
 		RecursiveExpand(ChildModel, ExpandBehavior);
 	}
+
+	if (Model->IsExpanded())
+	{
+		WidgetTreeView->SetItemExpansion(Model, true);
+	}
 }
 
 void SHierarchyView::RestoreSelectedItems()
