Improved UMG rename validation to respect the errors from the blueprint validator.
This fixes at least the case where it miss reported the issue when the name was greater than the length limit in blueprints.

diff --git a/Engine/Source/Editor/UMGEditor/Private/WidgetBlueprintEditorUtils.cpp b/Engine/Source/Editor/UMGEditor/Private/WidgetBlueprintEditorUtils.cpp
--- a/Engine/Source/Editor/UMGEditor/Private/WidgetBlueprintEditorUtils.cpp
+++ b/Engine/Source/Editor/UMGEditor/Private/WidgetBlueprintEditorUtils.cpp
@@ -153,12 +153,19 @@ bool FWidgetBlueprintEditorUtils::VerifyWidgetRename(TSharedRef<class FWidgetBlu
 	FKismetNameValidator Validator(Blueprint);
 
 	// For variable comparison, use the slug
-	const bool bUniqueNameForVariable = ( EValidatorResult::Ok == Validator.IsValid(NewNameSlug) );
+	EValidatorResult ValidatorResult = Validator.IsValid(NewNameSlug);
 
-	if (!bUniqueNameForVariable && !bIsSameWidget)
+	if (ValidatorResult != EValidatorResult::Ok)
 	{
-		OutErrorMessage = LOCTEXT("ExistingVariableName", "Existing Variable Name");
-		return false;
+		if (bIsSameWidget && (ValidatorResult == EValidatorResult::AlreadyInUse || ValidatorResult == EValidatorResult::ExistingName))
+		{
+			// Continue successfully
+		}
+		else
+		{
+			OutErrorMessage = INameValidatorInterface::GetErrorText(NewNameString, ValidatorResult);
+			return false;
+		}
 	}
 
 	return true;
