PR #4915: [UMG] Bind UWidgetAnimation from C++ to blueprint created animation (Contributed by [REDACTED_PERSON])

diff --git a/Engine/Source/Editor/UMGEditor/Private/TabFactory/AnimationTabSummoner.cpp b/Engine/Source/Editor/UMGEditor/Private/TabFactory/AnimationTabSummoner.cpp
--- a/Engine/Source/Editor/UMGEditor/Private/TabFactory/AnimationTabSummoner.cpp
+++ b/Engine/Source/Editor/UMGEditor/Private/TabFactory/AnimationTabSummoner.cpp
@@ -7,6 +7,7 @@
 #include "MovieScene.h"
 #include "Animation/WidgetAnimation.h"
 #include "WidgetBlueprint.h"
+#include "WidgetBlueprintEditorUtils.h"
 #include "Widgets/Input/SButton.h"
 #include "Widgets/Layout/SScrollBorder.h"
 #include "Widgets/Views/STableViewBase.h"
@@ -69,6 +70,13 @@ bool VerifyAnimationRename( FWidgetBlueprintEditor& BlueprintEditor, UWidgetAnim
 		return false;
 	}
 
+	// Check for BindWidgetAnim property
+	UProperty* Property = Blueprint->ParentClass->FindPropertyByName(NewAnimationNameAsName);
+	if (Property && FWidgetBlueprintEditorUtils::IsBindWidgetAnimProperty(Property))
+	{
+		return true;
+	}
+
 	FKismetNameValidator Validator( Blueprint );
 	EValidatorResult ValidationResult = Validator.IsValid( NewAnimationName );
 
@@ -136,7 +144,7 @@ private:
 	{
 		if( ListItem.IsValid() )
 		{
-			return FText::FromString( ListItem.Pin()->Animation->GetName() );
+			return ListItem.Pin()->Animation->GetDisplayName();
 		}
 
 		return FText::GetEmpty();
@@ -146,12 +154,12 @@ private:
 	{
 		UWidgetAnimation* Animation = ListItem.Pin()->Animation;
 
-		FString NewName = InText.ToString();
+		const FName NewName = MakeObjectNameFromDisplayLabel(InText.ToString(), Animation->GetFName());
 
-		if ( Animation->GetName() != NewName )
+		if ( Animation->GetFName() != NewName )
 		{
 			TSharedPtr<FWidgetBlueprintEditor> Editor = BlueprintEditor.Pin();
-			return Editor.IsValid() && VerifyAnimationRename( *Editor, Animation, NewName, OutErrorMessage );
+			return Editor.IsValid() && VerifyAnimationRename( *Editor, Animation, NewName.ToString(), OutErrorMessage );
 		}
 
 		return true;
@@ -162,23 +170,33 @@ private:
 		UWidgetAnimation* WidgetAnimation = ListItem.Pin()->Animation;
 		UWidgetBlueprint* Blueprint = BlueprintEditor.Pin()->GetWidgetBlueprintObj();
 
-		FName CurrentName = WidgetAnimation->GetFName();
+		// Get the new FName slug from the given display name
+		const FName NewFName = MakeObjectNameFromDisplayLabel(InText.ToString(), WidgetAnimation->GetFName());
+		const FName OldFName = WidgetAnimation->GetFName();
+
+		UObjectPropertyBase* ExistingProperty = Cast<UObjectPropertyBase>(Blueprint->ParentClass->FindPropertyByName(NewFName));
+		const bool bBindWidgetAnim = ExistingProperty && FWidgetBlueprintEditorUtils::IsBindWidgetAnimProperty(ExistingProperty) && ExistingProperty->PropertyClass->IsChildOf(UWidgetAnimation::StaticClass());
 
-		bool bNewAnimation = ListItem.Pin()->bNewAnimation;
+		const bool bValidName = !OldFName.IsEqual(NewFName) && !InText.IsEmpty();
+		const bool bCanRename = (bValidName || bBindWidgetAnim);
 
-		if(!CurrentName.ToString().Equals(InText.ToString()) && !InText.IsEmpty())
+		const bool bNewAnimation = ListItem.Pin()->bNewAnimation;
+		if (bCanRename)
 		{
-			FText TransactionName = bNewAnimation ? LOCTEXT("NewAnimation", "New Animation") : LOCTEXT("RenameAnimation", "Rename Animation");
+			const FString NewNameStr = NewFName.ToString();
+			const FString OldNameStr = OldFName.ToString();
 
+			FText TransactionName = bNewAnimation ? LOCTEXT("NewAnimation", "New Animation") : LOCTEXT("RenameAnimation", "Rename Animation");
 			{
 				const FScopedTransaction Transaction(TransactionName);
 				WidgetAnimation->Modify();
 				WidgetAnimation->GetMovieScene()->Modify();
 
-				WidgetAnimation->Rename(*InText.ToString());
-				WidgetAnimation->GetMovieScene()->Rename(*InText.ToString());
+				WidgetAnimation->SetDisplayLabel(InText.ToString());
+				WidgetAnimation->Rename(*NewNameStr);
+				WidgetAnimation->GetMovieScene()->Rename(*NewNameStr);
 
-				if(bNewAnimation)
+				if (bNewAnimation)
 				{
 					Blueprint->Modify();
 					Blueprint->Animations.Add(WidgetAnimation);
@@ -186,11 +204,10 @@ private:
 				}
 			}
 
-			FBlueprintEditorUtils::ReplaceVariableReferences(Blueprint, CurrentName, FName(*InText.ToString()));
-
+			FBlueprintEditorUtils::ReplaceVariableReferences(Blueprint, OldFName, NewFName);
 			FBlueprintEditorUtils::MarkBlueprintAsStructurallyModified(Blueprint);
 		}
-		else if( bNewAnimation )
+		else if (bNewAnimation)
 		{
 			const FScopedTransaction Transaction(LOCTEXT("NewAnimation", "New Animation"));
 			Blueprint->Modify();
@@ -398,9 +415,12 @@ private:
 			UniqueName = FString::Printf( TEXT( "%s_%i" ), *BaseName, NameIndex );
 			NameIndex++;
 		}
-		NewAnimation->Rename( *UniqueName );
 
-		NewAnimation->MovieScene = NewObject<UMovieScene>(NewAnimation, NewAnimation->GetFName(), RF_Transactional);
+		const FName NewFName = MakeObjectNameFromDisplayLabel(UniqueName, NewAnimation->GetFName());
+		NewAnimation->SetDisplayLabel( UniqueName );
+		NewAnimation->Rename(*UniqueName);
+
+		NewAnimation->MovieScene = NewObject<UMovieScene>(NewAnimation, NewFName, RF_Transactional);
 
 		// Default to 20 fps display rate (as was the previous default in USequencerSettings)
 		NewAnimation->MovieScene->SetDisplayRate(FFrameRate(20, 1));
diff --git a/Engine/Source/Editor/UMGEditor/Private/WidgetBlueprintCompiler.cpp b/Engine/Source/Editor/UMGEditor/Private/WidgetBlueprintCompiler.cpp
--- a/Engine/Source/Editor/UMGEditor/Private/WidgetBlueprintCompiler.cpp
+++ b/Engine/Source/Editor/UMGEditor/Private/WidgetBlueprintCompiler.cpp
@@ -287,19 +287,33 @@ void FWidgetBlueprintCompiler::CreateClassVariablesFromBlueprint()
 	}
 
 	// Add movie scenes variables here
-	for(UWidgetAnimation* Animation : WidgetBP->Animations)
+	for (UWidgetAnimation* Animation : WidgetBP->Animations)
 	{
+		UObjectPropertyBase* ExistingProperty = Cast<UObjectPropertyBase>(ParentClass->FindPropertyByName(Animation->GetFName()));
+		if (ExistingProperty &&
+			FWidgetBlueprintEditorUtils::IsBindWidgetAnimProperty(ExistingProperty) &&
+			ExistingProperty->PropertyClass->IsChildOf(UWidgetAnimation::StaticClass()))
+		{
+			WidgetAnimToMemberVariableMap.Add(Animation, ExistingProperty);
+			continue;
+		}
+
 		FEdGraphPinType WidgetPinType(UEdGraphSchema_K2::PC_Object, NAME_None, Animation->GetClass(), EPinContainerType::None, true, FEdGraphTerminalType());
 		UProperty* AnimationProperty = CreateVariable(Animation->GetFName(), WidgetPinType);
 
 		if ( AnimationProperty != nullptr )
 		{
+			const FString DisplayName = Animation->GetDisplayName().ToString();
+			AnimationProperty->SetMetaData(TEXT("DisplayName"), *DisplayName);
+
 			AnimationProperty->SetMetaData(TEXT("Category"), TEXT("Animations"));
 
 			AnimationProperty->SetPropertyFlags(CPF_Instanced);
 			AnimationProperty->SetPropertyFlags(CPF_BlueprintVisible);
 			AnimationProperty->SetPropertyFlags(CPF_BlueprintReadOnly);
 			AnimationProperty->SetPropertyFlags(CPF_RepSkip);
+
+			WidgetAnimToMemberVariableMap.Add(Animation, AnimationProperty);
 		}
 	}
 }
@@ -446,6 +460,7 @@ void FWidgetBlueprintCompiler::FinishCompilingClass(UClass* Class)
 {
 	UWidgetBlueprint* WidgetBP = WidgetBlueprint();
 	UWidgetBlueprintGeneratedClass* BPGClass = CastChecked<UWidgetBlueprintGeneratedClass>(Class);
+	UClass* ParentClass = WidgetBP->ParentClass;
 
 	// Don't do a bunch of extra work on the skeleton generated class
 	if ( WidgetBP->SkeletonGeneratedClass != Class )
@@ -514,7 +529,6 @@ void FWidgetBlueprintCompiler::FinishCompilingClass(UClass* Class)
 		}
 
 		// Check that all BindWidget properties are present and of the appropriate type
-		UClass* ParentClass = WidgetBP->ParentClass;
 		for (TUObjectPropertyBase<UWidget*>* WidgetProperty : TFieldRange<TUObjectPropertyBase<UWidget*>>(ParentClass))
 		{
 			bool bIsOptional = false;
@@ -554,6 +568,35 @@ void FWidgetBlueprintCompiler::FinishCompilingClass(UClass* Class)
 				}
 			}
 		}
+
+		// Check that all BindWidgetAnim properties are present
+		for (TUObjectPropertyBase<UWidgetAnimation*>* WidgetAnimProperty : TFieldRange<TUObjectPropertyBase<UWidgetAnimation*>>(ParentClass))
+		{
+			bool bIsOptional = false;
+
+			if (FWidgetBlueprintEditorUtils::IsBindWidgetAnimProperty(WidgetAnimProperty, bIsOptional))
+			{
+				const FText OptionalBindingAvailableNote = LOCTEXT("OptionalWidgetAnimNotBound", "An optional widget animation binding @@ is available.");
+				const FText RequiredWidgetAnimNotBoundError = LOCTEXT("RequiredWidgetAnimNotBound", "A required widget animation binding @@ was not found.");
+
+				UWidgetAnimation* const* WidgetAnim = WidgetAnimToMemberVariableMap.FindKey(WidgetAnimProperty);
+				if (!WidgetAnim)
+				{
+					if (bIsOptional)
+					{
+						MessageLog.Note(*OptionalBindingAvailableNote.ToString(), WidgetAnimProperty);
+					}
+					else if (Blueprint->bIsNewlyCreated)
+					{
+						MessageLog.Warning(*RequiredWidgetAnimNotBoundError.ToString(), WidgetAnimProperty);
+					}
+					else
+					{
+						MessageLog.Error(*RequiredWidgetAnimNotBoundError.ToString(), WidgetAnimProperty);
+					}
+				}
+			}
+		}
 	}
 
 	Super::FinishCompilingClass(Class);
@@ -564,6 +607,7 @@ void FWidgetBlueprintCompiler::PostCompile()
 	Super::PostCompile();
 
 	WidgetToMemberVariableMap.Empty();
+	WidgetAnimToMemberVariableMap.Empty();
 
 	UWidgetBlueprintGeneratedClass* WidgetClass = NewWidgetBlueprintClass;
 
diff --git a/Engine/Source/Editor/UMGEditor/Private/WidgetBlueprintCompiler.h b/Engine/Source/Editor/UMGEditor/Private/WidgetBlueprintCompiler.h
--- a/Engine/Source/Editor/UMGEditor/Private/WidgetBlueprintCompiler.h
+++ b/Engine/Source/Editor/UMGEditor/Private/WidgetBlueprintCompiler.h
@@ -54,6 +54,9 @@ protected:
 	// Map of properties created for widgets; to aid in debug data generation
 	TMap<class UWidget*, class UProperty*> WidgetToMemberVariableMap;
 
+	// Map of properties created for widget animations; to aid in debug data generation
+	TMap<class UWidgetAnimation*, class UProperty*> WidgetAnimToMemberVariableMap;
+
 	///----------------------------------------------------------------
 };
 
diff --git a/Engine/Source/Editor/UMGEditor/Private/WidgetBlueprintEditorUtils.cpp b/Engine/Source/Editor/UMGEditor/Private/WidgetBlueprintEditorUtils.cpp
--- a/Engine/Source/Editor/UMGEditor/Private/WidgetBlueprintEditorUtils.cpp
+++ b/Engine/Source/Editor/UMGEditor/Private/WidgetBlueprintEditorUtils.cpp
@@ -1312,6 +1312,25 @@ bool FWidgetBlueprintEditorUtils::IsBindWidgetProperty(UProperty* InProperty, bo
 	return false;
 }
 
+bool FWidgetBlueprintEditorUtils::IsBindWidgetAnimProperty(UProperty* InProperty)
+{
+	bool bIsOptional;
+	return IsBindWidgetAnimProperty(InProperty, bIsOptional);
+}
+
+bool FWidgetBlueprintEditorUtils::IsBindWidgetAnimProperty(UProperty* InProperty, bool& bIsOptional)
+{
+	if (InProperty)
+	{
+		bool bIsBindWidgetAnim = InProperty->HasMetaData("BindWidgetAnim") || InProperty->HasMetaData("BindWidgetAnimOptional");
+		bIsOptional = InProperty->HasMetaData("BindWidgetAnimOptional");
+
+		return bIsBindWidgetAnim;
+	}
+
+	return false;
+}
+
 bool FWidgetBlueprintEditorUtils::IsUsableWidgetClass(UClass* WidgetClass)
 {
 	if ( WidgetClass->IsChildOf(UWidget::StaticClass()) )
diff --git a/Engine/Source/Editor/UMGEditor/Private/WidgetBlueprintEditorUtils.h b/Engine/Source/Editor/UMGEditor/Private/WidgetBlueprintEditorUtils.h
--- a/Engine/Source/Editor/UMGEditor/Private/WidgetBlueprintEditorUtils.h
+++ b/Engine/Source/Editor/UMGEditor/Private/WidgetBlueprintEditorUtils.h
@@ -34,6 +34,9 @@ public:
 	static bool IsBindWidgetProperty(UProperty* InProperty);
 	static bool IsBindWidgetProperty(UProperty* InProperty, bool& bIsOptional);
 
+	static bool IsBindWidgetAnimProperty(UProperty* InProperty);
+	static bool IsBindWidgetAnimProperty(UProperty* InProperty, bool& bIsOptional);
+
 	static bool IsUsableWidgetClass(UClass* WidgetClass);
 
 public:
diff --git a/Engine/Source/Runtime/UMG/Private/Animation/WidgetAnimation.cpp b/Engine/Source/Runtime/UMG/Private/Animation/WidgetAnimation.cpp
--- a/Engine/Source/Runtime/UMG/Private/Animation/WidgetAnimation.cpp
+++ b/Engine/Source/Runtime/UMG/Private/Animation/WidgetAnimation.cpp
@@ -59,8 +59,18 @@ UWidgetAnimation* UWidgetAnimation::GetNullAnimation()
 	return NullAnimation;
 }
 
-#endif
+void UWidgetAnimation::SetDisplayLabel(const FString& InDisplayLabel)
+{
+	DisplayLabel = InDisplayLabel;
+}
 
+FText UWidgetAnimation::GetDisplayName() const
+{
+	const bool bHasDisplayLabel = !DisplayLabel.IsEmpty();
+	return bHasDisplayLabel ? FText::FromString(DisplayLabel) : Super::GetDisplayName();
+}
+
+#endif
 
 float UWidgetAnimation::GetStartTime() const
 {
diff --git a/Engine/Source/Runtime/UMG/Public/Animation/WidgetAnimation.h b/Engine/Source/Runtime/UMG/Public/Animation/WidgetAnimation.h
--- a/Engine/Source/Runtime/UMG/Public/Animation/WidgetAnimation.h
+++ b/Engine/Source/Runtime/UMG/Public/Animation/WidgetAnimation.h
@@ -30,6 +30,18 @@ public:
 	 * @return Placeholder animation.
 	 */
 	static UMG_API UWidgetAnimation* GetNullAnimation();
+
+	/** @return The friendly name of the animation */
+	UMG_API const FString& GetDisplayLabel() const
+	{
+		return DisplayLabel;
+	}
+
+	/** Sets the friendly name of the animation to display in the editor */
+	UMG_API void SetDisplayLabel(const FString& InDisplayLabel);
+
+	/** Returns the DisplayLabel if set, otherwise the object name */
+	UMG_API virtual FText GetDisplayName() const override;
 #endif
 
 	/**
@@ -99,4 +111,8 @@ private:
 	/** Whether to finish evaluation on stop. This legacy value is to preserve existing asset behavior to NOT finish on stop since content was created with this bug. If this is removed, evaluation should always finish on stop. */
 	UPROPERTY()
 	bool bLegacyFinishOnStop;
+
+	/** The friendly name for this animation displayed in the designer. */
+	UPROPERTY()
+	FString DisplayLabel;
 };
diff --git a/Engine/Source/Runtime/UMG/Public/Components/Widget.h b/Engine/Source/Runtime/UMG/Public/Components/Widget.h
--- a/Engine/Source/Runtime/UMG/Public/Components/Widget.h
+++ b/Engine/Source/Runtime/UMG/Public/Components/Widget.h
@@ -67,6 +67,14 @@ namespace UMWidget
 		// UPROPERTY(meta=(BindWidget, OptionalWidget=true))
 		OptionalWidget,
 
+		// [PropertyMetadata] This property requires a widget animation be bound to it in the designer.  Allows easy native access to designer defined animations.
+		// UPROPERTY(meta=(BindWidgetAnim))
+		BindWidgetAnim,
+
+		// [PropertyMetadata] This property optionally allows a animation widget be bound to it in the designer.  Allows easy native access to designer defined animation.
+		// UPROPERTY(meta=(BindWidgetAnimOptional))
+		BindWidgetAnimOptional,
+
 		// [PropertyMetadata] Exposes a dynamic delegate property in the details panel for the widget.
 		// UPROPERTY(meta=(IsBindableEvent))
 		IsBindableEvent,
