Fix for possible nullptr dereference.

diff --git a/Engine/Source/Editor/UMGEditor/Private/TabFactory/AnimationTabSummoner.cpp b/Engine/Source/Editor/UMGEditor/Private/TabFactory/AnimationTabSummoner.cpp
--- a/Engine/Source/Editor/UMGEditor/Private/TabFactory/AnimationTabSummoner.cpp
+++ b/Engine/Source/Editor/UMGEditor/Private/TabFactory/AnimationTabSummoner.cpp
@@ -71,10 +71,13 @@ bool VerifyAnimationRename( FWidgetBlueprintEditor& BlueprintEditor, UWidgetAnim
 	}
 
 	// Check for BindWidgetAnim property
-	UProperty* Property = Blueprint->ParentClass->FindPropertyByName(NewAnimationNameAsName);
-	if (Property && FWidgetBlueprintEditorUtils::IsBindWidgetAnimProperty(Property))
+	if (Blueprint)
 	{
-		return true;
+		UProperty* Property = Blueprint->ParentClass->FindPropertyByName(NewAnimationNameAsName);
+		if (Property && FWidgetBlueprintEditorUtils::IsBindWidgetAnimProperty(Property))
+		{
+			return true;
+		}
 	}
 
 	FKismetNameValidator Validator( Blueprint );
