Clearing focus on a variable once it is committed. Fixes assert on undo

diff --git a/Engine/Source/Editor/UMGEditor/Private/Details/SWidgetDetailsView.cpp b/Engine/Source/Editor/UMGEditor/Private/Details/SWidgetDetailsView.cpp
--- a/Engine/Source/Editor/UMGEditor/Private/Details/SWidgetDetailsView.cpp
+++ b/Engine/Source/Editor/UMGEditor/Private/Details/SWidgetDetailsView.cpp
@@ -525,6 +525,7 @@ void SWidgetDetailsView::NotifyPostChange(const FPropertyChangedEvent& PropertyC
 		// Any time we migrate a property value we need to mark the blueprint as structurally modified so users don't need 
 		// to recompile it manually before they see it play in game using the latest version.
 		FBlueprintEditorUtils::MarkBlueprintAsModified(BlueprintEditor.Pin()->GetBlueprintObj());
+		ClearFocusIfOwned();
 	}
 
 	// If the property that changed is marked as "DesignerRebuild" we invalidate
