Blueprints inheriting from UWidget will now show up in the palette view even if they are not loaded.

diff --git a/Engine/Source/Editor/UMGEditor/Private/Palette/SPaletteView.cpp b/Engine/Source/Editor/UMGEditor/Private/Palette/SPaletteView.cpp
--- a/Engine/Source/Editor/UMGEditor/Private/Palette/SPaletteView.cpp
+++ b/Engine/Source/Editor/UMGEditor/Private/Palette/SPaletteView.cpp
@@ -487,6 +487,41 @@ void SPaletteView::BuildClassWidgetList()
 
 	// Locate all widget BP assets (include unloaded)
 	const FAssetRegistryModule& AssetRegistryModule = FModuleManager::LoadModuleChecked<FAssetRegistryModule>(TEXT("AssetRegistry"));
+	TArray<FAssetData> AllBPsAssetData;
+	AssetRegistryModule.Get().GetAssetsByClass(UBlueprint::StaticClass()->GetFName(), AllBPsAssetData, true);
+
+	for (FAssetData& BPAssetData : AllBPsAssetData)
+	{
+		// Blueprints get the class type actions for their parent native class - this avoids us having to load the blueprint
+		UClass* ParentClass = nullptr;
+		FString ParentClassName;
+		if (!BPAssetData.GetTagValue(FBlueprintTags::NativeParentClassPath, ParentClassName))
+		{
+			BPAssetData.GetTagValue(FBlueprintTags::ParentClassPath, ParentClassName);
+		}
+		if (!ParentClassName.IsEmpty())
+		{
+			UObject* Outer = nullptr;
+			ResolveName(Outer, ParentClassName, false, false);
+			ParentClass = FindObject<UClass>(ANY_PACKAGE, *ParentClassName);
+			// UUserWidgets have their own loading section, and we don't want to process any blueprints that don't have UWidget parents
+			if (!ParentClass->IsChildOf(UWidget::StaticClass()) || ParentClass->IsChildOf(UUserWidget::StaticClass()))
+			{
+				continue;
+			}
+		}
+
+		if (!FilterAssetData(BPAssetData))
+		{
+			// If this object isn't currently loaded, add it to the palette view
+			if (BPAssetData.ToSoftObjectPath().ResolveObject() == nullptr)
+			{
+				auto Template = MakeShareable(new FWidgetTemplateClass(BPAssetData, nullptr));
+				AddWidgetTemplate(Template);
+			}
+		}
+	}
+
 	TArray<FAssetData> AllWidgetBPsAssetData;
 	AssetRegistryModule.Get().GetAssetsByClass(UWidgetBlueprint::StaticClass()->GetFName(), AllWidgetBPsAssetData, true);
 
@@ -499,51 +534,58 @@ void SPaletteView::BuildClassWidgetList()
 			continue;
 		}
 
-		// Excludes engine content if user sets it to false
-		if (!GetDefault<UContentBrowserSettings>()->GetDisplayEngineFolder() || !GetDefault<UUMGEditorProjectSettings>()->bShowWidgetsFromEngineContent)
+		if (!FilterAssetData(WidgetBPAssetData))
 		{
-			if (WidgetBPAssetData.PackagePath.ToString().Find(TEXT("/Engine")) == 0)
+			// Excludes this widget if it is on the hide list
+			bool bIsOnList = false;
+			for (FSoftClassPath Widget : WidgetClassesToHide)
 			{
-				continue;
+				if (Widget.ToString().Find(WidgetBPAssetData.ObjectPath.ToString()) == 0)
+				{
+					bIsOnList = true;
+					break;
+				}
 			}
-		}
-
-		// Excludes developer content if user sets it to false
-		if (!GetDefault<UContentBrowserSettings>()->GetDisplayDevelopersFolder() || !GetDefault<UUMGEditorProjectSettings>()->bShowWidgetsFromDeveloperContent)
-		{
-			if (WidgetBPAssetData.PackagePath.ToString().Find(TEXT("/Game/Developers")) == 0)
+			if (bIsOnList)
 			{
 				continue;
 			}
-		}
 
-		// Excludes this widget if it is on the hide list
-		bool bIsOnList = false;
-		for (FSoftClassPath Widget : WidgetClassesToHide)
-		{
-			if (Widget.ToString().Find(WidgetBPAssetData.ObjectPath.ToString()) == 0)
+			// If the blueprint generated class was found earlier, pass it to the template
+			TSubclassOf<UUserWidget> WidgetBPClass = nullptr;
+			auto LoadedWidgetBPClass = LoadedWidgetBlueprintClassesByName.Find(WidgetBPAssetData.AssetName);
+			if (LoadedWidgetBPClass)
 			{
-				bIsOnList = true;
-				break;
+				WidgetBPClass = *LoadedWidgetBPClass;
 			}
+
+			auto Template = MakeShareable(new FWidgetTemplateBlueprintClass(WidgetBPAssetData, WidgetBPClass));
+
+			AddWidgetTemplate(Template);
 		}
-		if (bIsOnList)
+	}
+}
+
+bool SPaletteView::FilterAssetData(FAssetData &InAssetData)
+{
+	// Excludes engine content if user sets it to false
+	if (!GetDefault<UContentBrowserSettings>()->GetDisplayEngineFolder() || !GetDefault<UUMGEditorProjectSettings>()->bShowWidgetsFromEngineContent)
+	{
+		if (InAssetData.PackagePath.ToString().Find(TEXT("/Engine")) == 0)
 		{
-			continue;
+			return true;
 		}
+	}
 
-		// If the blueprint generated class was found earlier, pass it to the template
-		TSubclassOf<UUserWidget> WidgetBPClass = nullptr;
-		auto LoadedWidgetBPClass = LoadedWidgetBlueprintClassesByName.Find(WidgetBPAssetData.AssetName);
-		if (LoadedWidgetBPClass)
+	// Excludes developer content if user sets it to false
+	if (!GetDefault<UContentBrowserSettings>()->GetDisplayDevelopersFolder() || !GetDefault<UUMGEditorProjectSettings>()->bShowWidgetsFromDeveloperContent)
+	{
+		if (InAssetData.PackagePath.ToString().Find(TEXT("/Game/Developers")) == 0)
 		{
-			WidgetBPClass = *LoadedWidgetBPClass;
+			return true;
 		}
-
-		auto Template = MakeShareable(new FWidgetTemplateBlueprintClass(WidgetBPAssetData, WidgetBPClass));
-
-		AddWidgetTemplate(Template);
 	}
+	return false;
 }
 
 void SPaletteView::BuildSpecialWidgetList()
diff --git a/Engine/Source/Editor/UMGEditor/Private/Palette/SPaletteView.h b/Engine/Source/Editor/UMGEditor/Private/Palette/SPaletteView.h
--- a/Engine/Source/Editor/UMGEditor/Private/Palette/SPaletteView.h
+++ b/Engine/Source/Editor/UMGEditor/Private/Palette/SPaletteView.h
@@ -64,6 +64,7 @@ private:
 
 	void BuildWidgetList();
 	void BuildClassWidgetList();
+	bool FilterAssetData(FAssetData &BPAssetData);
 	void BuildSpecialWidgetList();
 
 	void OnGetChildren(TSharedPtr<FWidgetViewModel> Item, TArray< TSharedPtr<FWidgetViewModel> >& Children);
diff --git a/Engine/Source/Editor/UMGEditor/Private/Templates/WidgetTemplateBlueprintClass.cpp b/Engine/Source/Editor/UMGEditor/Private/Templates/WidgetTemplateBlueprintClass.cpp
--- a/Engine/Source/Editor/UMGEditor/Private/Templates/WidgetTemplateBlueprintClass.cpp
+++ b/Engine/Source/Editor/UMGEditor/Private/Templates/WidgetTemplateBlueprintClass.cpp
@@ -13,17 +13,8 @@
 #define LOCTEXT_NAMESPACE "UMGEditor"
 
 FWidgetTemplateBlueprintClass::FWidgetTemplateBlueprintClass(const FAssetData& InWidgetAssetData, TSubclassOf<UUserWidget> InUserWidgetClass)
-	: FWidgetTemplateClass(), WidgetAssetData(InWidgetAssetData)
+	: FWidgetTemplateClass(InWidgetAssetData, InUserWidgetClass)
 {
-	if (InUserWidgetClass)
-	{
-		WidgetClass = *InUserWidgetClass;
-		Name = WidgetClass->GetDisplayNameText();
-	}
-	else
-	{
-		Name = FText::FromString(FName::NameToDisplayString(WidgetAssetData.AssetName.ToString(), false));
-	}
 }
 
 FWidgetTemplateBlueprintClass::~FWidgetTemplateBlueprintClass()
diff --git a/Engine/Source/Editor/UMGEditor/Private/Templates/WidgetTemplateBlueprintClass.h b/Engine/Source/Editor/UMGEditor/Private/Templates/WidgetTemplateBlueprintClass.h
--- a/Engine/Source/Editor/UMGEditor/Private/Templates/WidgetTemplateBlueprintClass.h
+++ b/Engine/Source/Editor/UMGEditor/Private/Templates/WidgetTemplateBlueprintClass.h
@@ -50,7 +50,4 @@ public:
 	/** Returns true if the supplied class is supported by this template */
 	static bool Supports(UClass* InClass);
 
-protected:
-	/** The asset data for the widget blueprint */
-	FAssetData WidgetAssetData;
 };
diff --git a/Engine/Source/Editor/UMGEditor/Private/Templates/WidgetTemplateClass.cpp b/Engine/Source/Editor/UMGEditor/Private/Templates/WidgetTemplateClass.cpp
--- a/Engine/Source/Editor/UMGEditor/Private/Templates/WidgetTemplateClass.cpp
+++ b/Engine/Source/Editor/UMGEditor/Private/Templates/WidgetTemplateClass.cpp
@@ -29,6 +29,20 @@ FWidgetTemplateClass::FWidgetTemplateClass(TSubclassOf<UWidget> InWidgetClass)
 	GEditor->OnObjectsReplaced().AddRaw(this, &FWidgetTemplateClass::OnObjectsReplaced);
 }
 
+FWidgetTemplateClass::FWidgetTemplateClass(const FAssetData& InWidgetAssetData, TSubclassOf<UWidget> InWidgetClass)
+	: WidgetAssetData(InWidgetAssetData)
+{
+	if (InWidgetClass)
+	{
+		WidgetClass = *InWidgetClass;
+		Name = WidgetClass->GetDisplayNameText();
+	}
+	else
+	{
+		Name = FText::FromString(FName::NameToDisplayString(WidgetAssetData.AssetName.ToString(), false));
+	}
+}
+
 FWidgetTemplateClass::~FWidgetTemplateClass()
 {
 	GEditor->OnObjectsReplaced().RemoveAll(this);
@@ -36,12 +50,28 @@ FWidgetTemplateClass::~FWidgetTemplateClass()
 
 FText FWidgetTemplateClass::GetCategory() const
 {
-	auto DefaultWidget = WidgetClass->GetDefaultObject<UWidget>();
-	return DefaultWidget->GetPaletteCategory();
+	if (WidgetClass.Get())
+	{
+		auto DefaultWidget = WidgetClass->GetDefaultObject<UWidget>();
+		return DefaultWidget->GetPaletteCategory();
+	}
+	else
+	{
+		auto DefaultWidget = UWidget::StaticClass()->GetDefaultObject<UWidget>();
+		return DefaultWidget->GetPaletteCategory();
+	}
 }
 
 UWidget* FWidgetTemplateClass::Create(UWidgetTree* Tree)
 {
+	// Load the blueprint asset if needed
+	if (!WidgetClass.Get())
+	{
+		FString AssetPath = WidgetAssetData.ObjectPath.ToString();
+		UBlueprint* LoadedBP = LoadObject<UBlueprint>(nullptr, *AssetPath);
+		WidgetClass = *LoadedBP->GeneratedClass;
+	}
+
 	return CreateNamed(Tree, NAME_None);
 }
 
@@ -51,12 +81,36 @@ const FSlateBrush* FWidgetTemplateClass::GetIcon() const
 	{
 		return FSlateIconFinder::FindIconBrushForClass(WidgetClass.Get());
 	}
+	else
+	{
+		return FSlateIconFinder::FindIconBrushForClass(UWidget::StaticClass());
+	}
 	return nullptr;
 }
 
 TSharedRef<IToolTip> FWidgetTemplateClass::GetToolTip() const
 {
-	return IDocumentation::Get()->CreateToolTip(WidgetClass->GetToolTipText(), nullptr, FString(TEXT("Shared/Types/")) + WidgetClass->GetName(), TEXT("Class"));
+	if (WidgetClass.IsValid())
+	{
+		return IDocumentation::Get()->CreateToolTip(WidgetClass->GetToolTipText(), nullptr, FString(TEXT("Shared/Types/")) + WidgetClass->GetName(), TEXT("Class"));
+	}
+	else
+	{
+		FText Description;
+
+		FString DescriptionStr = WidgetAssetData.GetTagValueRef<FString>(GET_MEMBER_NAME_CHECKED(UBlueprint, BlueprintDescription));
+		if (!DescriptionStr.IsEmpty())
+		{
+			DescriptionStr.ReplaceInline(TEXT("\\n"), TEXT("\n"));
+			Description = FText::FromString(MoveTemp(DescriptionStr));
+		}
+		else
+		{
+			Description = Name;
+		}
+
+		return IDocumentation::Get()->CreateToolTip(Description, nullptr, FString(TEXT("Shared/Types/")) + Name.ToString(), TEXT("Class"));
+	}
 }
 
 void FWidgetTemplateClass::OnObjectsReplaced(const TMap<UObject*, UObject*>& ReplacementMap)
diff --git a/Engine/Source/Editor/UMGEditor/Private/Templates/WidgetTemplateClass.h b/Engine/Source/Editor/UMGEditor/Private/Templates/WidgetTemplateClass.h
--- a/Engine/Source/Editor/UMGEditor/Private/Templates/WidgetTemplateClass.h
+++ b/Engine/Source/Editor/UMGEditor/Private/Templates/WidgetTemplateClass.h
@@ -7,6 +7,7 @@
 #include "Components/Widget.h"
 #include "Widgets/IToolTip.h"
 #include "WidgetTemplate.h"
+#include "AssetData.h"
 
 class UWidgetTree;
 
@@ -19,6 +20,8 @@ public:
 	/** Constructor */
 	explicit FWidgetTemplateClass(TSubclassOf<UWidget> InWidgetClass);
 
+	explicit FWidgetTemplateClass(const FAssetData& InWidgetAssetData, TSubclassOf<UWidget> InWidgetClass);
+
 	/** Destructor */
 	virtual ~FWidgetTemplateClass();
 
@@ -50,4 +53,7 @@ protected:
 protected:
 	/** The widget class that will be created by this template */
 	TWeakObjectPtr<UClass> WidgetClass;
+
+	/** The asset data for the widget blueprint */
+	FAssetData WidgetAssetData;
 };
