PR #4936: Duplicate widget functionality for UMG editor (Contributed by [REDACTED_PERSON])

Fixed up some variable names.

diff --git a/Engine/Source/Editor/UMGEditor/Private/WidgetBlueprintEditor.cpp b/Engine/Source/Editor/UMGEditor/Private/WidgetBlueprintEditor.cpp
--- a/Engine/Source/Editor/UMGEditor/Private/WidgetBlueprintEditor.cpp
+++ b/Engine/Source/Editor/UMGEditor/Private/WidgetBlueprintEditor.cpp
@@ -133,6 +133,11 @@ void FWidgetBlueprintEditor::InitWidgetBlueprintEditor(const EToolkitMode::Type
 		FExecuteAction::CreateSP(this, &FWidgetBlueprintEditor::PasteWidgets),
 		FCanExecuteAction::CreateSP(this, &FWidgetBlueprintEditor::CanPasteWidgets)
 		);
+
+	DesignerCommandList->MapAction(FGenericCommands::Get().Duplicate,
+		FExecuteAction::CreateSP(this, &FWidgetBlueprintEditor::DuplicateSelectedWidgets),
+		FCanExecuteAction::CreateSP(this, &FWidgetBlueprintEditor::CanDuplicateSelectedWidgets)
+		);
 }
 
 void FWidgetBlueprintEditor::RegisterApplicationModes(const TArray<UBlueprint*>& InBlueprints, bool bShouldOpenInDefaultsMode, bool bNewlyCreated/* = false*/)
@@ -426,6 +431,24 @@ void FWidgetBlueprintEditor::PasteWidgets()
 	SelectWidgets(PastedWidgetRefs, false);
 }
 
+bool FWidgetBlueprintEditor::CanDuplicateSelectedWidgets()
+{
+	TSet<FWidgetReference> Widgets = GetSelectedWidgets();
+	if (Widgets.Num() == 1)
+	{
+		FWidgetReference Target = *Widgets.CreateIterator();
+		UPanelWidget* ParentWidget = Target.GetTemplate()->GetParent();
+		return ParentWidget && ParentWidget->CanAddMoreChildren();
+	}
+	return false;
+}
+
+void FWidgetBlueprintEditor::DuplicateSelectedWidgets()
+{
+	TSet<FWidgetReference> Widgets = GetSelectedWidgets();
+	FWidgetBlueprintEditorUtils::DuplicateWidgets(SharedThis(this), GetWidgetBlueprintObj(), Widgets);
+}
+
 void FWidgetBlueprintEditor::Tick(float DeltaTime)
 {
 	FBlueprintEditor::Tick(DeltaTime);
diff --git a/Engine/Source/Editor/UMGEditor/Private/WidgetBlueprintEditor.h b/Engine/Source/Editor/UMGEditor/Private/WidgetBlueprintEditor.h
--- a/Engine/Source/Editor/UMGEditor/Private/WidgetBlueprintEditor.h
+++ b/Engine/Source/Editor/UMGEditor/Private/WidgetBlueprintEditor.h
@@ -218,6 +218,9 @@ private:
 	bool CanCutSelectedWidgets();
 	void CutSelectedWidgets();
 
+	bool CanDuplicateSelectedWidgets();
+	void DuplicateSelectedWidgets();
+
 private:
 	/** Called whenever the blueprint is structurally changed. */
 	virtual void OnBlueprintChangedImpl(UBlueprint* InBlueprint, bool bIsJustBeingCompiled = false ) override;
diff --git a/Engine/Source/Editor/UMGEditor/Private/WidgetBlueprintEditorUtils.cpp b/Engine/Source/Editor/UMGEditor/Private/WidgetBlueprintEditorUtils.cpp
--- a/Engine/Source/Editor/UMGEditor/Private/WidgetBlueprintEditorUtils.cpp
+++ b/Engine/Source/Editor/UMGEditor/Private/WidgetBlueprintEditorUtils.cpp
@@ -284,6 +284,7 @@ void FWidgetBlueprintEditorUtils::CreateWidgetContextMenu(FMenuBuilder& MenuBuil
 		MenuBuilder.AddMenuEntry(FGenericCommands::Get().Copy);
 		MenuBuilder.AddMenuEntry(FGenericCommands::Get().Paste);
 		MenuBuilder.AddMenuEntry(FGenericCommands::Get().Rename);
+		MenuBuilder.AddMenuEntry(FGenericCommands::Get().Duplicate);
 		MenuBuilder.AddMenuEntry(FGenericCommands::Get().Delete);
 	}
 	MenuBuilder.EndSection();
@@ -905,6 +906,12 @@ void FWidgetBlueprintEditorUtils::CutWidgets(UWidgetBlueprint* BP, TSet<FWidgetR
 }
 
 void FWidgetBlueprintEditorUtils::CopyWidgets(UWidgetBlueprint* BP, TSet<FWidgetReference> Widgets)
+{
+	FString ExportedText = CopyWidgetsInternal(BP, Widgets);
+	FPlatformApplicationMisc::ClipboardCopy(*ExportedText);
+}
+
+FString FWidgetBlueprintEditorUtils::CopyWidgetsInternal(UWidgetBlueprint* BP, TSet<FWidgetReference> Widgets)
 {
 	TSet<UWidget*> TemplateWidgets;
 
@@ -946,7 +953,36 @@ void FWidgetBlueprintEditorUtils::CopyWidgets(UWidgetBlueprint* BP, TSet<FWidget
 
 	FString ExportedText;
 	FWidgetBlueprintEditorUtils::ExportWidgetsToText(FinalWidgets, /*out*/ ExportedText);
-	FPlatformApplicationMisc::ClipboardCopy(*ExportedText);
+	return ExportedText;
+}
+
+void FWidgetBlueprintEditorUtils::DuplicateWidgets(TSharedRef<FWidgetBlueprintEditor> BlueprintEditor, UWidgetBlueprint* BP, TSet<FWidgetReference> Widgets)
+{
+	FWidgetReference ParentWidgetRef = (Widgets.Num() == 1) ? *Widgets.CreateIterator() : FWidgetReference();
+	if (ParentWidgetRef.IsValid())
+	{
+		if (UPanelWidget* TargetWidget = Cast<UPanelWidget>(ParentWidgetRef.GetPreview()->GetParent()))
+		{
+			ParentWidgetRef = BlueprintEditor->GetReferenceFromPreview(TargetWidget);
+		}
+		else
+		{
+			ParentWidgetRef = FWidgetReference();
+		}
+	}
+
+	if (ParentWidgetRef.IsValid())
+	{
+		FString ExportedText = CopyWidgetsInternal(BP, Widgets);
+
+		FScopedTransaction Transaction(FGenericCommands::Get().Duplicate->GetDescription());
+		bool TransactionSuccesful = true;
+		PasteWidgetsInternal(BlueprintEditor, BP, ExportedText, ParentWidgetRef, NAME_None, FVector2D::ZeroVector, TransactionSuccesful);
+		if (!TransactionSuccesful)
+		{
+			Transaction.Cancel();
+		}
+	}
 }
 
 void FWidgetBlueprintEditorUtils::ExportWidgetsToText(TArray<UWidget*> WidgetsToExport, /*out*/ FString& ExportedText)
@@ -1003,6 +1039,17 @@ TArray<UWidget*> FWidgetBlueprintEditorUtils::PasteWidgets(TSharedRef<FWidgetBlu
 	FString TextToImport;
 	FPlatformApplicationMisc::ClipboardPaste(TextToImport);
 
+	bool bTransactionSuccessful = true;
+	TArray<UWidget*> PastedWidgets = PasteWidgetsInternal(BlueprintEditor, BP, TextToImport, ParentWidgetRef, SlotName, PasteLocation, bTransactionSuccessful);
+	if (!bTransactionSuccessful)
+	{
+		Transaction.Cancel();
+	}
+	return PastedWidgets;
+}
+
+TArray<UWidget*> FWidgetBlueprintEditorUtils::PasteWidgetsInternal(TSharedRef<FWidgetBlueprintEditor> BlueprintEditor, UWidgetBlueprint* BP, const FString& TextToImport, FWidgetReference ParentWidgetRef, FName SlotName, FVector2D PasteLocation, bool& bTransactionSuccessful)
+{
 	// Import the nodes
 	TSet<UWidget*> PastedWidgets;
 	TMap<FName, UWidgetSlotPair*> PastedExtraSlotData;
@@ -1011,7 +1058,7 @@ TArray<UWidget*> FWidgetBlueprintEditorUtils::PasteWidgets(TSharedRef<FWidgetBlu
 	// Ignore an empty set of widget paste data.
 	if ( PastedWidgets.Num() == 0 )
 	{
-		Transaction.Cancel();
+		bTransactionSuccessful = false;
 		return TArray<UWidget*>();
 	}
 
@@ -1047,7 +1094,7 @@ TArray<UWidget*> FWidgetBlueprintEditorUtils::PasteWidgets(TSharedRef<FWidgetBlu
 			// If we already have a root widget, then we can't replace the root.
 			if ( BP->WidgetTree->RootWidget )
 			{
-				Transaction.Cancel();
+				bTransactionSuccessful = false;
 				return TArray<UWidget*>();
 			}
 		}
@@ -1070,7 +1117,7 @@ TArray<UWidget*> FWidgetBlueprintEditorUtils::PasteWidgets(TSharedRef<FWidgetBlu
 					FNotificationInfo Info(LOCTEXT("NotEnoughSlots", "Can't paste contents, not enough available slots in target widget."));
 					FSlateNotificationManager::Get().AddNotification(Info);
 
-					Transaction.Cancel();
+					bTransactionSuccessful = false;
 					return TArray<UWidget*>();
 				}
 			}
@@ -1086,11 +1133,11 @@ TArray<UWidget*> FWidgetBlueprintEditorUtils::PasteWidgets(TSharedRef<FWidgetBlu
 				if (!SlotProperties.Contains(LayoutDataLabel))
 				{
 					bShouldReproduceOffsets = false;
+					break;
 				}
 			}
 
 			FVector2D FirstWidgetPosition;
-			UPanelWidget* FirstWidgetParent = RootPasteWidgets[0]->GetParent();
 			ParentWidget->Modify();
 			for ( UWidget* NewWidget : RootPasteWidgets )
 			{
@@ -1164,7 +1211,7 @@ TArray<UWidget*> FWidgetBlueprintEditorUtils::PasteWidgets(TSharedRef<FWidgetBlu
 			FNotificationInfo Info(LOCTEXT("NamedSlotsOnlyHoldOneWidget", "Can't paste content, a slot can only hold one widget at the root."));
 			FSlateNotificationManager::Get().AddNotification(Info);
 
-			Transaction.Cancel();
+			bTransactionSuccessful = false;
 			return TArray<UWidget*>();
 		}
 
@@ -1182,7 +1229,6 @@ TArray<UWidget*> FWidgetBlueprintEditorUtils::PasteWidgets(TSharedRef<FWidgetBlu
 	}
 
 	return RootPasteWidgets;
-
 }
 
 void FWidgetBlueprintEditorUtils::ImportWidgetsFromText(UWidgetBlueprint* BP, const FString& TextToImport, /*out*/ TSet<UWidget*>& ImportedWidgetSet, /*out*/ TMap<FName, UWidgetSlotPair*>& PastedExtraSlotData)
diff --git a/Engine/Source/Editor/UMGEditor/Private/WidgetBlueprintEditorUtils.h b/Engine/Source/Editor/UMGEditor/Private/WidgetBlueprintEditorUtils.h
--- a/Engine/Source/Editor/UMGEditor/Private/WidgetBlueprintEditorUtils.h
+++ b/Engine/Source/Editor/UMGEditor/Private/WidgetBlueprintEditorUtils.h
@@ -31,6 +31,8 @@ public:
 
 	static void CutWidgets(UWidgetBlueprint* BP, TSet<FWidgetReference> Widgets);
 
+	static void DuplicateWidgets(TSharedRef<FWidgetBlueprintEditor> BlueprintEditor, UWidgetBlueprint* BP, TSet<FWidgetReference> Widgets);
+
 	static bool IsBindWidgetProperty(UProperty* InProperty);
 	static bool IsBindWidgetProperty(UProperty* InProperty, bool& bIsOptional);
 
@@ -39,7 +41,6 @@ public:
 
 	static bool IsUsableWidgetClass(UClass* WidgetClass);
 
-public:
 	static void ExportWidgetsToText(TArray<UWidget*> WidgetsToExport, /*out*/ FString& ExportedText);
 
 	static void ImportWidgetsFromText(UWidgetBlueprint* BP, const FString& TextToImport, /*out*/ TSet<UWidget*>& ImportedWidgetSet, /*out*/ TMap<FName, UWidgetSlotPair*>& PastedExtraSlotData);
@@ -60,6 +61,10 @@ public:
 
 private:
 
+	static FString CopyWidgetsInternal(UWidgetBlueprint* BP, TSet<FWidgetReference> Widgets);
+
+	static TArray<UWidget*> PasteWidgetsInternal(TSharedRef<FWidgetBlueprintEditor> BlueprintEditor, UWidgetBlueprint* BP, const FString& TextToImport, FWidgetReference ParentWidget, FName SlotName, FVector2D PasteLocation, bool& TransactionSuccesful);
+
 	static void ExecuteOpenSelectedWidgetsForEdit( TSet<FWidgetReference> SelectedWidgets );
 
 	static bool FindAndRemoveNamedSlotContent(UWidget* WidgetTemplate, UWidgetTree* WidgetTree);
