Fix for widget blueprint variable duplication when dragging and dropping to the hierarchy

diff --git a/Engine/Source/Editor/UMGEditor/Private/Designer/SDesignerView.cpp b/Engine/Source/Editor/UMGEditor/Private/Designer/SDesignerView.cpp
--- a/Engine/Source/Editor/UMGEditor/Private/Designer/SDesignerView.cpp
+++ b/Engine/Source/Editor/UMGEditor/Private/Designer/SDesignerView.cpp
@@ -1580,6 +1580,7 @@ void SDesignerView::ClearDropPreviews()
 		// it will remain outered to the widget tree and end up as a property in the BP class layout as a result.
 		if (DropPreview.Widget->GetOutermost() != GetTransientPackage())
 		{
+			DropPreview.Widget->SetFlags(RF_NoFlags);
 			DropPreview.Widget->Rename(nullptr, GetTransientPackage());
 		}
 	}
@@ -2812,6 +2813,7 @@ void SDesignerView::ProcessDropAndAddWidget(const FGeometry& MyGeometry, const F
 		{
 			if (Widget->GetOutermost() != GetTransientPackage())
 			{
+				Widget->SetFlags(RF_NoFlags);
 				Widget->Rename(nullptr, GetTransientPackage());
 			}
 		}
diff --git a/Engine/Source/Editor/UMGEditor/Private/Hierarchy/SHierarchyViewItem.cpp b/Engine/Source/Editor/UMGEditor/Private/Hierarchy/SHierarchyViewItem.cpp
--- a/Engine/Source/Editor/UMGEditor/Private/Hierarchy/SHierarchyViewItem.cpp
+++ b/Engine/Source/Editor/UMGEditor/Private/Hierarchy/SHierarchyViewItem.cpp
@@ -33,6 +33,9 @@
 #include "WidgetBlueprintEditorUtils.h"
 #include "ScopedTransaction.h"
 #include "Styling/SlateIconFinder.h"
+#include "DragAndDrop/AssetDragDropOp.h"
+#include "WidgetTemplateBlueprintClass.h"
+#include "WidgetTemplateImageClass.h"
 
 #define LOCTEXT_NAMESPACE "UMG"
 
@@ -459,7 +462,17 @@ FReply FHierarchyModel::HandleDragDetected(const FGeometry& MyGeometry, const FP
 
 void FHierarchyModel::HandleDragEnter(const FDragDropEvent& DragDropEvent)
 {
-
+	TArray<UWidget*> DragDropPreviewWidgets;
+	DetermineDragDropPreviewWidgets(DragDropPreviewWidgets, DragDropEvent);
+	// Move the remaining widgets into the transient package. Otherwise, they will remain outered to the WidgetTree and end up as properties in the BP class layout as a result.
+	UWidgetBlueprint* BP = Cast<UWidgetBlueprint>(BlueprintEditor.Pin()->GetBlueprintObj());
+	if (BP)
+	{
+		for (UWidget* Widget : DragDropPreviewWidgets)
+		{
+			FHierarchyModel::RemovePreviewWidget(BP, Widget);
+		}
+	}
 }
 
 void FHierarchyModel::HandleDragLeave(const FDragDropEvent& DragDropEvent)
@@ -469,6 +482,17 @@ void FHierarchyModel::HandleDragLeave(const FDragDropEvent& DragDropEvent)
 	{
 		DecoratedDragDropOp->ResetToDefaultToolTip();
 	}
+	TArray<UWidget*> DragDropPreviewWidgets;
+	DetermineDragDropPreviewWidgets(DragDropPreviewWidgets, DragDropEvent);
+	// Move the remaining widgets into the transient package. Otherwise, they will remain outered to the WidgetTree and end up as properties in the BP class layout as a result.
+	UWidgetBlueprint* BP = Cast<UWidgetBlueprint>(BlueprintEditor.Pin()->GetBlueprintObj());
+	if (BP)
+	{
+		for (UWidget* Widget : DragDropPreviewWidgets)
+		{
+			FHierarchyModel::RemovePreviewWidget(BP, Widget);
+		}
+	}
 }
 
 FReply FHierarchyModel::HandleAcceptDrop(FDragDropEvent const& DragDropEvent, EItemDropZone DropZone)
@@ -486,6 +510,86 @@ void FHierarchyModel::OnNameTextCommited(const FText& InText, ETextCommit::Type
 
 }
 
+
+void FHierarchyModel::DetermineDragDropPreviewWidgets(TArray<UWidget*>& OutWidgets, const FDragDropEvent& DragDropEvent)
+{
+	OutWidgets.Empty();
+	UWidgetBlueprint* Blueprint = Cast<UWidgetBlueprint>(BlueprintEditor.Pin()->GetBlueprintObj());
+
+	if (!Blueprint)
+	{
+		return;
+	}
+
+	TSharedPtr<FWidgetTemplateDragDropOp> TemplateDragDropOp = DragDropEvent.GetOperationAs<FWidgetTemplateDragDropOp>();
+	TSharedPtr<FAssetDragDropOp> AssetDragDropOp = DragDropEvent.GetOperationAs<FAssetDragDropOp>();
+
+	if (TemplateDragDropOp.IsValid())
+	{
+		UWidget* Widget = TemplateDragDropOp->Template->Create(Blueprint->WidgetTree);
+
+		if (Widget)
+		{
+			if (Cast<UUserWidget>(Widget) == nullptr || Blueprint->IsWidgetFreeFromCircularReferences(Cast<UUserWidget>(Widget)))
+			{
+				OutWidgets.Add(Widget);
+			}
+		}
+	}
+	else if (AssetDragDropOp.IsValid())
+	{
+		for (const FAssetData& AssetData : AssetDragDropOp->GetAssets())
+		{
+			UWidget* Widget = nullptr;
+			UClass* AssetClass = FindObjectChecked<UClass>(ANY_PACKAGE, *AssetData.AssetClass.ToString());
+
+			if (FWidgetTemplateBlueprintClass::Supports(AssetClass))
+			{
+				// Allows a UMG Widget Blueprint to be dragged from the Content Browser to another Widget Blueprint...as long as we're not trying to place a
+				// blueprint inside itself.
+				FString BlueprintPath = Blueprint->GetPathName();
+				if (BlueprintPath != AssetData.ObjectPath.ToString())
+				{
+					Widget = FWidgetTemplateBlueprintClass(AssetData).Create(Blueprint->WidgetTree);
+
+					// Check to make sure that this widget can be added to the current blueprint
+					if (Cast<UUserWidget>(Widget) != nullptr && !Blueprint->IsWidgetFreeFromCircularReferences(Cast<UUserWidget>(Widget)))
+					{
+						Widget = nullptr;
+					}
+				}
+			}
+			else if (FWidgetTemplateImageClass::Supports(AssetClass))
+			{
+				Widget = FWidgetTemplateImageClass(AssetData).Create(Blueprint->WidgetTree);
+			}
+
+			if (Widget)
+			{
+				OutWidgets.Add(Widget);
+			}
+		}
+	}
+
+	// Mark the widgets for design-time rendering
+	for (UWidget* Widget : OutWidgets)
+	{
+		Widget->SetDesignerFlags(BlueprintEditor.Pin()->GetCurrentDesignerFlags());
+	}
+}
+
+void FHierarchyModel::RemovePreviewWidget(UWidgetBlueprint* Blueprint, UWidget* Widget)
+{
+	Blueprint->WidgetTree->SetFlags(RF_Transactional);
+	Blueprint->WidgetTree->Modify();
+	Blueprint->WidgetTree->RemoveWidget(Widget);
+	if (Widget->GetOutermost() != GetTransientPackage())
+	{
+		Widget->SetFlags(RF_NoFlags);
+		Widget->Rename(nullptr, GetTransientPackage());
+	}
+}
+
 void FHierarchyModel::InitializeChildren()
 {
 	if ( !bInitialized )
@@ -620,6 +724,7 @@ bool FHierarchyRoot::DoesWidgetOverrideFlowDirection() const
 
 TOptional<EItemDropZone> FHierarchyRoot::HandleCanAcceptDrop(const FDragDropEvent& DragDropEvent, EItemDropZone DropZone)
 {
+	bool bIsFreeFromCircularReferences = true;
 	TSharedPtr<FWidgetTemplateDragDropOp> TemplateDragDropOp = DragDropEvent.GetOperationAs<FWidgetTemplateDragDropOp>();
 	if (TemplateDragDropOp.IsValid())
 	{
@@ -634,20 +739,36 @@ TOptional<EItemDropZone> FHierarchyRoot::HandleCanAcceptDrop(const FDragDropEven
 				{
 					TemplateDragDropOp->CurrentIconBrush = FEditorStyle::GetBrush(TEXT("Graph.ConnectorFeedback.Error"));
 					TemplateDragDropOp->CurrentHoverText = LOCTEXT("CircularReference", "This would cause a circular reference.");
-					return TOptional<EItemDropZone>();
+					bIsFreeFromCircularReferences = false;
 				}
+
+				FHierarchyModel::RemovePreviewWidget(Blueprint, Widget);
 			}
 		}
 	}
 
 	bool bIsDrop = false;
-	return ProcessHierarchyDragDrop(DragDropEvent, DropZone, bIsDrop, BlueprintEditor.Pin(), FWidgetReference());
+	return bIsFreeFromCircularReferences ? ProcessHierarchyDragDrop(DragDropEvent, DropZone, bIsDrop, BlueprintEditor.Pin(), FWidgetReference()) : TOptional<EItemDropZone>();
 }
 
 FReply FHierarchyRoot::HandleAcceptDrop(FDragDropEvent const& DragDropEvent, EItemDropZone DropZone)
 {
 	bool bIsDrop = true;
 	TOptional<EItemDropZone> Zone = ProcessHierarchyDragDrop(DragDropEvent, DropZone, bIsDrop, BlueprintEditor.Pin(), FWidgetReference());
+	if (Zone.IsSet())
+	{
+		TArray<UWidget*> DragDropPreviewWidgets;
+		DetermineDragDropPreviewWidgets(DragDropPreviewWidgets, DragDropEvent);
+		// Move the remaining widgets into the transient package. Otherwise, they will remain outered to the WidgetTree and end up as properties in the BP class layout as a result.
+		UWidgetBlueprint* BP = Cast<UWidgetBlueprint>(BlueprintEditor.Pin()->GetBlueprintObj());
+		if (BP)
+		{
+			for (UWidget* Widget : DragDropPreviewWidgets)
+			{
+				FHierarchyModel::RemovePreviewWidget(BP, Widget);
+			}
+		}
+	}
 	return Zone.IsSet() ? FReply::Handled() : FReply::Unhandled();
 }
 
@@ -770,19 +891,21 @@ TOptional<EItemDropZone> FNamedSlotModel::HandleCanAcceptDrop(const FDragDropEve
 			}
 
 			UWidget* Widget = TemplateDragDropOp->Template->Create(Blueprint->WidgetTree);
-
+			bool bIsFreeFromCircularReferences = true;
 			if (Widget)
 			{
 				if (!Blueprint->IsWidgetFreeFromCircularReferences(Cast<UUserWidget>(Widget)))
 				{
 					TemplateDragDropOp->CurrentIconBrush = FEditorStyle::GetBrush(TEXT("Graph.ConnectorFeedback.Error"));
 					TemplateDragDropOp->CurrentHoverText = LOCTEXT("CircularReference", "This would cause a circular reference.");
-					return TOptional<EItemDropZone>();
+					bIsFreeFromCircularReferences = false;
 				}
+
+				FHierarchyModel::RemovePreviewWidget(Blueprint, Widget);
 			}
 
 			TemplateDragDropOp->CurrentIconBrush = FEditorStyle::GetBrush(TEXT("Graph.ConnectorFeedback.OK"));
-			return EItemDropZone::OntoItem;
+			return bIsFreeFromCircularReferences ? EItemDropZone::OntoItem : TOptional<EItemDropZone>();
 		}
 	}
 
@@ -903,6 +1026,12 @@ void FNamedSlotModel::DoDrop(UWidget* NamedSlotHostWidget, UWidget* DroppingWidg
 	TSet<FWidgetReference> SelectedTemplates;
 	SelectedTemplates.Add(BlueprintEditor.Pin()->GetReferenceFromTemplate(DroppingWidget));
 
+	// Move the remaining widgets into the transient package. Otherwise, they will remain outered to the WidgetTree and end up as properties in the BP class layout as a result.
+	if (DroppingWidget->GetOutermost() != GetTransientPackage())
+	{
+		FHierarchyModel::RemovePreviewWidget(Blueprint, DroppingWidget);
+	}
+
 	BlueprintEditor.Pin()->SelectWidgets(SelectedTemplates, false);
 }
 
@@ -998,6 +1127,7 @@ FSlateFontInfo FHierarchyWidget::GetFont() const
 TOptional<EItemDropZone> FHierarchyWidget::HandleCanAcceptDrop(const FDragDropEvent& DragDropEvent, EItemDropZone DropZone)
 {
 	TSharedPtr<FWidgetTemplateDragDropOp> TemplateDragDropOp = DragDropEvent.GetOperationAs<FWidgetTemplateDragDropOp>();
+	bool bIsFreeFromCircularReferences = true;
 	if (TemplateDragDropOp.IsValid())
 	{
 		UWidgetBlueprint* Blueprint = BlueprintEditor.Pin()->GetWidgetBlueprintObj();
@@ -1011,13 +1141,16 @@ TOptional<EItemDropZone> FHierarchyWidget::HandleCanAcceptDrop(const FDragDropEv
 				{
 					TemplateDragDropOp->CurrentIconBrush = FEditorStyle::GetBrush(TEXT("Graph.ConnectorFeedback.Error"));
 					TemplateDragDropOp->CurrentHoverText = LOCTEXT("CircularReference", "This would cause a circular reference.");
-					return TOptional<EItemDropZone>();
+					bIsFreeFromCircularReferences = false;
 				}
+
+				RemovePreviewWidget(Blueprint, Widget);
+
 			}
 		}
 	}
 	bool bIsDrop = false;
-	return ProcessHierarchyDragDrop(DragDropEvent, DropZone, bIsDrop, BlueprintEditor.Pin(), Item);
+	return bIsFreeFromCircularReferences ? ProcessHierarchyDragDrop(DragDropEvent, DropZone, bIsDrop, BlueprintEditor.Pin(), Item) :  TOptional<EItemDropZone>();
 }
 
 void FHierarchyWidget::HandleDragLeave(const FDragDropEvent& DragDropEvent)
@@ -1027,12 +1160,27 @@ void FHierarchyWidget::HandleDragLeave(const FDragDropEvent& DragDropEvent)
 	{
 		DecoratedDragDropOp->ResetToDefaultToolTip();
 	}
+	FHierarchyModel::HandleDragLeave(DragDropEvent);
 }
 
 FReply FHierarchyWidget::HandleAcceptDrop(const FDragDropEvent& DragDropEvent, EItemDropZone DropZone)
 {
 	bool bIsDrop = true;
 	TOptional<EItemDropZone> Zone = ProcessHierarchyDragDrop(DragDropEvent, DropZone, bIsDrop, BlueprintEditor.Pin(), Item);
+	if (Zone.IsSet())
+	{
+		TArray<UWidget*> DragDropPreviewWidgets;
+		DetermineDragDropPreviewWidgets(DragDropPreviewWidgets, DragDropEvent);
+		// Move the remaining widgets into the transient package. Otherwise, they will remain outered to the WidgetTree and end up as properties in the BP class layout as a result.
+		UWidgetBlueprint* BP = Cast<UWidgetBlueprint>(BlueprintEditor.Pin()->GetBlueprintObj());
+		if (BP)
+		{
+			for (UWidget* Widget : DragDropPreviewWidgets)
+			{
+				FHierarchyModel::RemovePreviewWidget(BP, Widget);
+			}
+		}
+	}
 	return Zone.IsSet() ? FReply::Handled() : FReply::Unhandled();
 }
 
diff --git a/Engine/Source/Editor/UMGEditor/Private/Hierarchy/SHierarchyViewItem.h b/Engine/Source/Editor/UMGEditor/Private/Hierarchy/SHierarchyViewItem.h
--- a/Engine/Source/Editor/UMGEditor/Private/Hierarchy/SHierarchyViewItem.h
+++ b/Engine/Source/Editor/UMGEditor/Private/Hierarchy/SHierarchyViewItem.h
@@ -102,6 +102,8 @@ protected:
 	virtual void GetChildren(TArray< TSharedPtr<FHierarchyModel> >& Children) = 0;
 	virtual void UpdateSelection() = 0;
 	virtual FWidgetReference AsDraggedWidgetReference() const { return FWidgetReference(); }
+	void DetermineDragDropPreviewWidgets(TArray<class UWidget*>& OutWidgets, const FDragDropEvent& DragDropEvent);
+	void RemovePreviewWidget(class UWidgetBlueprint* Blueprint, class UWidget* Widget);
 
 private:
 	void InitializeChildren();
