Tentative fix for double-rename crash

diff --git a/Engine/Source/Editor/UMGEditor/Private/WidgetBlueprintEditorUtils.cpp b/Engine/Source/Editor/UMGEditor/Private/WidgetBlueprintEditorUtils.cpp
--- a/Engine/Source/Editor/UMGEditor/Private/WidgetBlueprintEditorUtils.cpp
+++ b/Engine/Source/Editor/UMGEditor/Private/WidgetBlueprintEditorUtils.cpp
@@ -216,9 +216,12 @@ bool FWidgetBlueprintEditorUtils::RenameWidget(TSharedRef<FWidgetBlueprintEditor
 			WidgetPreview->Rename(*NewNameStr);
 		}
 
-		// Find and update all variable references in the graph
-		Widget->SetDisplayLabel(NewDisplayName);
-		Widget->Rename(*NewNameStr);
+		if (!WidgetPreview || WidgetPreview != Widget)
+		{
+			// Find and update all variable references in the graph
+			Widget->SetDisplayLabel(NewDisplayName);
+			Widget->Rename(*NewNameStr);
+		}
 
 		// Update Variable References and
 		// Update Event References to member variables
