Fix to drag and drop causing widgets to be incorectly set to Transient after a drop.

diff --git a/Engine/Source/Editor/UMGEditor/Private/Hierarchy/SHierarchyViewItem.cpp b/Engine/Source/Editor/UMGEditor/Private/Hierarchy/SHierarchyViewItem.cpp
--- a/Engine/Source/Editor/UMGEditor/Private/Hierarchy/SHierarchyViewItem.cpp
+++ b/Engine/Source/Editor/UMGEditor/Private/Hierarchy/SHierarchyViewItem.cpp
@@ -1026,12 +1026,6 @@ void FNamedSlotModel::DoDrop(UWidget* NamedSlotHostWidget, UWidget* DroppingWidg
 	TSet<FWidgetReference> SelectedTemplates;
 	SelectedTemplates.Add(BlueprintEditor.Pin()->GetReferenceFromTemplate(DroppingWidget));
 
-	// Move the remaining widgets into the transient package. Otherwise, they will remain outered to the WidgetTree and end up as properties in the BP class layout as a result.
-	if (DroppingWidget->GetOutermost() != GetTransientPackage())
-	{
-		FHierarchyModel::RemovePreviewWidget(Blueprint, DroppingWidget);
-	}
-
 	BlueprintEditor.Pin()->SelectWidgets(SelectedTemplates, false);
 }
 
