Fix to copy paste of Widget with named slot that created a copy of the widget in the root.

diff --git a/Engine/Source/Editor/UMGEditor/Private/WidgetBlueprintEditorUtils.cpp b/Engine/Source/Editor/UMGEditor/Private/WidgetBlueprintEditorUtils.cpp
--- a/Engine/Source/Editor/UMGEditor/Private/WidgetBlueprintEditorUtils.cpp
+++ b/Engine/Source/Editor/UMGEditor/Private/WidgetBlueprintEditorUtils.cpp
@@ -1130,7 +1130,25 @@ TArray<UWidget*> FWidgetBlueprintEditorUtils::PasteWidgetsInternal(TSharedRef<FW
 		// automatically.
 		if ( NewWidget->GetParent() == nullptr )
 		{
-			RootPasteWidgets.Add(NewWidget);
+			// Check to see if this widget is content of another widget holding it in a named slot.
+			bool bIsNamedSlot = false;
+			for (UWidget* ContainerWidget : PastedWidgets)
+			{
+				if (INamedSlotInterface* NamedSlotContainer = Cast<INamedSlotInterface>(ContainerWidget))
+				{
+					if (NamedSlotContainer->ContainsContent(NewWidget))
+					{
+						bIsNamedSlot = true;
+						break;
+					}
+				}
+			}
+
+			// It's a Root widget only if it's not not in a named slot.
+			if (!bIsNamedSlot)
+			{
+				RootPasteWidgets.Add(NewWidget);
+			}
 		}
 	}
 
