Fix 'Assign Asset from Content Browser' for Create Widget node

diff --git a/Engine/Source/Editor/GraphEditor/Private/KismetPins/SGraphPinClass.cpp b/Engine/Source/Editor/GraphEditor/Private/KismetPins/SGraphPinClass.cpp
--- a/Engine/Source/Editor/GraphEditor/Private/KismetPins/SGraphPinClass.cpp
+++ b/Engine/Source/Editor/GraphEditor/Private/KismetPins/SGraphPinClass.cpp
@@ -19,6 +19,7 @@
 void SGraphPinClass::Construct(const FArguments& InArgs, UEdGraphPin* InGraphPinObj)
 {
 	SGraphPin::Construct(SGraphPin::FArguments(), InGraphPinObj);
+	bAllowAbstractClasses = true;
 }
 
 FReply SGraphPinClass::OnClickUse()
@@ -52,6 +53,8 @@ public:
 	/** All children of these classes will be included unless filtered out by another setting. */
 	TSet< const UClass* > AllowedChildrenOfClasses;
 
+	bool bAllowAbstractClasses = true;
+
 	virtual bool IsClassAllowed(const FClassViewerInitializationOptions& InInitOptions, const UClass* InClass, TSharedRef< FClassViewerFilterFuncs > InFilterFuncs ) override
 	{
 		// If it appears on the allowed child-of classes list (or there is nothing on that list)
@@ -64,6 +67,7 @@ public:
 		
 			// Don't allow classes from a loaded map (e.g. LSBPs) unless we're already working inside that package context. Otherwise, choosing the class would lead to a GLEO at save time.
 			Result &= !ClassPackage->ContainsMap() || ClassPackage == GraphPinOutermostPackage;
+			Result &= bAllowAbstractClasses || (InClass != nullptr && !InClass->HasAnyClassFlags(CLASS_Abstract));
 		}
 
 		return Result;
@@ -71,7 +75,7 @@ public:
 
 	virtual bool IsUnloadedClassAllowed(const FClassViewerInitializationOptions& InInitOptions, const TSharedRef< const IUnloadedBlueprintData > InUnloadedClassData, TSharedRef< FClassViewerFilterFuncs > InFilterFuncs) override
 	{
-		return (InFilterFuncs->IfInChildOfClassesSet( AllowedChildrenOfClasses, InUnloadedClassData) != EFilterReturn::Failed);
+		return (InFilterFuncs->IfInChildOfClassesSet( AllowedChildrenOfClasses, InUnloadedClassData) != EFilterReturn::Failed) && (bAllowAbstractClasses || !InUnloadedClassData->HasAnyClassFlags(CLASS_Abstract));
 	}
 };
 
@@ -93,6 +97,7 @@ TSharedRef<SWidget> SGraphPinClass::GenerateAssetPicker()
 	}
 
 	TSharedPtr<FGraphPinFilter> Filter = MakeShareable(new FGraphPinFilter);
+	Filter->bAllowAbstractClasses = bAllowAbstractClasses;
 	Options.ClassFilter = Filter;
 
 	Filter->AllowedChildrenOfClasses.Add(PinRequiredParentClass);
diff --git a/Engine/Source/Editor/GraphEditor/Private/KismetPins/SGraphPinClass.h b/Engine/Source/Editor/GraphEditor/Private/KismetPins/SGraphPinClass.h
deleted file mode 100644
diff --git a/Engine/Source/Editor/GraphEditor/Public/KismetPins/SGraphPinClass.h b/Engine/Source/Editor/GraphEditor/Public/KismetPins/SGraphPinClass.h
new file mode 100644
--- /dev/null
+++ b/Engine/Source/Editor/GraphEditor/Public/KismetPins/SGraphPinClass.h
@@ -0,0 +1,41 @@
+// Copyright 1998-2019 Epic Games, Inc. All Rights Reserved.
+
+#pragma once
+
+#include "CoreMinimal.h"
+#include "Widgets/DeclarativeSyntaxSupport.h"
+#include "Input/Reply.h"
+#include "Widgets/SWidget.h"
+#include "KismetPins/SGraphPinObject.h"
+
+/////////////////////////////////////////////////////
+// SGraphPinClass
+
+class GRAPHEDITOR_API SGraphPinClass : public SGraphPinObject
+{
+public:
+	SLATE_BEGIN_ARGS(SGraphPinClass) {}
+	SLATE_END_ARGS()
+
+	void Construct(const FArguments& InArgs, UEdGraphPin* InGraphPinObj);
+	void SetAllowAbstractClasses(bool bAllow) { bAllowAbstractClasses = bAllow; }
+protected:
+
+	// Called when a new class was picked via the asset picker
+	void OnPickedNewClass(UClass* ChosenClass);
+
+	//~ Begin SGraphPinObject Interface
+	virtual FReply OnClickUse() override;
+	virtual bool AllowSelfPinWidget() const override { return false; }
+	virtual TSharedRef<SWidget> GenerateAssetPicker() override;
+	virtual FText GetDefaultComboText() const override;
+	virtual FOnClicked GetOnUseButtonDelegate() override;
+	virtual const FAssetData& GetAssetData(bool bRuntimePath) const override;
+	//~ End SGraphPinObject Interface
+
+	/** Cached AssetData without the _C */
+	mutable FAssetData CachedEditorAssetData;
+
+	/** Whether abstract classes should be filtered out in the class viewer */
+	bool bAllowAbstractClasses;
+};
diff --git a/Engine/Source/Editor/UMGEditor/Private/Widgets/SGraphNodeCreateWidget.cpp b/Engine/Source/Editor/UMGEditor/Private/Widgets/SGraphNodeCreateWidget.cpp
--- a/Engine/Source/Editor/UMGEditor/Private/Widgets/SGraphNodeCreateWidget.cpp
+++ b/Engine/Source/Editor/UMGEditor/Private/Widgets/SGraphNodeCreateWidget.cpp
@@ -1,104 +1,9 @@
 // Copyright 1998-2019 Epic Games, Inc. All Rights Reserved.
 
 #include "Widgets/SGraphNodeCreateWidget.h"
-#include "Modules/ModuleManager.h"
-#include "Widgets/SBoxPanel.h"
-#include "Widgets/Layout/SBox.h"
-#include "Editor.h"
-#include "EdGraphSchema_K2.h"
-#include "Nodes/K2Node_CreateWidget.h"
-#include "KismetPins/SGraphPinObject.h"
-#include "NodeFactory.h"
-#include "ClassViewerModule.h"
-#include "ClassViewerFilter.h"
-#include "ScopedTransaction.h"
-#include "Blueprint/UserWidget.h"
-
-#define LOCTEXT_NAMESPACE "SGraphPinUserWidgetBasedClass"
-
-//////////////////////////////////////////////////////////////////////////
-// SGraphPinUserWidgetBasedClass
-
-/**
- * GraphPin can select only UUserWidget classes.
- * Instead of asset picker, a class viewer is used.
- */
-class SGraphPinUserWidgetBasedClass : public SGraphPinObject
-{
-	class FUserWidgetBasedClassFilter : public IClassViewerFilter
-	{
-	public:
-
-		virtual bool IsClassAllowed(const FClassViewerInitializationOptions& InInitOptions, const UClass* InClass, TSharedRef< FClassViewerFilterFuncs > InFilterFuncs ) override
-		{
-			if(NULL != InClass)
-			{
-				const bool bUserWidgetBased = InClass->IsChildOf(UUserWidget::StaticClass());
-				const bool bBlueprintType = UEdGraphSchema_K2::IsAllowableBlueprintVariableType(InClass);
-				const bool bNotAbstract = !InClass->HasAnyClassFlags(CLASS_Abstract);
-				return bUserWidgetBased && bBlueprintType && bNotAbstract;
-			}
-			return false;
-		}
-
-		virtual bool IsUnloadedClassAllowed(const FClassViewerInitializationOptions& InInitOptions, const TSharedRef< const IUnloadedBlueprintData > InUnloadedClassData, TSharedRef< FClassViewerFilterFuncs > InFilterFuncs) override
-		{
-			const bool bUserWidgetBased = InUnloadedClassData->IsChildOf(UUserWidget::StaticClass());
-			const bool bNotAbstract = !InUnloadedClassData->HasAnyClassFlags(CLASS_Abstract);
-			return bUserWidgetBased && bNotAbstract;
-		}
-	};
-
-protected:
-
-	void OnPickedNewClass(UClass* ChosenClass)
-	{
-		if (GraphPinObj->DefaultObject != ChosenClass)
-		{
-			const FScopedTransaction Transaction(NSLOCTEXT("GraphEditor", "ChangeClassPinValue", "Change Class Pin Value"));
-			GraphPinObj->Modify();
-
-			AssetPickerAnchor->SetIsOpen(false);
-			GraphPinObj->GetSchema()->TrySetDefaultObject(*GraphPinObj, ChosenClass);
-		}
-	}
-
-	virtual TSharedRef<SWidget> GenerateAssetPicker() override
-	{
-		FClassViewerModule& ClassViewerModule = FModuleManager::LoadModuleChecked<FClassViewerModule>("ClassViewer");
 
-		FClassViewerInitializationOptions Options;
-		Options.Mode = EClassViewerMode::ClassPicker;
-		Options.bIsActorsOnly = false;
-		Options.DisplayMode = EClassViewerDisplayMode::DefaultView;
-		Options.bShowUnloadedBlueprints = true;
-		Options.bShowNoneOption = false;
-		Options.bShowObjectRootClass = true;
-		TSharedPtr< FUserWidgetBasedClassFilter > Filter = MakeShareable(new FUserWidgetBasedClassFilter);
-		Options.ClassFilter = Filter;
-
-		return
-			SNew(SBox)
-			.WidthOverride(280)
-			[
-				SNew(SVerticalBox)
-				+SVerticalBox::Slot()
-				.AutoHeight()
-				.MaxHeight(500)
-				[
-					SNew(SBorder)
-					.Padding(4)
-					.BorderImage( FEditorStyle::GetBrush("ToolPanel.GroupBorder") )
-					[
-						ClassViewerModule.CreateClassViewer(Options, FOnClassPicked::CreateSP(this, &SGraphPinUserWidgetBasedClass::OnPickedNewClass))
-					]
-				]
-			];
-	}
-};
-
-//////////////////////////////////////////////////////////////////////////
-// SGraphNodeCreateWidget
+#include "Nodes/K2Node_CreateWidget.h"
+#include "KismetPins/SGraphPinClass.h"
 
 TSharedPtr<SGraphPin> SGraphNodeCreateWidget::CreatePinWidget(UEdGraphPin* Pin) const
 {
@@ -106,11 +11,10 @@ TSharedPtr<SGraphPin> SGraphNodeCreateWidget::CreatePinWidget(UEdGraphPin* Pin)
 	UEdGraphPin* ClassPin = CreateWidgetNode->GetClassPin();
 	if ((ClassPin == Pin) && (!ClassPin->bHidden || (ClassPin->LinkedTo.Num() > 0)))
 	{
-		TSharedPtr<SGraphPinUserWidgetBasedClass> NewPin = SNew(SGraphPinUserWidgetBasedClass, ClassPin);
+		TSharedPtr<SGraphPinClass> NewPin = SNew(SGraphPinClass, ClassPin);
 		check(NewPin.IsValid());
+		NewPin->SetAllowAbstractClasses(false);
 		return NewPin;
 	}
 	return SGraphNodeK2Default::CreatePinWidget(Pin);
 }
-
-#undef LOCTEXT_NAMESPACE
