Copying the details panel and single property UMG widgets to this stream. Adding UX updates to those widgets.
#rnx

diff --git a/Engine/Source/Editor/Blutility/Classes/EditorUtilityWidgetBlueprint.h b/Engine/Source/Editor/Blutility/Classes/EditorUtilityWidgetBlueprint.h
--- a/Engine/Source/Editor/Blutility/Classes/EditorUtilityWidgetBlueprint.h
+++ b/Engine/Source/Editor/Blutility/Classes/EditorUtilityWidgetBlueprint.h
@@ -38,6 +38,8 @@ public:
 	// UBlueprint interface
 	virtual void GetReparentingRules(TSet< const UClass* >& AllowedChildrenOfClasses, TSet< const UClass* >& DisallowedChildrenOfClasses) const override;
 
+	virtual bool AllowEditorWidget() const override { return true; }
+
 	UEditorUtilityWidget* GetCreatedWidget() const
 	{
 		return CreatedUMGWidget;
diff --git a/Engine/Source/Editor/PropertyEditor/Private/SSingleProperty.cpp b/Engine/Source/Editor/PropertyEditor/Private/SSingleProperty.cpp
--- a/Engine/Source/Editor/PropertyEditor/Private/SSingleProperty.cpp
+++ b/Engine/Source/Editor/PropertyEditor/Private/SSingleProperty.cpp
@@ -169,7 +169,7 @@ void SSingleProperty::SetObject( UObject* InObject )
 		if( NamePlacement != EPropertyNamePlacement::Hidden )
 		{
 			HorizontalBox->AddSlot()
-			.Padding( 2.0f, 0.0f, 2.0f, 4.0f )
+			.Padding(4.0f, 0.0f)
 			.AutoWidth()
 			.VAlign( VAlign_Center )
 			[
@@ -179,7 +179,7 @@ void SSingleProperty::SetObject( UObject* InObject )
 		}
 
 		HorizontalBox->AddSlot()
-		.Padding( 0.0f, 2.0f, 0.0f, 2.0f )
+		.Padding( 4.0f, 0.0f)
 		.FillWidth(1.0f)
 		.VAlign( VAlign_Center )
 		[
diff --git a/Engine/Source/Editor/UMGEditor/Private/Components/DetailsView.cpp b/Engine/Source/Editor/UMGEditor/Private/Components/DetailsView.cpp
new file mode 100644
--- /dev/null
+++ b/Engine/Source/Editor/UMGEditor/Private/Components/DetailsView.cpp
@@ -0,0 +1,184 @@
+// Copyright 1998-2019 Epic Games, Inc. All Rights Reserved.
+
+#include "Components/DetailsView.h"
+
+#include "Components/PropertyViewHelper.h"
+#include "IDetailsView.h"
+#include "Modules/ModuleManager.h"
+#include "PropertyEditorModule.h"
+#include "Widgets/Layout/SBorder.h"
+#include "Widgets/Text/STextBlock.h"
+#include "UObject/UObjectGlobals.h"
+#include "ObjectEditorUtils.h"
+
+#define LOCTEXT_NAMESPACE "UMG"
+
+/////////////////////////////////////////////////////
+// UDetailsView
+
+
+void UDetailsView::ReleaseSlateResources(bool bReleaseChildren)
+{
+	Super::ReleaseSlateResources(bReleaseChildren);
+
+	DetailViewWidget.Reset();
+}
+
+
+void UDetailsView::BuildContentWidget()
+{
+	DetailViewWidget.Reset();
+
+	if (!GetDisplayWidget().IsValid())
+	{
+		return;
+	}
+
+	bool bCreateMissingWidget = true;
+	FText MissingWidgetText = FPropertyViewHelper::EditorOnlyText;
+
+	if (GIsEditor)
+	{
+		UObject* ViewedObject = GetObject();
+		if (ViewedObject == nullptr)
+		{
+			bool bIsLazyObjectNull = LazyObject.IsNull();
+			if (bIsLazyObjectNull)
+			{
+				MissingWidgetText = FPropertyViewHelper::UndefinedObjectText;
+			}
+			else
+			{
+				MissingWidgetText = FPropertyViewHelper::UnloadedObjectText;
+			}
+		}
+		else
+		{
+			FPropertyEditorModule& PropertyEditorModule = FModuleManager::GetModuleChecked<FPropertyEditorModule>("PropertyEditor");
+
+			FDetailsViewArgs DetailsViewArgs;
+			DetailsViewArgs.NameAreaSettings = FDetailsViewArgs::HideNameArea;
+			DetailsViewArgs.bUpdatesFromSelection = false;
+			DetailsViewArgs.bLockable = false;
+			DetailsViewArgs.bShowPropertyMatrixButton = false;
+			DetailsViewArgs.NotifyHook = this;
+
+			DetailsViewArgs.ViewIdentifier = ViewIdentifier;
+			DetailsViewArgs.bAllowSearch = bAllowFiltering;
+			DetailsViewArgs.bAllowFavoriteSystem = bAllowFavoriteSystem;
+			DetailsViewArgs.bShowOptions = bAllowFiltering;
+			DetailsViewArgs.bShowModifiedPropertiesOption = bShowModifiedPropertiesOption;
+			DetailsViewArgs.bShowKeyablePropertiesOption = bShowKeyablePropertiesOption;
+			DetailsViewArgs.bShowAnimatedPropertiesOption = bShowAnimatedPropertiesOption;
+			DetailsViewArgs.bShowScrollBar = bShowScrollBar;
+			DetailsViewArgs.bForceHiddenPropertyVisibility = bForceHiddenPropertyVisibility;
+			DetailsViewArgs.ColumnWidth = ColumnWidth;
+			DetailsViewArgs.bShowCustomFilterOption = PropertiesToShow.Num() != 0 || CategoriesToShow.Num() != 0;
+			
+			DetailViewWidget = PropertyEditorModule.CreateDetailView(DetailsViewArgs);
+			
+			DetailViewWidget->SetCustomFilterLabel(LOCTEXT("ShowAllParameters", "Show All Parameters"));
+			DetailViewWidget->SetCustomFilterDelegate(FSimpleDelegate::CreateUObject(this, &UDetailsView::ToggleWhitelistedProperties));
+
+			DetailViewWidget->SetIsPropertyVisibleDelegate(FIsPropertyVisible::CreateUObject(this, &UDetailsView::GetIsPropertyVisible));
+			DetailViewWidget->SetObject(ViewedObject);
+			if (DetailViewWidget.IsValid())
+			{
+				GetDisplayWidget()->SetContent(DetailViewWidget.ToSharedRef());
+				bCreateMissingWidget = false;
+			}
+			else
+			{
+				MissingWidgetText = FPropertyViewHelper::UnknownErrorText;
+			}
+		}
+	}
+
+	if (bCreateMissingWidget)
+	{
+		GetDisplayWidget()->SetContent(
+			SNew(STextBlock)
+			.Text(MissingWidgetText)
+		);
+	}
+}
+
+
+void UDetailsView::OnObjectChanged()
+{
+	UObject* ViewedObject = GetObject();
+	if (DetailViewWidget.IsValid() && ViewedObject != nullptr)
+	{
+		DetailViewWidget->SetObject(ViewedObject);
+	}
+	else
+	{
+		AsynBuildContentWidget();
+	}
+}
+
+
+void UDetailsView::NotifyPostChange(const FPropertyChangedEvent& PropertyChangedEvent, UProperty* PropertyThatChanged)
+{
+	FNotifyHook::NotifyPostChange(PropertyChangedEvent, PropertyThatChanged);
+
+	FName PropertyName = PropertyThatChanged ? PropertyThatChanged->GetFName() : NAME_None;
+	OnPropertyChangedBroadcast(PropertyName);
+}
+
+
+void UDetailsView::PostEditChangeProperty(FPropertyChangedEvent& PropertyChangedEvent)
+{
+	Super::PostEditChangeProperty(PropertyChangedEvent);
+	
+	if (IsDesignTime())
+	{
+		if (PropertyChangedEvent.GetPropertyName() == GET_MEMBER_NAME_CHECKED(UDetailsView, ViewIdentifier)
+			|| PropertyChangedEvent.GetPropertyName() == GET_MEMBER_NAME_CHECKED(UDetailsView, bAllowFiltering)
+			|| PropertyChangedEvent.GetPropertyName() == GET_MEMBER_NAME_CHECKED(UDetailsView, bAllowFavoriteSystem)
+			|| PropertyChangedEvent.GetPropertyName() == GET_MEMBER_NAME_CHECKED(UDetailsView, bShowModifiedPropertiesOption)
+			|| PropertyChangedEvent.GetPropertyName() == GET_MEMBER_NAME_CHECKED(UDetailsView, bShowKeyablePropertiesOption)
+			|| PropertyChangedEvent.GetPropertyName() == GET_MEMBER_NAME_CHECKED(UDetailsView, bShowAnimatedPropertiesOption)
+			|| PropertyChangedEvent.GetPropertyName() == GET_MEMBER_NAME_CHECKED(UDetailsView, bShowScrollBar)
+			|| PropertyChangedEvent.GetPropertyName() == GET_MEMBER_NAME_CHECKED(UDetailsView, bForceHiddenPropertyVisibility)
+			|| PropertyChangedEvent.GetPropertyName() == GET_MEMBER_NAME_CHECKED(UDetailsView, ColumnWidth))
+		{
+			SoftObjectPath = LazyObject.Get();
+			AsynBuildContentWidget();
+		}
+	}
+}
+
+
+void UDetailsView::ToggleWhitelistedProperties()
+{
+	bShowOnlyWhitelisted = !bShowOnlyWhitelisted;
+	if (DetailViewWidget.IsValid())
+	{
+		DetailViewWidget->ForceRefresh();
+	}
+}
+
+
+bool UDetailsView::GetIsPropertyVisible(const FPropertyAndParent& PropertyAndParent) const
+{
+	if (!bShowOnlyWhitelisted
+		|| (PropertiesToShow.Num() == 0 && CategoriesToShow.Num() == 0))
+	{
+		return true;
+	}
+
+	if (PropertiesToShow.Contains(PropertyAndParent.Property.GetFName()))
+	{
+		return true;
+	}
+	if (CategoriesToShow.Contains(FObjectEditorUtils::GetCategoryFName(&PropertyAndParent.Property)))
+	{
+		return true;
+	}
+	return false;
+}
+
+/////////////////////////////////////////////////////
+
+#undef LOCTEXT_NAMESPACE
diff --git a/Engine/Source/Editor/UMGEditor/Private/Components/PropertyViewBase.cpp b/Engine/Source/Editor/UMGEditor/Private/Components/PropertyViewBase.cpp
new file mode 100644
--- /dev/null
+++ b/Engine/Source/Editor/UMGEditor/Private/Components/PropertyViewBase.cpp
@@ -0,0 +1,156 @@
+// Copyright 1998-2019 Epic Games, Inc. All Rights Reserved.
+
+#include "Components/PropertyViewBase.h"
+
+#include "Async/Async.h"
+#include "Editor.h"
+#include "Engine/World.h"
+#include "PropertyEditorModule.h"
+#include "Widgets/Layout/SBorder.h"
+#include "UObject/UObjectGlobals.h"
+
+#define LOCTEXT_NAMESPACE "UMG"
+
+/////////////////////////////////////////////////////
+// UPropertyViewBase
+
+
+void UPropertyViewBase::ReleaseSlateResources(bool bReleaseChildren)
+{
+	Super::ReleaseSlateResources(bReleaseChildren);
+
+	FCoreUObjectDelegates::OnAssetLoaded.Remove(AssetLoadedHandle);
+	AssetLoadedHandle.Reset();
+	FCoreUObjectDelegates::PostLoadMapWithWorld.Remove(PostLoadMapHandle);
+	PostLoadMapHandle.Reset();
+	FEditorDelegates::MapChange.Remove(MapChangeHandle);
+	MapChangeHandle.Reset();
+
+	DisplayedWidget.Reset();
+}
+
+
+TSharedRef<SWidget> UPropertyViewBase::RebuildWidget()
+{
+	DisplayedWidget = SNew(SBorder)
+		.Padding(0.0f)
+		.HAlign(HAlign_Fill)
+		.VAlign(VAlign_Fill)
+		.BorderImage(FEditorStyle::GetBrush("NoBorder"));
+	
+	BuildContentWidget();
+
+	if (!AssetLoadedHandle.IsValid())
+	{
+		AssetLoadedHandle = FCoreUObjectDelegates::OnAssetLoaded.AddUObject(this, &UPropertyViewBase::InternalOnAssetLoaded);
+	}
+	if (!PostLoadMapHandle.IsValid())
+	{
+		PostLoadMapHandle = FCoreUObjectDelegates::PostLoadMapWithWorld.AddUObject(this, &UPropertyViewBase::InternalPostLoadMapWithWorld);
+	}
+	if (!MapChangeHandle.IsValid())
+	{
+		MapChangeHandle = FEditorDelegates::MapChange.AddUObject(this, &UPropertyViewBase::InternalOnMapChange);
+	}
+
+	return DisplayedWidget.ToSharedRef();
+}
+
+void UPropertyViewBase::AsynBuildContentWidget()
+{
+	if (!bIsAsyncBuildContentRequested)
+	{
+		bIsAsyncBuildContentRequested = true;
+
+		AsyncTask(ENamedThreads::GameThread, [this]()
+		{
+			if (bIsAsyncBuildContentRequested)
+			{
+				BuildContentWidget();
+				bIsAsyncBuildContentRequested = false;
+			}
+		});
+	}
+}
+
+
+UObject* UPropertyViewBase::GetObject() const
+{
+	return LazyObject.Get();
+}
+
+
+void UPropertyViewBase::SetObject(UObject* InObject)
+{
+	if (LazyObject.Get() != InObject)
+	{
+		LazyObject = InObject;
+		SoftObjectPath = InObject;
+		OnObjectChanged();
+	}
+}
+
+
+void UPropertyViewBase::OnPropertyChangedBroadcast(FName PropertyName)
+{
+	OnPropertyChanged.Broadcast(PropertyName);
+}
+
+
+void UPropertyViewBase::InternalOnAssetLoaded(UObject* AssetLoaded)
+{
+	if (SoftObjectPath.GetAssetPathName() == FSoftObjectPath(AssetLoaded).GetAssetPathName())
+	{
+		AsynBuildContentWidget();
+	}
+}
+
+
+void UPropertyViewBase::InternalPostLoadMapWithWorld(UWorld* AssetLoaded)
+{
+	InternalOnMapChange(0);
+}
+
+
+void UPropertyViewBase::InternalOnMapChange(uint32)
+{
+	AsynBuildContentWidget();
+}
+
+
+void UPropertyViewBase::PostLoad()
+{
+	Super::PostLoad();
+
+	if (!LazyObject.IsValid() && SoftObjectPath.IsAsset() && bAutoLoadAsset)
+	{
+		AsyncTask(ENamedThreads::GameThread, [this]()
+		{
+			LazyObject = SoftObjectPath.TryLoad();
+			BuildContentWidget();
+			bIsAsyncBuildContentRequested = false;
+		});
+	}
+}
+
+
+void UPropertyViewBase::PostEditChangeProperty(FPropertyChangedEvent& PropertyChangedEvent)
+{
+	if (PropertyChangedEvent.GetPropertyName() == GET_MEMBER_NAME_CHECKED(UPropertyViewBase, LazyObject))
+	{
+		SoftObjectPath = LazyObject.Get();
+		OnObjectChanged();
+	}
+
+	Super::PostEditChangeProperty(PropertyChangedEvent);
+}
+
+
+const FText UPropertyViewBase::GetPaletteCategory()
+{
+	return LOCTEXT("Editor", "Editor");
+}
+
+/////////////////////////////////////////////////////
+
+#undef LOCTEXT_NAMESPACE
diff --git a/Engine/Source/Editor/UMGEditor/Private/Components/PropertyViewHelper.cpp b/Engine/Source/Editor/UMGEditor/Private/Components/PropertyViewHelper.cpp
new file mode 100644
--- /dev/null
+++ b/Engine/Source/Editor/UMGEditor/Private/Components/PropertyViewHelper.cpp
@@ -0,0 +1,22 @@
+// Copyright 1998-2019 Epic Games, Inc. All Rights Reserved.
+
+#include "Components/PropertyViewHelper.h"
+
+
+#define LOCTEXT_NAMESPACE "UMG"
+
+/////////////////////////////////////////////////////
+// FPropertyViewHelper
+
+const FText FPropertyViewHelper::UndefinedObjectText = LOCTEXT("UndefinedObject", "Undefined Object");
+const FText FPropertyViewHelper::UnloadedObjectText = LOCTEXT("UnloadedObject", "Unloaded Object");
+const FText FPropertyViewHelper::UnknownErrorText = LOCTEXT("UnknownError", "Unknown Error");
+const FText FPropertyViewHelper::EditorOnlyText = LOCTEXT("EditorWidget", "Editor Only");
+const FText FPropertyViewHelper::UndefinedPropertyText = LOCTEXT("UndefinedProperty", "Undefined Property");
+const FText FPropertyViewHelper::UnknownPropertyText = LOCTEXT("UnknownProperty", "Unknown Property");
+const FText FPropertyViewHelper::InvalidPropertyText = LOCTEXT("InvalidPropertyName", "Invalid Property");
+const FText FPropertyViewHelper::UnsupportedPropertyText = LOCTEXT("SinglePropertyView_UnsupporteddProperty", "Unsupported Property");
+
+/////////////////////////////////////////////////////
+
+#undef LOCTEXT_NAMESPACE
\ No newline at end of file
diff --git a/Engine/Source/Editor/UMGEditor/Private/Components/PropertyViewHelper.h b/Engine/Source/Editor/UMGEditor/Private/Components/PropertyViewHelper.h
new file mode 100644
--- /dev/null
+++ b/Engine/Source/Editor/UMGEditor/Private/Components/PropertyViewHelper.h
@@ -0,0 +1,17 @@
+// Copyright 1998-2019 Epic Games, Inc. All Rights Reserved.
+
+#pragma once
+
+#include "CoreMinimal.h"
+
+struct FPropertyViewHelper
+{
+	static const FText UndefinedObjectText;
+	static const FText UnloadedObjectText;
+	static const FText UnknownErrorText;
+	static const FText EditorOnlyText;
+	static const FText UndefinedPropertyText;
+	static const FText UnknownPropertyText;
+	static const FText InvalidPropertyText;
+	static const FText UnsupportedPropertyText;
+};
diff --git a/Engine/Source/Editor/UMGEditor/Private/Components/SinglePropertyView.cpp b/Engine/Source/Editor/UMGEditor/Private/Components/SinglePropertyView.cpp
new file mode 100644
--- /dev/null
+++ b/Engine/Source/Editor/UMGEditor/Private/Components/SinglePropertyView.cpp
@@ -0,0 +1,153 @@
+// Copyright 1998-2019 Epic Games, Inc. All Rights Reserved.
+
+#include "Components/SinglePropertyView.h"
+
+#include "Components/PropertyViewHelper.h"
+#include "ISinglePropertyView.h"
+#include "Modules/ModuleManager.h"
+#include "PropertyEditorModule.h"
+#include "Widgets/Layout/SBorder.h"
+#include "Widgets/Text/STextBlock.h"
+#include "UObject/UObjectGlobals.h"
+
+
+#define LOCTEXT_NAMESPACE "UMG"
+
+/////////////////////////////////////////////////////
+// USinglePropertyView
+
+
+void USinglePropertyView::ReleaseSlateResources(bool bReleaseChildren)
+{
+	Super::ReleaseSlateResources(bReleaseChildren);
+
+	SinglePropertyViewWidget.Reset();
+}
+
+
+void USinglePropertyView::BuildContentWidget()
+{
+	SinglePropertyViewWidget.Reset();
+
+	if (!GetDisplayWidget().IsValid())
+	{
+		return;
+	}
+
+	bool bCreateMissingWidget = true;
+	FText MissingWidgetText = FPropertyViewHelper::EditorOnlyText;
+
+	if (GIsEditor)
+	{
+		UObject* ViewedObject = GetObject();
+		if (ViewedObject == nullptr)
+		{
+			bool bIsLazyObjectNull = LazyObject.IsNull();
+			if (bIsLazyObjectNull)
+			{
+				MissingWidgetText = FPropertyViewHelper::UndefinedObjectText;
+			}
+			else
+			{
+				MissingWidgetText = FPropertyViewHelper::UnloadedObjectText;
+			}
+		}
+		else if (PropertyName == NAME_None)
+		{
+			MissingWidgetText = FPropertyViewHelper::UndefinedPropertyText;
+		}
+		else
+		{
+			UProperty* Property = ViewedObject->GetClass()->FindPropertyByName(PropertyName);
+			if (Property == nullptr)
+			{
+				MissingWidgetText = FPropertyViewHelper::UnknownPropertyText;
+			}
+			else if (!Property->HasAllPropertyFlags(CPF_Edit))
+			{
+				MissingWidgetText = FPropertyViewHelper::InvalidPropertyText;
+			}
+			else if (Cast<UStructProperty>(Property) || Cast<UArrayProperty>(Property)
+				|| Cast<UMapProperty>(Property) || Cast<USetProperty>(Property)
+				)
+			{
+				MissingWidgetText = FPropertyViewHelper::UnsupportedPropertyText;
+			}
+			else
+			{
+				FPropertyEditorModule& PropertyEditorModule = FModuleManager::GetModuleChecked<FPropertyEditorModule>("PropertyEditor");
+				FSinglePropertyParams SinglePropertyArgs;
+				SinglePropertyArgs.NameOverride = NameOverride;
+				SinglePropertyViewWidget = PropertyEditorModule.CreateSingleProperty(ViewedObject, PropertyName, SinglePropertyArgs);
+
+				if (SinglePropertyViewWidget.IsValid())
+				{
+					FSimpleDelegate PropertyChanged = FSimpleDelegate::CreateUObject(this, &USinglePropertyView::InternalSinglePropertyChanged);
+					SinglePropertyViewWidget->SetOnPropertyValueChanged(PropertyChanged);
+
+					GetDisplayWidget()->SetContent(SinglePropertyViewWidget.ToSharedRef());
+					bCreateMissingWidget = false;
+				}
+				else
+				{
+					MissingWidgetText = FPropertyViewHelper::UnknownErrorText;
+				}
+			}
+		}
+	}
+
+	if (bCreateMissingWidget)
+	{
+		GetDisplayWidget()->SetContent(
+			SNew(STextBlock)
+			.Text(MissingWidgetText)
+		);
+	}
+}
+
+
+FName USinglePropertyView::GetPropertyName() const
+{
+	return PropertyName;
+}
+
+
+void USinglePropertyView::SetPropertyName(FName InPropertyName)
+{
+	if (PropertyName != InPropertyName)
+	{
+		PropertyName = InPropertyName;
+		AsynBuildContentWidget();
+	}
+}
+
+
+void USinglePropertyView::OnObjectChanged()
+{
+	AsynBuildContentWidget();
+}
+
+
+void USinglePropertyView::InternalSinglePropertyChanged()
+{
+	OnPropertyChangedBroadcast(GetPropertyName());
+}
+
+
+void USinglePropertyView::PostEditChangeProperty(FPropertyChangedEvent& PropertyChangedEvent)
+{
+	Super::PostEditChangeProperty(PropertyChangedEvent);
+
+	if (IsDesignTime())
+	{
+		if (PropertyChangedEvent.GetPropertyName() == GET_MEMBER_NAME_CHECKED(USinglePropertyView, PropertyName)
+			|| PropertyChangedEvent.GetPropertyName() == GET_MEMBER_NAME_CHECKED(USinglePropertyView, NameOverride))
+		{
+			AsynBuildContentWidget();
+		}
+	}
+}
+
+/////////////////////////////////////////////////////
+
+#undef LOCTEXT_NAMESPACE
diff --git a/Engine/Source/Editor/UMGEditor/Private/Palette/SPaletteView.cpp b/Engine/Source/Editor/UMGEditor/Private/Palette/SPaletteView.cpp
--- a/Engine/Source/Editor/UMGEditor/Private/Palette/SPaletteView.cpp
+++ b/Engine/Source/Editor/UMGEditor/Private/Palette/SPaletteView.cpp
@@ -195,7 +195,8 @@ void SPaletteView::Construct(const FArguments& InArgs, TSharedPtr<FWidgetBluepri
 
 	BlueprintEditor = InBlueprintEditor;
 
-	UBlueprint* BP = InBlueprintEditor->GetBlueprintObj();
+	UWidgetBlueprint* WBP = InBlueprintEditor->GetWidgetBlueprintObj();
+	bAllowEditorWidget = WBP ? WBP->AllowEditorWidget() : false;
 
 	WidgetFilter = MakeShareable(new WidgetViewModelTextFilter(
 		WidgetViewModelTextFilter::FItemToStringArray::CreateSP(this, &SPaletteView::TransformWidgetViewModelToString)));
@@ -476,6 +477,14 @@ void SPaletteView::BuildClassWidgetList()
 			continue;
 		}
 
+		if (!bAllowEditorWidget)
+		{
+			if (WidgetClass->GetOutermost()->IsEditorOnly())
+			{
+				continue;
+			}
+		}
+
 		if (WidgetClass->IsChildOf(UUserWidget::StaticClass()))
 		{
 			if ( WidgetClass->ClassGeneratedBy )
diff --git a/Engine/Source/Editor/UMGEditor/Private/Palette/SPaletteView.h b/Engine/Source/Editor/UMGEditor/Private/Palette/SPaletteView.h
--- a/Engine/Source/Editor/UMGEditor/Private/Palette/SPaletteView.h
+++ b/Engine/Source/Editor/UMGEditor/Private/Palette/SPaletteView.h
@@ -124,4 +124,7 @@ private:
 
 	/** Controls rebuilding the list of spawnable widgets */
 	bool bRebuildRequested;
+
+	/** Are blutility widget supported. */
+	bool bAllowEditorWidget;
 };
diff --git a/Engine/Source/Editor/UMGEditor/Public/Components/DetailsView.h b/Engine/Source/Editor/UMGEditor/Public/Components/DetailsView.h
new file mode 100644
--- /dev/null
+++ b/Engine/Source/Editor/UMGEditor/Public/Components/DetailsView.h
@@ -0,0 +1,95 @@
+// Copyright 1998-2019 Epic Games, Inc. All Rights Reserved.
+
+#pragma once
+
+#include "Components/PropertyViewBase.h"
+#include "Misc/NotifyHook.h"
+#include "PropertyEditorDelegates.h"
+#include "DetailsView.generated.h"
+
+
+class IDetailsView;
+
+/**
+ * The details view allows you to display the value of an object properties.
+ */
+UCLASS()
+class UMGEDITOR_API UDetailsView : public UPropertyViewBase, public FNotifyHook
+{
+	GENERATED_BODY()
+
+public:
+	/** True if we allow filtering through search and the filter dropdown menu. */
+	UPROPERTY(EditAnywhere, Category = "View")
+	bool bAllowFiltering = false;
+
+	/** If false, the current properties editor will never display the favorite system */
+	UPROPERTY(EditAnywhere, AdvancedDisplay, Category = "View")
+	bool bAllowFavoriteSystem = false;
+
+	/** True if you want to show the 'Show Only Modified Properties'. Only valid in conjunction with bAllowFiltering */
+	UPROPERTY(EditAnywhere, AdvancedDisplay, Category = "View", meta = (EditCondition = "bAllowFiltering"))
+	bool bShowModifiedPropertiesOption = false;
+
+	/** True if you want to show the 'Show Only Keyable Properties'. Only valid in conjunction with bAllowFiltering */
+	UPROPERTY(EditAnywhere, AdvancedDisplay, Category = "View", meta = (EditCondition = "bAllowFiltering"))
+	bool bShowKeyablePropertiesOption = true;
+
+	/** True if you want to show the 'Show Only Animated Properties'. Only valid in conjunction with bAllowFiltering */
+	UPROPERTY(EditAnywhere, AdvancedDisplay, Category = "View", meta = (EditCondition = "bAllowFiltering"))
+	bool bShowAnimatedPropertiesOption = true;
+
+	/** The default column width */
+	UPROPERTY(EditAnywhere, Category = "View")
+	float ColumnWidth = 0.65f;
+
+	/** If false, the details panel's scrollbar will always be hidden. Useful when embedding details panels in widgets that either grow to accommodate them, or with scrollbars of their own. */
+	UPROPERTY(EditAnywhere, AdvancedDisplay, Category = "View")
+	bool bShowScrollBar = true;
+
+	/** If true, all properties will be visible, not just those with CPF_Edit */
+	UPROPERTY(EditAnywhere, AdvancedDisplay, Category = "View")
+	bool bForceHiddenPropertyVisibility = false;
+
+	/** Identifier for this details view; NAME_None if this view is anonymous */
+	UPROPERTY(EditAnywhere, AdvancedDisplay, Category = "View")
+	FName ViewIdentifier = NAME_None;
+
+	/** Which categories to show in the details panel. If both this and the Properties To Show whitelist are empty, all properties will show. */
+	UPROPERTY(EditAnywhere, Category = "View")
+	TArray<FName> CategoriesToShow;
+
+	/** Which properties to show in the details panel. If both this and the Categories To Show whitelist are empty, all properties will show. */
+	UPROPERTY(EditAnywhere, Category = "View")
+	TArray<FName> PropertiesToShow;
+
+	//~ UPropertyViewBase interface
+private:
+	virtual void BuildContentWidget() override;
+	virtual void OnObjectChanged() override;
+	//~ End of UPropertyViewBase interface
+
+public:
+	//~ UWidget interface
+	virtual void ReleaseSlateResources(bool bReleaseChildren) override;
+	// End of UWidget interface
+
+	//~ FNotifyHook interface
+	virtual void NotifyPostChange(const FPropertyChangedEvent& PropertyChangedEvent, UProperty* PropertyThatChanged) override;
+	//~ End of FNotifyHook interface
+
+	//~ UObject interface
+	virtual void PostEditChangeProperty(FPropertyChangedEvent& PropertyChangedEvent) override;
+	//~ End of UObject interface
+
+protected:
+	void ToggleWhitelistedProperties();
+	bool GetIsPropertyVisible(const FPropertyAndParent& PropertyAndParent) const;
+
+private:
+	TSharedPtr<IDetailsView> DetailViewWidget;
+
+	/** Showing properties in this details panel works by whitelisting categories and properties. This flag enables you to show all properties without needing to specify. */
+	UPROPERTY()
+	bool bShowOnlyWhitelisted = true;
+};
diff --git a/Engine/Source/Editor/UMGEditor/Public/Components/PropertyViewBase.h b/Engine/Source/Editor/UMGEditor/Public/Components/PropertyViewBase.h
new file mode 100644
--- /dev/null
+++ b/Engine/Source/Editor/UMGEditor/Public/Components/PropertyViewBase.h
@@ -0,0 +1,84 @@
+// Copyright 1998-2019 Epic Games, Inc. All Rights Reserved.
+
+#pragma once
+
+#include "CoreMinimal.h"
+#include "Components/Widget.h"
+#include "Widgets/SWidget.h"
+#include "UObject/LazyObjectPtr.h"
+#include "UObject/SoftObjectPath.h"
+#include "PropertyViewBase.generated.h"
+
+class SBorder;
+
+/** Sets a delegate called when the property value changes */
+DECLARE_DYNAMIC_MULTICAST_DELEGATE_OneParam(FOnPropertyValueChanged, FName, PropertyName);
+
+
+/**
+ * Base of property view allows you to display the value of an object properties.
+ */
+UCLASS(Abstract)
+class UMGEDITOR_API UPropertyViewBase : public UWidget
+{
+	GENERATED_BODY()
+
+protected:
+	/** The object to view. */
+	UPROPERTY(meta = (DisplayName="Object"))
+	TLazyObjectPtr<UObject> LazyObject;
+
+	UPROPERTY()
+	FSoftObjectPath SoftObjectPath;
+
+	/** Load the object (if it's an asset) when the widget is created. */
+	UPROPERTY(EditAnywhere, AdvancedDisplay, Category = "View")
+	bool bAutoLoadAsset = true;
+
+	/** Sets a delegate called when the property value changes */
+	UPROPERTY(BlueprintAssignable, Category = "View|Event")
+	FOnPropertyValueChanged OnPropertyChanged;
+
+
+public:
+	UFUNCTION(BlueprintCallable, Category = "View")
+	UObject* GetObject() const;
+
+	UFUNCTION(BlueprintCallable, Category = "View")
+	void SetObject(UObject* NewObject);
+
+protected:
+	virtual void BuildContentWidget() PURE_VIRTUAL(UPropertyViewBase::BuildContentWidget, );
+	virtual void OnObjectChanged() { }
+
+	void AsynBuildContentWidget();
+	TSharedPtr<SBorder> GetDisplayWidget() const { return DisplayedWidget; }
+	void OnPropertyChangedBroadcast(FName PropertyName);
+
+private:
+	void InternalOnAssetLoaded(UObject* LoadedAsset);
+	void InternalOnMapChange(uint32);
+	void InternalPostLoadMapWithWorld(UWorld* LoadedWorld);
+
+	//~ UWidget interface
+public:
+	virtual void ReleaseSlateResources(bool bReleaseChildren) override;
+	virtual const FText GetPaletteCategory() override;
+
+protected:
+	virtual TSharedRef<SWidget> RebuildWidget() override;
+	//~ End of UWidget interface
+
+	//~ UObject interface
+public:
+	virtual void PostLoad() override;
+	virtual void PostEditChangeProperty(FPropertyChangedEvent& PropertyChangedEvent) override;
+	//~ End of UObject interface
+
+private:
+	TSharedPtr<SBorder> DisplayedWidget;
+	FDelegateHandle AssetLoadedHandle;
+	FDelegateHandle PostLoadMapHandle;
+	FDelegateHandle MapChangeHandle;
+	bool bIsAsyncBuildContentRequested = false;
+};
diff --git a/Engine/Source/Editor/UMGEditor/Public/Components/SinglePropertyView.h b/Engine/Source/Editor/UMGEditor/Public/Components/SinglePropertyView.h
new file mode 100644
--- /dev/null
+++ b/Engine/Source/Editor/UMGEditor/Public/Components/SinglePropertyView.h
@@ -0,0 +1,56 @@
+// Copyright 1998-2019 Epic Games, Inc. All Rights Reserved.
+
+#pragma once
+
+#include "Components/PropertyViewBase.h"
+#include "SinglePropertyView.generated.h"
+
+
+class ISinglePropertyView;
+
+/**
+ * The single property view allows you to display the value of an object's property.
+ */
+UCLASS()
+class UMGEDITOR_API USinglePropertyView : public UPropertyViewBase
+{
+	GENERATED_BODY()
+
+private:
+	/** The name of the property to display. */
+	UPROPERTY(EditAnywhere, Category = "View")
+	FName PropertyName;
+
+	/** Override for the property name that will be displayed instead of the property name. */
+	UPROPERTY(EditAnywhere, AdvancedDisplay, Category = "View")
+	FText NameOverride;
+
+public:
+	UFUNCTION(BlueprintCallable, Category = "View")
+	FName GetPropertyName() const;
+
+	UFUNCTION(BlueprintCallable, Category = "View")
+	void SetPropertyName(FName NewPropertyName);
+
+private:
+	void InternalSinglePropertyChanged();
+
+	//~ UPropertyViewBase interface
+private:
+	virtual void BuildContentWidget() override;
+	virtual void OnObjectChanged() override;
+	//~ End of UPropertyViewBase interface
+
+	//~ UWidget interface
+public:
+	virtual void ReleaseSlateResources(bool bReleaseChildren) override;
+	//~ End of UWidget interface
+
+	// UObject interface
+public:
+	virtual void PostEditChangeProperty(FPropertyChangedEvent& PropertyChangedEvent) override;
+	// End of UObject interface
+
+private:
+	TSharedPtr<ISinglePropertyView> SinglePropertyViewWidget;
+};
diff --git a/Engine/Source/Editor/UMGEditor/Public/WidgetBlueprint.h b/Engine/Source/Editor/UMGEditor/Public/WidgetBlueprint.h
--- a/Engine/Source/Editor/UMGEditor/Public/WidgetBlueprint.h
+++ b/Engine/Source/Editor/UMGEditor/Public/WidgetBlueprint.h
@@ -314,6 +314,9 @@ public:
 
 	bool ArePropertyBindingsAllowed() const;
 
+	/** Does the editor support widget from an editor package. */
+	virtual bool AllowEditorWidget() const { return false; }
+
 private:
 #if WITH_EDITOR
 	virtual void LoadModulesRequiredForCompilation() override;
diff --git a/Engine/Source/Runtime/UMG/Private/UMGStyle.cpp b/Engine/Source/Runtime/UMG/Private/UMGStyle.cpp
--- a/Engine/Source/Runtime/UMG/Private/UMGStyle.cpp
+++ b/Engine/Source/Runtime/UMG/Private/UMGStyle.cpp
@@ -86,6 +86,9 @@ TSharedRef< FSlateStyleSet > FUMGStyle::Create()
 
 	Style->Set("ClassIcon.UserWidget", new IMAGE_BRUSH(TEXT("UserWidget"), Icon16x16));
 
+	Style->Set("ClassIcon.DetailsView", new IMAGE_BRUSH(TEXT("icon_umg_DetailView_16x"), Icon16x16));
+	Style->Set("ClassIcon.SinglePropertyView", new IMAGE_BRUSH(TEXT("icon_umg_SinglePropertyView_16x"), Icon16x16));
+
 	Style->Set("Animations.TabIcon", new IMAGE_BRUSH(TEXT("Animations_16x"), Icon16x16));
 	Style->Set("Designer.TabIcon", new IMAGE_BRUSH(TEXT("Designer_16x"), Icon16x16));
 	Style->Set("Palette.TabIcon", new IMAGE_BRUSH(TEXT("Palette_16x"), Icon16x16));
