Fix for async loading the property view object after the property viewer has been destroyed
#rnx

diff --git a/Engine/Source/Editor/UMGEditor/Private/Components/PropertyViewBase.cpp b/Engine/Source/Editor/UMGEditor/Private/Components/PropertyViewBase.cpp
--- a/Engine/Source/Editor/UMGEditor/Private/Components/PropertyViewBase.cpp
+++ b/Engine/Source/Editor/UMGEditor/Private/Components/PropertyViewBase.cpp
@@ -126,9 +126,12 @@ void UPropertyViewBase::PostLoad()
 	{
 		AsyncTask(ENamedThreads::GameThread, [this]()
 		{
-			LazyObject = SoftObjectPath.TryLoad();
-			BuildContentWidget();
-			bIsAsyncBuildContentRequested = false;
+			if (!HasAnyFlags(RF_BeginDestroyed))
+			{
+				LazyObject = SoftObjectPath.TryLoad();
+				BuildContentWidget();
+				bIsAsyncBuildContentRequested = false;
+			}
 		});
 	}
 }
