Fix unresolved symbols on PropertyPathCache implicit methods linking UMGEditor module on clang win64.

diff --git a/Engine/Source/Editor/UMGEditor/UMGEditor.Build.cs b/Engine/Source/Editor/UMGEditor/UMGEditor.Build.cs
--- a/Engine/Source/Editor/UMGEditor/UMGEditor.Build.cs
+++ b/Engine/Source/Editor/UMGEditor/UMGEditor.Build.cs
@@ -40,7 +40,6 @@ public class UMGEditor : ModuleRules
 				"CoreUObject",
 				"ApplicationCore",
 				"InputCore",
-				"Slate",
 				"Engine",
 				"AssetTools",
 				"UnrealEd", // for FAssetEditorManager
@@ -49,6 +48,7 @@ public class UMGEditor : ModuleRules
 				"BlueprintGraph",
 				"GraphEditor",
 				"Kismet",  // for FWorkflowCentricApplication
+				"PropertyPath",
 				"PropertyEditor",
 				"UMG",
 				"EditorStyle",
