Integrating 6533688
UMG Editor: Change the copy/paste in the Hierarchy view to keep the current name of object and add _X at the end instead of generating a new name using the Widget Type everytime there is a name clash.

diff --git a/Engine/Source/Editor/UMGEditor/Private/WidgetBlueprintEditorUtils.cpp b/Engine/Source/Editor/UMGEditor/Private/WidgetBlueprintEditorUtils.cpp
--- a/Engine/Source/Editor/UMGEditor/Private/WidgetBlueprintEditorUtils.cpp
+++ b/Engine/Source/Editor/UMGEditor/Private/WidgetBlueprintEditorUtils.cpp
@@ -1349,15 +1349,15 @@ void FWidgetBlueprintEditorUtils::ImportWidgetsFromText(UWidgetBlueprint* BP, co
 
 		// If there is an existing widget with the same name, rename the newly placed widget.
 		FString WidgetOldName = Widget->GetName();
-		if ( FindObject<UObject>(BP->WidgetTree, *WidgetOldName) )
+		FString NewName = FindNextValidName(BP->WidgetTree, WidgetOldName);
+		if (NewName != WidgetOldName)
 		{
 			UWidgetSlotPair* SlotData = PastedExtraSlotData.FindRef(Widget->GetFName());
 			if ( SlotData )
 			{
 				PastedExtraSlotData.Remove(Widget->GetFName());
 			}
-
-			Widget->Rename(nullptr, BP->WidgetTree);
+			Widget->Rename(*NewName, BP->WidgetTree);
 
 			if (Widget->GetDisplayLabel().Equals(WidgetOldName))
 			{
@@ -1488,4 +1488,45 @@ bool FWidgetBlueprintEditorUtils::IsUsableWidgetClass(UClass* WidgetClass)
 	return false;
 }
 
+FString RemoveSuffixFromName(const FString OldName)
+{
+	int NameLen = OldName.Len();
+	int SuffixIndex = 0;
+	if (OldName.FindLastChar('_', SuffixIndex))
+	{
+		NameLen = SuffixIndex;
+		for (int32 i = SuffixIndex + 1; i < OldName.Len(); ++i)
+		{
+			const TCHAR& C = OldName[i];
+			const bool bGoodChar = ((C >= '0') && (C <= '9'));
+			if (!bGoodChar)
+			{
+				return OldName;
+			}
+		}
+	}
+	return FString(NameLen, *OldName);
+}
+
+FString FWidgetBlueprintEditorUtils::FindNextValidName(UWidgetTree* WidgetTree, const FString& Name)
+{
+	// If the name of the widget is not already used, we use it.	
+	if (FindObject<UObject>(WidgetTree, *Name))
+	{
+		// If the name is already used, we will suffix it with '_X'
+		FString NameWithoutSuffix = RemoveSuffixFromName(Name);
+		FString NewName = NameWithoutSuffix;
+
+		int32 Postfix = 0;
+		while (FindObject<UObject>(WidgetTree, *NewName))
+		{
+			++Postfix;
+			NewName = FString::Printf(TEXT("%s_%d"), *NameWithoutSuffix, Postfix);
+		}
+
+		return NewName;
+	}
+	return Name;
+}
+
 #undef LOCTEXT_NAMESPACE
diff --git a/Engine/Source/Editor/UMGEditor/Private/WidgetBlueprintEditorUtils.h b/Engine/Source/Editor/UMGEditor/Private/WidgetBlueprintEditorUtils.h
--- a/Engine/Source/Editor/UMGEditor/Private/WidgetBlueprintEditorUtils.h
+++ b/Engine/Source/Editor/UMGEditor/Private/WidgetBlueprintEditorUtils.h
@@ -86,4 +86,6 @@ private:
 	static void ReplaceWidgetWithChildren(TSharedRef<FWidgetBlueprintEditor> BlueprintEditor, UWidgetBlueprint* BP, FWidgetReference Widget);
 
 	static void ReplaceWidgets(TSharedRef<FWidgetBlueprintEditor> BlueprintEditor, UWidgetBlueprint* BP, TSet<FWidgetReference> Widgets, UClass* WidgetClass);
+
+	static FString FindNextValidName(UWidgetTree* WidgetTree, const FString& Name);
 };
