Integrating 6537170
UMG EditorCorrection to the replace widget function to maintain the name when it's not a generic name.

diff --git a/Engine/Source/Editor/UMGEditor/Private/WidgetBlueprintEditorUtils.cpp b/Engine/Source/Editor/UMGEditor/Private/WidgetBlueprintEditorUtils.cpp
--- a/Engine/Source/Editor/UMGEditor/Private/WidgetBlueprintEditorUtils.cpp
+++ b/Engine/Source/Editor/UMGEditor/Private/WidgetBlueprintEditorUtils.cpp
@@ -781,6 +781,17 @@ void FWidgetBlueprintEditorUtils::ReplaceWidgetWithSelectedTemplate(TSharedRef<F
 		CurrentParent->SetFlags(RF_Transactional);
 		CurrentParent->Modify();
 		CurrentParent->ReplaceChild(ThisWidget, NewReplacementWidget);
+
+		FString ReplaceName = ThisWidget->GetName();
+		// Rename the removed widget to the transient package so that it doesn't conflict with future widgets sharing the same name.
+		ThisWidget->Rename(nullptr, nullptr);
+
+		// Rename the new Widget to maintain the current name if it's not a generic name
+		if (!IsGenericName(ReplaceName, ThisWidget->GetClass()))
+		{
+			ReplaceName = FindNextValidName(BP->WidgetTree, ReplaceName);
+			NewReplacementWidget->Rename(*ReplaceName, BP->WidgetTree);
+		}
 	}
 	else if (ThisWidget == BP->WidgetTree->RootWidget)
 	{
@@ -936,8 +947,17 @@ void FWidgetBlueprintEditorUtils::ReplaceWidgets(TSharedRef<FWidgetBlueprintEdit
 				}
 			}
 
+			FString ReplaceName = Item.GetTemplate()->GetName();
 			// Rename the removed widget to the transient package so that it doesn't conflict with future widgets sharing the same name.
 			Item.GetTemplate()->Rename(nullptr, nullptr);
+
+			// Rename the new Widget to maintain the current name if it's not a generic name
+			if (!IsGenericName(ReplaceName, Item.GetTemplate()->GetClass()))
+			{
+				ReplaceName = FindNextValidName(BP->WidgetTree, ReplaceName);
+				NewReplacementWidget->Rename(*ReplaceName, BP->WidgetTree);
+			}
+
 	}
 
 	FBlueprintEditorUtils::MarkBlueprintAsStructurallyModified(BP);
@@ -1529,4 +1549,10 @@ FString FWidgetBlueprintEditorUtils::FindNextValidName(UWidgetTree* WidgetTree,
 	return Name;
 }
 
+bool FWidgetBlueprintEditorUtils::IsGenericName(const FString& Name, const UClass* WidgetClass)
+{
+	FString NewName = RemoveSuffixFromName(Name);
+	return (NewName == WidgetClass->GetName());
+}
+
 #undef LOCTEXT_NAMESPACE
diff --git a/Engine/Source/Editor/UMGEditor/Private/WidgetBlueprintEditorUtils.h b/Engine/Source/Editor/UMGEditor/Private/WidgetBlueprintEditorUtils.h
--- a/Engine/Source/Editor/UMGEditor/Private/WidgetBlueprintEditorUtils.h
+++ b/Engine/Source/Editor/UMGEditor/Private/WidgetBlueprintEditorUtils.h
@@ -88,4 +88,6 @@ private:
 	static void ReplaceWidgets(TSharedRef<FWidgetBlueprintEditor> BlueprintEditor, UWidgetBlueprint* BP, TSet<FWidgetReference> Widgets, UClass* WidgetClass);
 
 	static FString FindNextValidName(UWidgetTree* WidgetTree, const FString& Name);
+
+	static bool IsGenericName(const FString& Name, const UClass* WidgetClass);
 };
