Remove use of FAssetData to avoid side effects that can result in CDO creation

diff --git a/Engine/Source/Editor/UMGEditor/Private/Settings/UMGEditorProjectSettings.cpp b/Engine/Source/Editor/UMGEditor/Private/Settings/UMGEditorProjectSettings.cpp
--- a/Engine/Source/Editor/UMGEditor/Private/Settings/UMGEditorProjectSettings.cpp
+++ b/Engine/Source/Editor/UMGEditor/Private/Settings/UMGEditorProjectSettings.cpp
@@ -1,9 +1,9 @@
 // Copyright 1998-2018 Epic Games, Inc. All Rights Reserved.
 
 #include "UMGEditorProjectSettings.h"
-#include "AssetData.h"
 #include "WidgetBlueprint.h"
 #include "WidgetCompilerRule.h"
+#include "UObject/Package.h"
 #include "UObject/UObjectIterator.h"
 
 UUMGEditorProjectSettings::UUMGEditorProjectSettings()
@@ -81,9 +81,9 @@ TArray<UWidgetCompilerRule*> UUMGEditorProjectSettings::CompilerOption_Rules(con
 
 void UUMGEditorProjectSettings::GetCompilerOptionsForWidget(const UWidgetBlueprint* WidgetBlueprint, TFunctionRef<bool(const FWidgetCompilerOptions&)> Operator) const
 {
-	FAssetData WidgetBlueprintAsset(WidgetBlueprint);
-	FString AssetPath = WidgetBlueprintAsset.PackageName.ToString();
-
+	FString AssetPath = WidgetBlueprint->GetOutermost()->GetName();
+	FSoftObjectPath SoftObjectPath = WidgetBlueprint->GetPathName();
+	
 	for (int32 DirectoryIndex = DirectoryCompilerOptions.Num() - 1; DirectoryIndex >= 0; DirectoryIndex--)
 	{
 		const FDirectoryWidgetCompilerOptions& CompilerOptions = DirectoryCompilerOptions[DirectoryIndex];
@@ -93,8 +93,8 @@ void UUMGEditorProjectSettings::GetCompilerOptionsForWidget(const UWidgetBluepri
 		{
 			if (AssetPath.StartsWith(DirectoryPath))
 			{
-				const bool bIgnoreWidget = CompilerOptions.IgnoredWidgets.ContainsByPredicate([&WidgetBlueprintAsset](const TSoftObjectPtr<UWidgetBlueprint>& IgnoredWidget) {
-					return IgnoredWidget.ToSoftObjectPath() == WidgetBlueprintAsset.ToSoftObjectPath();
+				const bool bIgnoreWidget = CompilerOptions.IgnoredWidgets.ContainsByPredicate([&SoftObjectPath](const TSoftObjectPtr<UWidgetBlueprint>& IgnoredWidget) {
+					return IgnoredWidget.ToSoftObjectPath() == SoftObjectPath;
 				});
 
 				if (bIgnoreWidget)
