Fix for trying to move widget to a new parent on every tick, renaming widget on every tick during a drag and drop

diff --git a/Engine/Source/Editor/UMGEditor/Private/Designer/SDesignerView.cpp b/Engine/Source/Editor/UMGEditor/Private/Designer/SDesignerView.cpp
--- a/Engine/Source/Editor/UMGEditor/Private/Designer/SDesignerView.cpp
+++ b/Engine/Source/Editor/UMGEditor/Private/Designer/SDesignerView.cpp
@@ -1558,7 +1558,10 @@ void SDesignerView::ClearDropPreviews()
 
 		// Since the widget has been removed from the widget tree, move it into the transient package. Otherwise,
 		// it will remain outered to the widget tree and end up as a property in the BP class layout as a result.
-		DropPreview.Widget->Rename(nullptr, GetTransientPackage());
+		if (DropPreview.Widget->GetOutermost() != GetTransientPackage())
+		{
+			DropPreview.Widget->Rename(nullptr, GetTransientPackage());
+		}
 	}
 	DropPreviews.Empty();
 }
@@ -2787,7 +2790,10 @@ void SDesignerView::ProcessDropAndAddWidget(const FGeometry& MyGeometry, const F
 		// Move the remaining widgets into the transient package. Otherwise, they will remain outered to the WidgetTree and end up as properties in the BP class layout as a result.
 		for (UWidget* Widget : DragDropPreviewWidgets)
 		{
-			Widget->Rename(nullptr, GetTransientPackage());
+			if (Widget->GetOutermost() != GetTransientPackage())
+			{
+				Widget->Rename(nullptr, GetTransientPackage());
+			}
 		}
 
 		// If we had preview widgets, we know that we can not be performing a selected widget drag/drop operation. Bail.
@@ -2841,24 +2847,24 @@ void SDesignerView::ProcessDropAndAddWidget(const FGeometry& MyGeometry, const F
 			BlueprintEditor.Pin()->SetHoveredWidget(TargetReference);
 
 			// If the widget being hovered over is a panel, attempt to place it into that panel.
-			if (Target && Target->IsA(UPanelWidget::StaticClass()))
+			if (Target && Target->IsA(UPanelWidget::StaticClass()) && !DraggedWidget.bStayingInParent)
 			{
 				bWidgetMoved = true;
 				UPanelWidget* NewParent = Cast<UPanelWidget>(Target);
 
 				UWidget* Widget = bIsPreview ? DraggedWidget.Preview : DraggedWidget.Template;
-
+				UPanelWidget* ParentWidget = bIsPreview ? Cast<UPanelWidget>(DraggedWidget.ParentWidget.GetPreview()) : Cast<UPanelWidget>(DraggedWidget.ParentWidget.GetTemplate());
 				if (ensure(Widget))
 				{
-					bool bIsChangingParent = Widget->GetParent() != NewParent;
+					bool bIsChangingParent = ParentWidget != NewParent;
 					UBlueprint* OriginalBP = nullptr;
 
-					check(Widget->GetParent() != nullptr || bIsChangingParent);
-
-					// If this isn't a preview operation we need to modify a few things to properly undo the operation.
-					if (!bIsPreview)
+					if (bIsChangingParent)
 					{
-						if (bIsChangingParent)
+						check(ParentWidget != nullptr);
+
+						// If this isn't a preview operation we need to modify a few things to properly undo the operation.
+						if (!bIsPreview)
 						{
 							NewParent->SetFlags(RF_Transactional);
 							NewParent->Modify();
@@ -2868,26 +2874,17 @@ void SDesignerView::ProcessDropAndAddWidget(const FGeometry& MyGeometry, const F
 
 							// If the Widget is changing parents, there's a chance it might be moving to a different WidgetTree as well.
 							UWidgetTree* OriginalWidgetTree = Cast<UWidgetTree>(Widget->GetOuter());
-							
+
 							if (UWidgetTree::TryMoveWidgetToNewTree(Widget, BP->WidgetTree))
 							{
 								// The Widget likely originated from a different blueprint, so get what blueprint it was originally a part of.
 								OriginalBP = OriginalWidgetTree ? OriginalWidgetTree->GetTypedOuter<UBlueprint>() : nullptr;
 							}
-						}
 
-						Widget->Modify();
-					}
-
-					if (Widget->GetParent() && bIsChangingParent)
-					{
-						if (!bIsPreview)
-						{
-							Widget->GetParent()->Modify();
+							Widget->Modify();
+							ParentWidget->Modify();
 						}
 
-						Widget->GetParent()->RemoveChild(Widget);
-
 						// The Widget originated from a different blueprint, so mark it as modified.
 						if (OriginalBP && OriginalBP != BP)
 						{
@@ -2927,17 +2924,11 @@ void SDesignerView::ProcessDropAndAddWidget(const FGeometry& MyGeometry, const F
 
 					FVector2D LocalPosition = WidgetUnderCursorGeometry.AbsoluteToLocal(ScreenSpacePosition);
 
-					UPanelSlot* Slot;
+					UPanelSlot* Slot = nullptr;
 					if (bIsChangingParent)
 					{
 						Slot = NewParent->AddChild(Widget);
 					}
-					else
-					{
-						check(Widget->GetParent()->GetChildIndex(Widget) != INDEX_NONE);
-
-						Slot = Widget->GetParent()->GetSlots()[Widget->GetParent()->GetChildIndex(Widget)];
-					}
 
 					if (Slot != nullptr)
 					{
