Fix for not being able to move widgets within a canvas

diff --git a/Engine/Source/Editor/UMGEditor/Private/Designer/SDesignerView.cpp b/Engine/Source/Editor/UMGEditor/Private/Designer/SDesignerView.cpp
--- a/Engine/Source/Editor/UMGEditor/Private/Designer/SDesignerView.cpp
+++ b/Engine/Source/Editor/UMGEditor/Private/Designer/SDesignerView.cpp
@@ -2847,7 +2847,7 @@ void SDesignerView::ProcessDropAndAddWidget(const FGeometry& MyGeometry, const F
 			BlueprintEditor.Pin()->SetHoveredWidget(TargetReference);
 
 			// If the widget being hovered over is a panel, attempt to place it into that panel.
-			if (Target && Target->IsA(UPanelWidget::StaticClass()) && !DraggedWidget.bStayingInParent)
+			if (Target && Target->IsA(UPanelWidget::StaticClass()))
 			{
 				bWidgetMoved = true;
 				UPanelWidget* NewParent = Cast<UPanelWidget>(Target);
@@ -2881,8 +2881,12 @@ void SDesignerView::ProcessDropAndAddWidget(const FGeometry& MyGeometry, const F
 								OriginalBP = OriginalWidgetTree ? OriginalWidgetTree->GetTypedOuter<UBlueprint>() : nullptr;
 							}
 
+							Widget->SetFlags(RF_Transactional);
 							Widget->Modify();
+						
+							ParentWidget->SetFlags(RF_Transactional);
 							ParentWidget->Modify();
+						
 						}
 
 						// The Widget originated from a different blueprint, so mark it as modified.
@@ -2891,6 +2895,14 @@ void SDesignerView::ProcessDropAndAddWidget(const FGeometry& MyGeometry, const F
 							FBlueprintEditorUtils::MarkBlueprintAsStructurallyModified(OriginalBP);
 						}
 					}
+					else if (!bIsPreview)
+					{
+						Widget->SetFlags(RF_Transactional);
+						Widget->Modify();
+
+						ParentWidget->SetFlags(RF_Transactional);
+						ParentWidget->Modify();
+					}
 
 					FVector2D ScreenSpacePosition = DragDropEvent.GetScreenSpacePosition();
 
@@ -2929,6 +2941,10 @@ void SDesignerView::ProcessDropAndAddWidget(const FGeometry& MyGeometry, const F
 					{
 						Slot = NewParent->AddChild(Widget);
 					}
+					else
+					{
+						Slot = ParentWidget->AddChild(Widget);
+					}
 
 					if (Slot != nullptr)
 					{
