v1 of Editor Utility Widgets

diff --git a/Engine/Source/Editor/Blutility/Blutility.Build.cs b/Engine/Source/Editor/Blutility/Blutility.Build.cs
--- a/Engine/Source/Editor/Blutility/Blutility.Build.cs
+++ b/Engine/Source/Editor/Blutility/Blutility.Build.cs
@@ -28,8 +28,11 @@ public class Blutility : ModuleRules
 				"ClassViewer",
 				"CollectionManager",
                 "PropertyEditor",
-                "BlueprintGraph"
-			}
+                "BlueprintGraph",
+				"UMG",
+                "UMGEditor",
+                "KismetCompiler"
+            }
 			);
 	}
 }
diff --git a/Engine/Source/Editor/Blutility/Classes/EditorUtilityWidget.h b/Engine/Source/Editor/Blutility/Classes/EditorUtilityWidget.h
new file mode 100644
--- /dev/null
+++ b/Engine/Source/Editor/Blutility/Classes/EditorUtilityWidget.h
@@ -0,0 +1,54 @@
+// Copyright 1998-2018 Epic Games, Inc. All Rights Reserved.
+
+/**
+ * Widget for editor utilities
+ */
+
+#pragma once
+
+#include "CoreMinimal.h"
+#include "UObject/ObjectMacros.h"
+#include "UObject/Object.h"
+#include "UObject/ScriptMacros.h"
+#include "Blueprint/UserWidget.h"
+#include "EditorUtilityWidget.generated.h"
+
+class AActor;
+class UEditorPerProjectUserSettings;
+
+UCLASS(Abstract, config = Editor)
+class BLUTILITY_API UEditorUtilityWidget : public UUserWidget
+{
+	GENERATED_UCLASS_BODY()
+
+public:
+	// The default action called when the blutility is invoked if bAutoRunDefaultAction=true (it is never called otherwise)
+	UFUNCTION(BlueprintImplementableEvent)
+	void OnDefaultActionClicked();
+
+	// Run the default action
+	void ExecuteDefaultAction();
+
+	bool ShouldAlwaysReregisterWithWindowsMenu() const
+	{
+		return bAlwaysReregisterWithWindowsMenu;
+	}
+
+	bool ShouldAutoRunDefaultAction() const
+	{
+		return bAutoRunDefaultAction;
+	}
+
+protected:
+	UPROPERTY(Category = Config, EditDefaultsOnly, BlueprintReadWrite, AssetRegistrySearchable)
+	FString HelpText;
+
+	// Should this widget always be re-added to the windows menu once it's opened
+	UPROPERTY(Config, Category = Settings, EditDefaultsOnly)
+	bool bAlwaysReregisterWithWindowsMenu;
+
+	// Should this blueprint automatically run OnDefaultActionClicked, or should it open up a details panel to edit properties and/or offer multiple buttons
+	UPROPERTY(Category = Settings, EditDefaultsOnly, BlueprintReadOnly)
+	bool bAutoRunDefaultAction;
+
+};
diff --git a/Engine/Source/Editor/Blutility/Classes/EditorUtilityWidgetBlueprint.h b/Engine/Source/Editor/Blutility/Classes/EditorUtilityWidgetBlueprint.h
new file mode 100644
--- /dev/null
+++ b/Engine/Source/Editor/Blutility/Classes/EditorUtilityWidgetBlueprint.h
@@ -0,0 +1,47 @@
+// Copyright 1998-2018 Epic Games, Inc. All Rights Reserved.
+
+/**
+ * Widget for editor utilities
+ */
+
+#pragma once
+
+#include "CoreMinimal.h"
+#include "UObject/ObjectMacros.h"
+#include "UObject/Object.h"
+#include "WidgetBlueprint.h"
+#include "Widgets/Docking/SDockTab.h"
+#include "EditorUtilityWidgetBlueprint.generated.h"
+
+class UBlueprint;
+class UEditorUtilityWidget;
+
+UCLASS()
+class UEditorUtilityWidgetBlueprint : public UWidgetBlueprint
+{
+	GENERATED_UCLASS_BODY()
+
+public:
+	TSharedRef<SDockTab> SpawnEditorUITab(const FSpawnTabArgs& SpawnTabArgs);
+
+	/** Creates the slate widget from the UMG widget */
+	TSharedRef<SWidget> CreateUtilityWidget();
+
+	/** Recreate the tab's content on recompile */
+	void RegenerateCreatedTab(UBlueprint* RecompiledBlueprint);
+	
+	void UpdateRespawnListIfNeeded(TSharedRef<SDockTab> TabBeingClosed);
+
+	// UBlueprint interface
+	virtual void GetReparentingRules(TSet< const UClass* >& AllowedChildrenOfClasses, TSet< const UClass* >& DisallowedChildrenOfClasses) const override;
+
+	UEditorUtilityWidget* GetCreatedWidget() const
+	{
+		return CreatedUMGWidget;
+	}
+
+private:
+	TWeakPtr<SDockTab> CreatedTab;
+
+	UEditorUtilityWidget* CreatedUMGWidget;
+};
diff --git a/Engine/Source/Editor/Blutility/Classes/EditorUtilityWidgetBlueprintFactory.h b/Engine/Source/Editor/Blutility/Classes/EditorUtilityWidgetBlueprintFactory.h
new file mode 100644
--- /dev/null
+++ b/Engine/Source/Editor/Blutility/Classes/EditorUtilityWidgetBlueprintFactory.h
@@ -0,0 +1,38 @@
+// Copyright 1998-2018 Epic Games, Inc. All Rights Reserved.
+
+#pragma once
+
+/**
+ * Factory for utility widget blueprints
+ */
+
+#include "CoreMinimal.h"
+#include "UObject/ObjectMacros.h"
+#include "Templates/SubclassOf.h"
+#include "Factories/Factory.h"
+#include "EditorUtilityWidgetBlueprintFactory.generated.h"
+
+class SWindow;
+
+UCLASS(HideCategories = Object, MinimalAPI)
+class UEditorUtilityWidgetBlueprintFactory : public UFactory
+{
+	GENERATED_UCLASS_BODY()
+
+	// The type of blueprint that will be created
+	UPROPERTY(EditAnywhere, Category = WidgetBlueprintFactory)
+	TEnumAsByte<enum EBlueprintType> BlueprintType;
+
+	// The parent class of the created blueprint
+	UPROPERTY(EditAnywhere, Category = WidgetBlueprintFactory, meta = (AllowAbstract = ""))
+	TSubclassOf<class UUserWidget> ParentClass;
+
+	// UFactory interface
+	virtual bool ConfigureProperties() override;
+	virtual bool ShouldShowInNewMenu() const override { return true; }
+	virtual UObject* FactoryCreateNew(UClass* Class, UObject* InParent, FName Name, EObjectFlags Flags, UObject* Context, FFeedbackContext* Warn) override;
+	virtual bool CanCreateNew() const override;
+	
+	// End of UFactory interface
+
+};
diff --git a/Engine/Source/Editor/Blutility/Private/AssetTypeActions_EditorUtilityBlueprint.cpp b/Engine/Source/Editor/Blutility/Private/AssetTypeActions_EditorUtilityBlueprint.cpp
--- a/Engine/Source/Editor/Blutility/Private/AssetTypeActions_EditorUtilityBlueprint.cpp
+++ b/Engine/Source/Editor/Blutility/Private/AssetTypeActions_EditorUtilityBlueprint.cpp
@@ -8,9 +8,10 @@
 #include "GlobalEditorUtilityBase.h"
 #include "Kismet2/KismetEditorUtilities.h"
 #include "BlueprintEditorModule.h"
-#include "GlobalBlutilityDialog.h"
+#include "EditorUtilityDialog.h"
 #include "IContentBrowserSingleton.h"
 #include "ContentBrowserModule.h"
+#include "IBlutilityModule.h"
 
 #define LOCTEXT_NAMESPACE "AssetTypeActions"
 
@@ -119,8 +120,8 @@ void FAssetTypeActions_EditorUtilityBlueprint::OpenAssetEditor(const TArray<UObj
 				}
 				else
 				{
-					// This one needs settings or has multiple actions to execute, so invoke the blutility dialog
-					TSharedRef<FGlobalBlutilityDialog> NewBlutilityDialog(new FGlobalBlutilityDialog());
+					// This one needs settings or has multiple actions to execute, so invoke the dialog
+					TSharedRef<FGlobalEditorUtilityDialog> NewBlutilityDialog(new FGlobalEditorUtilityDialog());
 					NewBlutilityDialog->InitBlutilityDialog(Mode, EditWithinLevelEditor, Blueprint);
 				}
 			}
@@ -136,7 +137,8 @@ void FAssetTypeActions_EditorUtilityBlueprint::OpenAssetEditor(const TArray<UObj
 
 uint32 FAssetTypeActions_EditorUtilityBlueprint::GetCategories()
 {
-	return EAssetTypeCategories::Blueprint;
+	IBlutilityModule* BlutilityModule = FModuleManager::GetModulePtr<IBlutilityModule>("Blutility");
+	return BlutilityModule->GetAssetCategory();
 }
 
 void FAssetTypeActions_EditorUtilityBlueprint::ExecuteEdit(FWeakBlueprintPointerArray Objects)
diff --git a/Engine/Source/Editor/Blutility/Private/AssetTypeActions_EditorUtilityWidgetBlueprint.cpp b/Engine/Source/Editor/Blutility/Private/AssetTypeActions_EditorUtilityWidgetBlueprint.cpp
new file mode 100644
--- /dev/null
+++ b/Engine/Source/Editor/Blutility/Private/AssetTypeActions_EditorUtilityWidgetBlueprint.cpp
@@ -0,0 +1,133 @@
+// Copyright 1998-2018 Epic Games, Inc. All Rights Reserved.
+
+#include "AssetTypeActions_EditorUtilityWidgetBlueprint.h"
+#include "Framework/MultiBox/MultiBoxBuilder.h"
+#include "Misc/PackageName.h"
+#include "Misc/MessageDialog.h"
+#include "EditorUtilityBlueprintFactory.h"
+#include "Kismet2/KismetEditorUtilities.h"
+#include "BlueprintEditorModule.h"
+#include "EditorUtilityDialog.h"
+#include "IContentBrowserSingleton.h"
+#include "ContentBrowserModule.h"
+#include "EditorUtilityWidgetBlueprint.h"
+#include "EditorUtilityWidget.h"
+#include "WidgetBlueprintEditor.h"
+#include "LevelEditor.h"
+#include "Editor.h"
+#include "Widgets/Docking/SDockTab.h"
+#include "Framework/Docking/TabManager.h"
+#include "IBlutilityModule.h"
+
+#define LOCTEXT_NAMESPACE "AssetTypeActions"
+
+/////////////////////////////////////////////////////
+// FAssetTypeActions_EditorUtilityWidget
+
+FText FAssetTypeActions_EditorUtilityWidgetBlueprint::GetName() const
+{
+	return LOCTEXT("AssetTypeActions_EditorUtilityWidget", "Editor Widget");
+}
+
+FColor FAssetTypeActions_EditorUtilityWidgetBlueprint::GetTypeColor() const
+{
+	return FColor(0, 169, 255);
+}
+
+UClass* FAssetTypeActions_EditorUtilityWidgetBlueprint::GetSupportedClass() const
+{
+	return UEditorUtilityWidgetBlueprint::StaticClass();
+}
+
+bool FAssetTypeActions_EditorUtilityWidgetBlueprint::HasActions(const TArray<UObject*>& InObjects) const
+{
+	return true;
+}
+
+void FAssetTypeActions_EditorUtilityWidgetBlueprint::GetActions(const TArray<UObject*>& InObjects, FMenuBuilder& MenuBuilder)
+{
+	auto Blueprints = GetTypedWeakObjectPtrs<UWidgetBlueprint>(InObjects);
+
+	MenuBuilder.AddMenuEntry(
+		LOCTEXT("Blutility_Edit", "Edit Blueprint"),
+		LOCTEXT("Blutility_EditTooltip", "Opens the selected blueprints in the full blueprint editor."),
+		FSlateIcon(),
+		FUIAction(
+			FExecuteAction::CreateSP(this, &FAssetTypeActions_EditorUtilityWidgetBlueprint::ExecuteEdit, Blueprints),
+			FCanExecuteAction()
+		)
+	);
+
+}
+
+void FAssetTypeActions_EditorUtilityWidgetBlueprint::OpenAssetEditor(const TArray<UObject*>& InObjects, TSharedPtr<class IToolkitHost> EditWithinLevelEditor)
+{
+	EToolkitMode::Type Mode = EditWithinLevelEditor.IsValid() ? EToolkitMode::WorldCentric : EToolkitMode::Standalone;
+
+	for (auto ObjIt = InObjects.CreateConstIterator(); ObjIt; ++ObjIt)
+	{
+		if (UWidgetBlueprint* Blueprint = Cast<UWidgetBlueprint>(*ObjIt))
+		{
+			if (Blueprint->GeneratedClass->IsChildOf(UEditorUtilityWidget::StaticClass()))
+			{
+				const UEditorUtilityWidget* CDO = Blueprint->GeneratedClass->GetDefaultObject<UEditorUtilityWidget>();
+				if (CDO->ShouldAutoRunDefaultAction())
+				{
+					// This is an instant-run blueprint, just execute it
+					UEditorUtilityWidget* Instance = NewObject<UEditorUtilityWidget>(GetTransientPackage(), Blueprint->GeneratedClass);
+					Instance->ExecuteDefaultAction();
+				}
+				else
+				{
+					FName RegistrationName = FName(*CDO->GetPathName());
+					FText DisplayName = FText::FromString(Blueprint->GetName());
+					FLevelEditorModule& LevelEditorModule = FModuleManager::GetModuleChecked<FLevelEditorModule>(TEXT("LevelEditor"));
+					TSharedPtr<FTabManager> LevelEditorTabManager = LevelEditorModule.GetLevelEditorTabManager();
+					if (!LevelEditorTabManager->CanSpawnTab(RegistrationName))
+					{
+						IBlutilityModule* BlutilityModule = FModuleManager::GetModulePtr<IBlutilityModule>("Blutility");
+						UEditorUtilityWidgetBlueprint* WidgetBlueprint = Cast<UEditorUtilityWidgetBlueprint>(Blueprint);
+
+						LevelEditorTabManager->RegisterTabSpawner(RegistrationName, FOnSpawnTab::CreateUObject(WidgetBlueprint, &UEditorUtilityWidgetBlueprint::SpawnEditorUITab))
+							.SetDisplayName(DisplayName)
+							.SetGroup(BlutilityModule->GetMenuGroup().ToSharedRef());
+						BlutilityModule->AddLoadedScriptUI(WidgetBlueprint);
+					}
+					TSharedRef<SDockTab> NewDockTab = LevelEditorTabManager->InvokeTab(RegistrationName);
+				
+				}
+			}
+		}
+	}
+}
+
+uint32 FAssetTypeActions_EditorUtilityWidgetBlueprint::GetCategories()
+{
+	IBlutilityModule* BlutilityModule = FModuleManager::GetModulePtr<IBlutilityModule>("Blutility");
+	return BlutilityModule->GetAssetCategory();
+}
+
+void FAssetTypeActions_EditorUtilityWidgetBlueprint::ExecuteEdit(FWeakBlueprintPointerArray Objects)
+{
+	for (auto ObjIt = Objects.CreateConstIterator(); ObjIt; ++ObjIt)
+	{
+		if (auto Object = (*ObjIt).Get())
+		{
+			auto Blueprint = Cast<UBlueprint>(*ObjIt);
+			if (Blueprint && Blueprint->SkeletonGeneratedClass && Blueprint->GeneratedClass)
+			{
+				TSharedRef< FWidgetBlueprintEditor > NewBlueprintEditor(new FWidgetBlueprintEditor());
+
+				TArray<UBlueprint*> Blueprints;
+				Blueprints.Add(Blueprint);
+				NewBlueprintEditor->InitWidgetBlueprintEditor(EToolkitMode::Standalone, nullptr, Blueprints, true);
+			}
+			else
+			{
+				FMessageDialog::Open(EAppMsgType::Ok, LOCTEXT("FailedToLoadWidgetBlueprint", "Editor Utility Widget could not be loaded because it derives from an invalid class.\nCheck to make sure the parent class for this blueprint hasn't been removed!"));
+			}
+		}
+	}
+}
+
+#undef LOCTEXT_NAMESPACE
diff --git a/Engine/Source/Editor/Blutility/Private/AssetTypeActions_EditorUtilityWidgetBlueprint.h b/Engine/Source/Editor/Blutility/Private/AssetTypeActions_EditorUtilityWidgetBlueprint.h
new file mode 100644
--- /dev/null
+++ b/Engine/Source/Editor/Blutility/Private/AssetTypeActions_EditorUtilityWidgetBlueprint.h
@@ -0,0 +1,31 @@
+// Copyright 1998-2018 Epic Games, Inc. All Rights Reserved.
+
+#pragma once
+
+#include "CoreMinimal.h"
+#include "EditorUtilityWidget.h"
+#include "Toolkits/IToolkitHost.h"
+#include "AssetTypeActions_Base.h"
+
+class FMenuBuilder;
+
+class FAssetTypeActions_EditorUtilityWidgetBlueprint : public FAssetTypeActions_Base
+{
+public:
+	// IAssetTypeActions interface
+	virtual FText GetName() const override;
+	virtual FColor GetTypeColor() const override;
+	virtual UClass* GetSupportedClass() const override;
+	virtual bool HasActions(const TArray<UObject*>& InObjects) const override;
+	virtual void GetActions(const TArray<UObject*>& InObjects, FMenuBuilder& MenuBuilder) override;
+	virtual void OpenAssetEditor(const TArray<UObject*>& InObjects, TSharedPtr<class IToolkitHost> EditWithinLevelEditor = TSharedPtr<IToolkitHost>()) override;
+	virtual uint32 GetCategories() override;
+	virtual bool CanLocalize() const override { return false; }
+	// End of IAssetTypeActions interface
+
+protected:
+	typedef TArray< TWeakObjectPtr<class UWidgetBlueprint> > FWeakBlueprintPointerArray;
+
+	void ExecuteEdit(FWeakBlueprintPointerArray Objects);
+
+};
diff --git a/Engine/Source/Editor/Blutility/Private/BlutilityModule.cpp b/Engine/Source/Editor/Blutility/Private/BlutilityModule.cpp
--- a/Engine/Source/Editor/Blutility/Private/BlutilityModule.cpp
+++ b/Engine/Source/Editor/Blutility/Private/BlutilityModule.cpp
@@ -19,6 +19,20 @@
 #include "Widgets/Docking/SDockTab.h"
 #include "BlutilityContentBrowserExtensions.h"
 #include "BlutilityLevelEditorExtensions.h"
+#include "AssetTypeActions_EditorUtilityWidgetBlueprint.h"
+#include "KismetCompiler.h"
+#include "EditorUtilityWidgetBlueprint.h"
+#include "ComponentReregisterContext.h"
+#include "KismetCompilerModule.h"
+#include "WidgetBlueprintCompiler.h"
+#include "UMGEditorModule.h"
+#include "EditorUtilityContext.h"
+#include "LevelEditor.h"
+#include "Editor.h"
+#include "UnrealEdMisc.h"
+
+
+#define LOCTEXT_NAMESPACE "AssetTypeActions"
 
 /////////////////////////////////////////////////////
 
@@ -31,18 +45,23 @@ namespace BlutilityModule
 // FBlutilityModule 
 
 // Blutility module implementation (private)
-class FBlutilityModule : public IBlutilityModule
+class FBlutilityModule : public IBlutilityModule, public FGCObject
 {
 public:
-	/** Asset type actions for MovieScene assets.  Cached here so that we can unregister it during shutdown. */
+	/** Asset type actions for blutility assets.  Cached here so that we can unregister it during shutdown. */
 	TSharedPtr<FAssetTypeActions_EditorUtilityBlueprint> EditorBlueprintAssetTypeActions;
+	TSharedPtr<FAssetTypeActions_EditorUtilityWidgetBlueprint> EditorWidgetBlueprintAssetTypeActions;
 
 public:
 	virtual void StartupModule() override
 	{
 		// Register the asset type
+		IAssetTools& AssetTools = FModuleManager::LoadModuleChecked<FAssetToolsModule>("AssetTools").Get();
+		EditorUtilityAssetCategory = AssetTools.RegisterAdvancedAssetCategory(FName(TEXT("EditorUtilities")), LOCTEXT("EditorUtilitiesAssetCategory", "Editor Utilities"));
 		EditorBlueprintAssetTypeActions = MakeShareable(new FAssetTypeActions_EditorUtilityBlueprint);
-		FModuleManager::LoadModuleChecked<FAssetToolsModule>("AssetTools").Get().RegisterAssetTypeActions(EditorBlueprintAssetTypeActions.ToSharedRef());
+		AssetTools.RegisterAssetTypeActions(EditorBlueprintAssetTypeActions.ToSharedRef());
+		EditorWidgetBlueprintAssetTypeActions = MakeShareable(new FAssetTypeActions_EditorUtilityWidgetBlueprint);
+		AssetTools.RegisterAssetTypeActions(EditorWidgetBlueprintAssetTypeActions.ToSharedRef());
 		
 		// Register the details customizer
 		FPropertyEditorModule& PropertyModule = FModuleManager::LoadModuleChecked<FPropertyEditorModule>("PropertyEditor");
@@ -53,9 +72,80 @@ public:
 		FGlobalTabmanager::Get()->RegisterTabSpawner(BlutilityModule::BlutilityShelfApp, FOnSpawnTab::CreateStatic(&SpawnBlutilityShelfTab))
 			.SetDisplayName(NSLOCTEXT("BlutilityShelf", "TabTitle", "Blutility Shelf"))
 			.SetGroup(WorkspaceMenu::GetMenuStructure().GetToolsCategory());
+		FKismetCompilerContext::RegisterCompilerForBP(UEditorUtilityWidgetBlueprint::StaticClass(), &UWidgetBlueprint::GetCompilerForWidgetBP);
+
+		// Register widget blueprint compiler we do this no matter what.
+		IUMGEditorModule& UMGEditorModule = FModuleManager::LoadModuleChecked<IUMGEditorModule>("UMGEditor");
+		IKismetCompilerInterface& KismetCompilerModule = FModuleManager::LoadModuleChecked<IKismetCompilerInterface>("KismetCompiler");
+		KismetCompilerModule.GetCompilers().Add(UMGEditorModule.GetRegisteredCompiler());
 
 		FBlutilityContentBrowserExtensions::InstallHooks();
 		FBlutilityLevelEditorExtensions::InstallHooks();
+
+		ScriptedEditorWidgetsGroup = WorkspaceMenu::GetMenuStructure().GetToolsCategory()->AddGroup(
+			LOCTEXT("WorkspaceMenu_EditorUtilityWidgetsGroup", "Editor Utility Widgets"),
+			LOCTEXT("ScriptedEditorWidgetsGroupTooltipText", "Custom editor UI created with Blueprints or Python."),
+			FSlateIcon(FEditorStyle::GetStyleSetName(), "WorkspaceMenu.AdditionalUI"),
+			true);
+
+		FLevelEditorModule& LevelEditorModule = FModuleManager::LoadModuleChecked<FLevelEditorModule>(TEXT("LevelEditor"));
+		LevelEditorModule.OnTabManagerChanged().AddRaw(this, &FBlutilityModule::ReinitializeUIs);
+		LevelEditorModule.OnMapChanged().AddRaw(this, &FBlutilityModule::OnMapChanged);
+	}
+
+	void ReinitializeUIs()
+	{
+		EditorUtilityContext = NewObject<UEditorUtilityContext>();
+		if (EditorUtilityContext)
+		{
+			for (FSoftObjectPath BlueprintPath : EditorUtilityContext->LoadedUIs)
+			{ 
+				UObject* BlueprintObject = BlueprintPath.TryLoad();
+				if (BlueprintObject)
+				{
+					UEditorUtilityWidgetBlueprint* Blueprint = Cast<UEditorUtilityWidgetBlueprint>(BlueprintObject);
+					const UEditorUtilityWidget* CDO = Blueprint->GeneratedClass->GetDefaultObject<UEditorUtilityWidget>();
+					FName RegistrationName = FName(*CDO->GetPathName());
+					FText DisplayName = FText::FromString(Blueprint->GetName());
+					FLevelEditorModule& LevelEditorModule = FModuleManager::GetModuleChecked<FLevelEditorModule>(TEXT("LevelEditor"));
+					TSharedPtr<FTabManager> LevelEditorTabManager = LevelEditorModule.GetLevelEditorTabManager();
+					if (LevelEditorTabManager && !LevelEditorTabManager->CanSpawnTab(RegistrationName))
+					{
+						LevelEditorTabManager->RegisterTabSpawner(RegistrationName, FOnSpawnTab::CreateUObject(Blueprint, &UEditorUtilityWidgetBlueprint::SpawnEditorUITab))
+							.SetDisplayName(DisplayName)
+							.SetGroup(GetMenuGroup().ToSharedRef());
+					}
+				}
+			}
+		}
+	}
+
+	void OnMapChanged(UWorld* World, EMapChangeType MapChangeType)
+	{
+		if (EditorUtilityContext)
+		{
+			for (FSoftObjectPath LoadedUI : EditorUtilityContext->LoadedUIs)
+			{
+				UEditorUtilityWidgetBlueprint* LoadedEditorUtilityBlueprint = Cast<UEditorUtilityWidgetBlueprint>(LoadedUI.ResolveObject());
+				if (LoadedEditorUtilityBlueprint)
+				{
+					UEditorUtilityWidget* CreatedWidget = LoadedEditorUtilityBlueprint->GetCreatedWidget();
+					if (CreatedWidget)
+					{
+						if (MapChangeType == EMapChangeType::TearDownWorld)
+						{
+							CreatedWidget->Rename(*CreatedWidget->GetName(), GetTransientPackage());
+						}
+						else if (MapChangeType == EMapChangeType::LoadMap || MapChangeType == EMapChangeType::NewMap)
+						{
+							UWorld* World = GEditor->GetEditorWorldContext().World();
+							check(World);
+							CreatedWidget->Rename(*CreatedWidget->GetName(), World);
+						}
+					}
+				}
+			}
+		}
 	}
 
 	virtual void ShutdownModule() override
@@ -65,6 +155,11 @@ public:
 			return;
 		}
 
+		// Unregister widget blueprint compiler we do this no matter what.
+		IUMGEditorModule& UMGEditorModule = FModuleManager::LoadModuleChecked<IUMGEditorModule>("UMGEditor");
+		IKismetCompilerInterface& KismetCompilerModule = FModuleManager::LoadModuleChecked<IKismetCompilerInterface>("KismetCompiler");
+		KismetCompilerModule.GetCompilers().Remove(UMGEditorModule.GetRegisteredCompiler());
+
 		FBlutilityLevelEditorExtensions::RemoveHooks();
 		FBlutilityContentBrowserExtensions::RemoveHooks();
 
@@ -74,9 +169,12 @@ public:
 		check( EditorBlueprintAssetTypeActions.IsValid() );
 		if (FModuleManager::Get().IsModuleLoaded("AssetTools"))
 		{
-			FModuleManager::GetModuleChecked<FAssetToolsModule>("AssetTools").Get().UnregisterAssetTypeActions(EditorBlueprintAssetTypeActions.ToSharedRef());
+			IAssetTools& AssetTools = FModuleManager::LoadModuleChecked<FAssetToolsModule>("AssetTools").Get();
+			AssetTools.UnregisterAssetTypeActions(EditorBlueprintAssetTypeActions.ToSharedRef());
+			AssetTools.UnregisterAssetTypeActions(EditorWidgetBlueprintAssetTypeActions.ToSharedRef());
 		}
 		EditorBlueprintAssetTypeActions.Reset();
+		EditorWidgetBlueprintAssetTypeActions.Reset();
 
 		// Unregister the details customization
 		if (FModuleManager::Get().IsModuleLoaded("PropertyEditor"))
@@ -99,6 +197,35 @@ public:
 		return false;
 	}
 
+	virtual TSharedPtr<class FWorkspaceItem> GetMenuGroup() const override
+	{
+		return ScriptedEditorWidgetsGroup;
+	}
+
+	virtual EAssetTypeCategories::Type GetAssetCategory() const override
+	{
+		return EditorUtilityAssetCategory;
+	}
+
+	virtual void AddLoadedScriptUI(class UEditorUtilityWidgetBlueprint* InBlueprint) override
+	{
+		if (EditorUtilityContext)
+		{
+			EditorUtilityContext->LoadedUIs.Add(InBlueprint);
+			EditorUtilityContext->SaveConfig();
+		}
+	}
+
+
+	virtual void RemoveLoadedScriptUI(class UEditorUtilityWidgetBlueprint* InBlueprint) override
+	{
+		if (EditorUtilityContext)
+		{
+			EditorUtilityContext->LoadedUIs.Remove(InBlueprint);
+			EditorUtilityContext->SaveConfig();
+		}
+	}
+
 protected:
 	static TSharedRef<SDockTab> SpawnBlutilityShelfTab(const FSpawnTabArgs& Args)
 	{
@@ -108,8 +235,24 @@ protected:
 				SNew(SBlutilityShelf)
 			];
 	}
+
+	virtual void AddReferencedObjects(FReferenceCollector& Collector) override
+	{
+		if (EditorUtilityContext)
+		{
+			Collector.AddReferencedObject(EditorUtilityContext);
+		}
+	}
+
+	/** Scripted Editor Widgets workspace menu item */
+	TSharedPtr<class FWorkspaceItem> ScriptedEditorWidgetsGroup;
+
+	EAssetTypeCategories::Type EditorUtilityAssetCategory;
+
+	UEditorUtilityContext* EditorUtilityContext;
 };
 
 
 
 IMPLEMENT_MODULE( FBlutilityModule, Blutility );
+#undef LOCTEXT_NAMESPACE 
\ No newline at end of file
diff --git a/Engine/Source/Editor/Blutility/Private/EditorUtilityContext.cpp b/Engine/Source/Editor/Blutility/Private/EditorUtilityContext.cpp
new file mode 100644
--- /dev/null
+++ b/Engine/Source/Editor/Blutility/Private/EditorUtilityContext.cpp
@@ -0,0 +1,22 @@
+// Copyright 1998-2018 Epic Games, Inc. All Rights Reserved.
+
+#include "EditorUtilityContext.h"
+#include "Engine/Selection.h"
+#include "Editor.h"
+#include "GameFramework/Actor.h"
+#include "ContentBrowserModule.h"
+#include "Modules/ModuleManager.h"
+#include "IContentBrowserSingleton.h"
+#include "AssetToolsModule.h"
+#include "IAssetTools.h"
+
+
+#define LOCTEXT_NAMESPACE "BlutilityLevelEditorExtensions"
+
+UEditorUtilityContext::UEditorUtilityContext(const FObjectInitializer& ObjectInitializer)
+	: Super(ObjectInitializer)
+{
+}
+
+
+#undef LOCTEXT_NAMESPACE
diff --git a/Engine/Source/Editor/Blutility/Private/EditorUtilityContext.h b/Engine/Source/Editor/Blutility/Private/EditorUtilityContext.h
new file mode 100644
--- /dev/null
+++ b/Engine/Source/Editor/Blutility/Private/EditorUtilityContext.h
@@ -0,0 +1,24 @@
+// Copyright 1998-2018 Epic Games, Inc. All Rights Reserved.
+
+#pragma once
+
+#include "CoreMinimal.h"
+#include "UObject/ObjectMacros.h"
+#include "UObject/Object.h"
+#include "UObject/ScriptMacros.h"
+#include "Kismet/BlueprintFunctionLibrary.h"
+#include "EditorUtilityContext.generated.h"
+
+class AActor;
+class UEditorPerProjectUserSettings;
+
+// Expose editor utility functions to Blutilities 
+UCLASS(config = EditorPerProjectUserSettings)
+class UEditorUtilityContext : public UObject
+{
+	GENERATED_UCLASS_BODY()
+
+public:
+	UPROPERTY(config)
+	TArray<FSoftObjectPath> LoadedUIs;
+};
diff --git a/Engine/Source/Editor/Blutility/Private/EditorUtilityDialog.cpp b/Engine/Source/Editor/Blutility/Private/EditorUtilityDialog.cpp
new file mode 100644
--- /dev/null
+++ b/Engine/Source/Editor/Blutility/Private/EditorUtilityDialog.cpp
@@ -0,0 +1,131 @@
+// Copyright 1998-2018 Epic Games, Inc. All Rights Reserved.
+
+#include "EditorUtilityDialog.h"
+#include "GlobalEditorUtilityBase.h"
+#include "Modules/ModuleManager.h"
+#include "EditorUtilityBlueprint.h"
+#include "PropertyEditorModule.h"
+
+#include "IDetailsView.h"
+#include "Widgets/Docking/SDockTab.h"
+
+#define LOCTEXT_NAMESPACE "GlobalBlutilityDialog"
+
+const FName NAME_DetailsPanel(TEXT("GlobalBlutilityDialog_DetailsPanel"));
+const FName NAME_GlobalBlutilityDialogAppIdentifier = FName(TEXT("GlobalBlutilityDialogApp"));
+
+//////////////////////////////////////////////////////////////////////////
+// FGlobalBlutilityDialog
+
+void FGlobalEditorUtilityDialog::RegisterTabSpawners(const TSharedRef<class FTabManager>& InTabManager)
+{
+	InTabManager->RegisterTabSpawner( NAME_DetailsPanel, FOnSpawnTab::CreateRaw( this, &FGlobalEditorUtilityDialog::SpawnTab_DetailsPanel ) );
+
+}
+
+void FGlobalEditorUtilityDialog::UnregisterTabSpawners(const TSharedRef<class FTabManager>& InTabManager)
+{
+	InTabManager->UnregisterTabSpawner( NAME_DetailsPanel );
+}
+
+TSharedRef<SDockTab> FGlobalEditorUtilityDialog::SpawnTab_DetailsPanel( const FSpawnTabArgs& SpawnTabArgs )
+{
+	const TSharedRef<SDockTab> SpawnedTab = SNew(SDockTab)
+		//@TODO: Add an icon .Icon( FEditorStyle::GetBrush("SoundClassEditor.Tabs.Properties") )
+		.Label( LOCTEXT("GlobalBlutilityDetailsTitle", "Blutility Details") )
+		[
+			DetailsView.ToSharedRef()
+		];
+
+	// Make sure the blutility instance is selected
+	TArray<UObject*> SelectedObjects;
+	SelectedObjects.Add(BlutilityInstance.Get());
+	UpdatePropertyWindow(SelectedObjects);
+
+	return SpawnedTab;
+}
+
+void FGlobalEditorUtilityDialog::InitBlutilityDialog(const EToolkitMode::Type Mode, const TSharedPtr< class IToolkitHost >& InitToolkitHost, UObject* ObjectToEdit)
+{
+	// Create an instance of the blutility
+	check(ObjectToEdit != NULL);
+	UEditorUtilityBlueprint* BlutilityBP = CastChecked<UEditorUtilityBlueprint>(ObjectToEdit);
+	check(BlutilityBP->GeneratedClass->IsChildOf(UGlobalEditorUtilityBase::StaticClass()));
+
+	UGlobalEditorUtilityBase* Instance = NewObject<UGlobalEditorUtilityBase>(GetTransientPackage(), BlutilityBP->GeneratedClass);
+	Instance->AddToRoot();
+	BlutilityInstance = Instance;
+
+	//
+	CreateInternalWidgets();
+
+	const TSharedRef<FTabManager::FLayout> StandaloneDefaultLayout = FTabManager::NewLayout( "Standalone_GlobalBlutility_Layout" ) 
+	->AddArea
+	(
+		FTabManager::NewPrimaryArea()
+		->Split
+		(
+			FTabManager::NewStack()
+			->AddTab( NAME_DetailsPanel, ETabState::OpenedTab)
+		)
+	);
+
+	const bool bCreateDefaultStandaloneMenu = false;
+	const bool bCreateDefaultToolbar = false;
+	FAssetEditorToolkit::InitAssetEditor( Mode, InitToolkitHost, NAME_GlobalBlutilityDialogAppIdentifier, StandaloneDefaultLayout, bCreateDefaultStandaloneMenu, bCreateDefaultToolbar, ObjectToEdit );
+
+}
+
+FGlobalEditorUtilityDialog::~FGlobalEditorUtilityDialog()
+{
+	if (UGlobalEditorUtilityBase* Instance = BlutilityInstance.Get())
+	{
+		Instance->RemoveFromRoot();
+	}
+
+	DetailsView.Reset();
+}
+
+void FGlobalEditorUtilityDialog::AddReferencedObjects(FReferenceCollector& Collector)
+{
+	if (UGlobalEditorUtilityBase* Instance = BlutilityInstance.Get())
+	{
+		Collector.AddReferencedObject(Instance);
+	}
+}
+
+FName FGlobalEditorUtilityDialog::GetToolkitFName() const
+{
+	return FName("Blutility");
+}
+
+FText FGlobalEditorUtilityDialog::GetBaseToolkitName() const
+{
+	return LOCTEXT( "AppLabel", "Blutility" );
+}
+
+FString FGlobalEditorUtilityDialog::GetWorldCentricTabPrefix() const
+{
+	return LOCTEXT("WorldCentricTabPrefix", "Blutility ").ToString();
+}
+
+FLinearColor FGlobalEditorUtilityDialog::GetWorldCentricTabColorScale() const
+{
+	return FLinearColor( 0.3f, 0.2f, 0.5f, 0.5f );
+}
+
+void FGlobalEditorUtilityDialog::CreateInternalWidgets()
+{
+	// Create a details view
+	FPropertyEditorModule& PropertyEditorModule = FModuleManager::GetModuleChecked<FPropertyEditorModule>("PropertyEditor");
+	FDetailsViewArgs DetailsViewArgs(/*bUpdateFromSelection=*/ false, /*bLockable=*/ false, /*bAllowSearch=*/ false, FDetailsViewArgs::HideNameArea, /*bHideSelectionTip=*/ true);
+	DetailsViewArgs.DefaultsOnlyVisibility = EEditDefaultsOnlyNodeVisibility::Hide;
+	DetailsView = PropertyEditorModule.CreateDetailView(DetailsViewArgs);
+}
+
+void FGlobalEditorUtilityDialog::UpdatePropertyWindow(const TArray<UObject*>& SelectedObjects)
+{
+	DetailsView->SetObjects(SelectedObjects);
+}
+
+#undef LOCTEXT_NAMESPACE
diff --git a/Engine/Source/Editor/Blutility/Private/EditorUtilityDialog.h b/Engine/Source/Editor/Blutility/Private/EditorUtilityDialog.h
new file mode 100644
--- /dev/null
+++ b/Engine/Source/Editor/Blutility/Private/EditorUtilityDialog.h
@@ -0,0 +1,64 @@
+// Copyright 1998-2018 Epic Games, Inc. All Rights Reserved.
+
+#pragma once
+
+#include "CoreMinimal.h"
+#include "UObject/GCObject.h"
+#include "Toolkits/IToolkitHost.h"
+#include "Toolkits/AssetEditorToolkit.h"
+
+class IDetailsView;
+class SDockableTab;
+class UGlobalEditorUtilityBase;
+
+//////////////////////////////////////////////////////////////////////////
+// FGlobalBlutilityDialog
+
+class FGlobalEditorUtilityDialog : public FAssetEditorToolkit, public FGCObject
+{
+public:
+	virtual void RegisterTabSpawners(const TSharedRef<class FTabManager>& TabManager) override;
+	virtual void UnregisterTabSpawners(const TSharedRef<class FTabManager>& TabManager) override;
+	TSharedRef<SDockTab> SpawnTab_DetailsPanel( const FSpawnTabArgs& SpawnTabArgs );
+
+	/**
+	 * Edits the specified editor utility object
+	 *
+	 * @param	Mode					Asset editing mode for this editor (standalone or world-centric)
+	 * @param	InitToolkitHost			When Mode is WorldCentric, this is the level editor instance to spawn this editor within
+	 * @param	ObjectToEdit			The sound class to edit
+	 */
+	void InitBlutilityDialog(const EToolkitMode::Type Mode, const TSharedPtr< class IToolkitHost >& InitToolkitHost, UObject* ObjectToEdit);
+
+	/**
+	 * Updates the editor's property window to contain sound class properties if any are selected.
+	 */
+	void UpdatePropertyWindow( const TArray<UObject*>& SelectedObjects );
+
+	virtual ~FGlobalEditorUtilityDialog();
+
+	/** FGCObject interface */
+	virtual void AddReferencedObjects( FReferenceCollector& Collector ) override;
+
+	/** IToolkit interface */
+	virtual FName GetToolkitFName() const override;
+	virtual FText GetBaseToolkitName() const override;
+	virtual FString GetWorldCentricTabPrefix() const override;
+
+	/** @return Returns the color and opacity to use for the color that appears behind the tab text for this toolkit's tab in world-centric mode. */
+	virtual FLinearColor GetWorldCentricTabColorScale() const override;
+
+private:
+
+	/** Creates all internal widgets for the tabs to point at */
+	void CreateInternalWidgets();
+
+private:
+	/** List of open tool panels; used to ensure only one exists at any one time */
+	TMap< FName, TWeakPtr<class SDockableTab> > SpawnedToolPanels;
+
+	/** Property View */
+	TSharedPtr<class IDetailsView> DetailsView;
+
+	TWeakObjectPtr<UGlobalEditorUtilityBase> BlutilityInstance;
+};
diff --git a/Engine/Source/Editor/Blutility/Private/EditorUtilityLibrary.cpp b/Engine/Source/Editor/Blutility/Private/EditorUtilityLibrary.cpp
new file mode 100644
--- /dev/null
+++ b/Engine/Source/Editor/Blutility/Private/EditorUtilityLibrary.cpp
@@ -0,0 +1,98 @@
+// Copyright 1998-2018 Epic Games, Inc. All Rights Reserved.
+
+#include "EditorUtilityLibrary.h"
+#include "Engine/Selection.h"
+#include "Editor.h"
+#include "GameFramework/Actor.h"
+#include "ContentBrowserModule.h"
+#include "Modules/ModuleManager.h"
+#include "IContentBrowserSingleton.h"
+#include "AssetToolsModule.h"
+#include "IAssetTools.h"
+
+
+#define LOCTEXT_NAMESPACE "BlutilityLevelEditorExtensions"
+
+UEditorUtilityLibrary::UEditorUtilityLibrary(const FObjectInitializer& ObjectInitializer)
+	: Super(ObjectInitializer)
+{
+}
+
+TArray<AActor*> UEditorUtilityLibrary::GetSelectionSet()
+{
+	TArray<AActor*> Result;
+	for (FSelectionIterator It(GEditor->GetSelectedActorIterator()); It; ++It)
+	{
+		if (AActor* Actor = Cast<AActor>(*It))
+		{
+			Result.Add(Actor);
+		}
+	}
+
+	return Result;
+}
+
+void UEditorUtilityLibrary::GetSelectionBounds(FVector& Origin, FVector& BoxExtent, float& SphereRadius)
+{
+	bool bFirstItem = true;
+
+	FBoxSphereBounds Extents;
+	for (FSelectionIterator It(GEditor->GetSelectedActorIterator()); It; ++It)
+	{
+		if (AActor* Actor = Cast<AActor>(*It))
+		{
+			if (bFirstItem)
+			{
+				Extents = Actor->GetRootComponent()->Bounds;
+			}
+			else
+			{
+				Extents = Extents + Actor->GetRootComponent()->Bounds;
+			}
+
+			bFirstItem = false;
+		}
+	}
+
+	Origin = Extents.Origin;
+	BoxExtent = Extents.BoxExtent;
+	SphereRadius = Extents.SphereRadius;
+}
+
+TArray<UObject*> UEditorUtilityLibrary::GetSelectedAssets()
+{
+	//@TODO: Blocking load, no slow dialog
+	FContentBrowserModule& ContentBrowserModule = FModuleManager::LoadModuleChecked<FContentBrowserModule>("ContentBrowser");
+	TArray<FAssetData> SelectedAssets;
+	ContentBrowserModule.Get().GetSelectedAssets(SelectedAssets);
+
+	TArray<UObject*> Result;
+	for (FAssetData& AssetData : SelectedAssets)
+	{
+		Result.Add(AssetData.GetAsset());
+	}
+
+	return Result;
+}
+
+void UEditorUtilityLibrary::RenameAsset(UObject* Asset, const FString& NewName)
+{
+	FAssetToolsModule& AssetToolsModule = FModuleManager::LoadModuleChecked<FAssetToolsModule>("AssetTools");
+
+	TArray<FAssetRenameData> AssetsAndNames;
+	const FString PackagePath = FPackageName::GetLongPackagePath(Asset->GetOutermost()->GetName());
+	new (AssetsAndNames) FAssetRenameData(Asset, PackagePath, NewName);
+
+	AssetToolsModule.Get().RenameAssetsWithDialog(AssetsAndNames);
+}
+
+AActor* UEditorUtilityLibrary::GetActorReference(FString PathToActor)
+{
+#if WITH_EDITOR
+	return Cast<AActor>(StaticFindObject(AActor::StaticClass(), GEditor->GetEditorWorldContext().World(), *PathToActor, false));
+#else
+	return nullptr;
+#endif //WITH_EDITOR
+}
+
+#undef LOCTEXT_NAMESPACE
diff --git a/Engine/Source/Editor/Blutility/Private/EditorUtilityWidget.cpp b/Engine/Source/Editor/Blutility/Private/EditorUtilityWidget.cpp
new file mode 100644
--- /dev/null
+++ b/Engine/Source/Editor/Blutility/Private/EditorUtilityWidget.cpp
@@ -0,0 +1,37 @@
+// Copyright 1998-2018 Epic Games, Inc. All Rights Reserved.
+
+#include "EditorUtilityWidget.h"
+#include "Modules/ModuleManager.h"
+#include "Misc/PackageName.h"
+#include "GameFramework/Actor.h"
+#include "Editor/EditorPerProjectUserSettings.h"
+#include "AssetData.h"
+#include "Engine/Selection.h"
+#include "Editor.h"
+#include "ScopedTransaction.h"
+#include "IAssetTools.h"
+#include "IContentBrowserSingleton.h"
+#include "ContentBrowserModule.h"
+#include "AssetToolsModule.h"
+
+/////////////////////////////////////////////////////
+#define LOCTEXT_NAMESPACE "EditorUtility"
+
+UEditorUtilityWidget::UEditorUtilityWidget(const FObjectInitializer& ObjectInitializer)
+	: Super(ObjectInitializer)
+{
+}
+
+void UEditorUtilityWidget::ExecuteDefaultAction()
+{
+	check(bAutoRunDefaultAction);
+
+	FScopedTransaction Transaction(NSLOCTEXT("UnrealEd", "BlutilityAction", "Blutility Action"));
+	FEditorScriptExecutionGuard ScriptGuard;
+
+	OnDefaultActionClicked();
+}
+
+
+
+#undef LOCTEXT_NAMESPACE
\ No newline at end of file
diff --git a/Engine/Source/Editor/Blutility/Private/EditorUtilityWidgetBlueprint.cpp b/Engine/Source/Editor/Blutility/Private/EditorUtilityWidgetBlueprint.cpp
new file mode 100644
--- /dev/null
+++ b/Engine/Source/Editor/Blutility/Private/EditorUtilityWidgetBlueprint.cpp
@@ -0,0 +1,81 @@
+// Copyright 1998-2018 Epic Games, Inc. All Rights Reserved.
+
+#include "EditorUtilityWidgetBlueprint.h"
+#include "WidgetBlueprint.h"
+#include "Editor.h"
+#include "EditorUtilityWidget.h"
+#include "IBlutilityModule.h"
+#include "Modules/ModuleManager.h"
+
+
+
+
+/////////////////////////////////////////////////////
+// UEditorUtilityWidgetBlueprint
+
+UEditorUtilityWidgetBlueprint::UEditorUtilityWidgetBlueprint(const FObjectInitializer& ObjectInitializer)
+	: Super(ObjectInitializer)
+{
+}
+
+TSharedRef<SDockTab> UEditorUtilityWidgetBlueprint::SpawnEditorUITab(const FSpawnTabArgs& SpawnTabArgs)
+{
+	TSharedRef<SDockTab> SpawnedTab = SNew(SDockTab);
+
+	TSharedRef<SWidget> TabWidget = CreateUtilityWidget();
+	SpawnedTab->SetContent(TabWidget);
+	
+	SpawnedTab->SetOnTabClosed(SDockTab::FOnTabClosedCallback::CreateUObject(this, &UEditorUtilityWidgetBlueprint::UpdateRespawnListIfNeeded));
+	CreatedTab = SpawnedTab;
+	
+	OnCompiled().AddUObject(this, &UEditorUtilityWidgetBlueprint::RegenerateCreatedTab);
+
+	return SpawnedTab;
+}
+
+TSharedRef<SWidget> UEditorUtilityWidgetBlueprint::CreateUtilityWidget()
+{
+	UClass* BlueprintClass = GeneratedClass;
+	TSubclassOf<UEditorUtilityWidget> WidgetClass = BlueprintClass;
+	UWorld* World = GEditor->GetEditorWorldContext().World();
+	check(World);
+	CreatedUMGWidget = CreateWidget<UEditorUtilityWidget>(World, WidgetClass);
+	TSharedRef<SWidget> TabWidget = SNullWidget::NullWidget;
+	if (CreatedUMGWidget)
+	{
+		TSharedRef<SWidget> CreatedSlateWidget = CreatedUMGWidget->TakeWidget();
+		TabWidget = SNew(SVerticalBox)
+			+ SVerticalBox::Slot()
+			.HAlign(HAlign_Fill)
+			[
+				CreatedSlateWidget
+			];
+	}
+	return TabWidget;
+}
+
+void UEditorUtilityWidgetBlueprint::RegenerateCreatedTab(UBlueprint* RecompiledBlueprint)
+{
+	if (CreatedTab.IsValid())
+	{
+		TSharedRef<SWidget> TabWidget = CreateUtilityWidget();
+		CreatedTab.Pin()->SetContent(TabWidget);
+	}
+}
+
+void UEditorUtilityWidgetBlueprint::UpdateRespawnListIfNeeded(TSharedRef<SDockTab> TabBeingClosed)
+{
+	UEditorUtilityWidget* EditorUtilityWidget = Cast<UEditorUtilityWidget>(GeneratedClass);
+	if (EditorUtilityWidget && EditorUtilityWidget->ShouldAlwaysReregisterWithWindowsMenu() == false)
+	{
+		IBlutilityModule* BlutilityModule = FModuleManager::GetModulePtr<IBlutilityModule>("Blutility");
+		BlutilityModule->RemoveLoadedScriptUI(this);
+	}
+}
+
+void UEditorUtilityWidgetBlueprint::GetReparentingRules(TSet< const UClass* >& AllowedChildrenOfClasses, TSet< const UClass* >& DisallowedChildrenOfClasses) const
+{
+	AllowedChildrenOfClasses.Empty();
+	AllowedChildrenOfClasses.Add(UEditorUtilityWidget::StaticClass());
+}
+
diff --git a/Engine/Source/Editor/Blutility/Private/EditorUtilityWidgetBlueprintFactory.cpp b/Engine/Source/Editor/Blutility/Private/EditorUtilityWidgetBlueprintFactory.cpp
new file mode 100644
--- /dev/null
+++ b/Engine/Source/Editor/Blutility/Private/EditorUtilityWidgetBlueprintFactory.cpp
@@ -0,0 +1,88 @@
+// Copyright 1998-2018 Epic Games, Inc. All Rights Reserved.
+
+#include "EditorUtilityWidgetBlueprintFactory.h"
+#include "Misc/MessageDialog.h"
+#include "Modules/ModuleManager.h"
+#include "Widgets/SWindow.h"
+#include "Engine/Blueprint.h"
+#include "Settings/EditorExperimentalSettings.h"
+#include "Engine/BlueprintGeneratedClass.h"
+#include "EditorUtilityBlueprint.h"
+#include "GlobalEditorUtilityBase.h"
+#include "PlacedEditorUtilityBase.h"
+#include "ClassViewerModule.h"
+#include "ClassViewerFilter.h"
+#include "Kismet2/KismetEditorUtilities.h"
+#include "Kismet2/SClassPickerDialog.h"
+#include "EditorUtilityWidget.h"
+#include "EditorUtilityWidgetBlueprint.h"
+#include "Components/CanvasPanel.h"
+#include "BaseWidgetBlueprint.h"
+#include "Blueprint/WidgetTree.h"
+
+class FEditorUtilityWidgetBlueprintFactoryFilter : public IClassViewerFilter
+{
+public:
+	TSet< const UClass* > AllowedChildOfClasses;
+
+	bool IsClassAllowed(const FClassViewerInitializationOptions& InInitOptions, const UClass* InClass, TSharedRef< FClassViewerFilterFuncs > InFilterFuncs ) override
+	{
+		return InFilterFuncs->IfInChildOfClassesSet(AllowedChildOfClasses, InClass) != EFilterReturn::Failed;
+	}
+
+	virtual bool IsUnloadedClassAllowed(const FClassViewerInitializationOptions& InInitOptions, const TSharedRef< const IUnloadedBlueprintData > InUnloadedClassData, TSharedRef< FClassViewerFilterFuncs > InFilterFuncs) override
+	{
+		return InFilterFuncs->IfInChildOfClassesSet(AllowedChildOfClasses, InUnloadedClassData) != EFilterReturn::Failed;
+	}
+};
+
+/////////////////////////////////////////////////////
+// UEditorUtilityWidgetBlueprintFactory
+
+UEditorUtilityWidgetBlueprintFactory::UEditorUtilityWidgetBlueprintFactory(const FObjectInitializer& ObjectInitializer)
+	: Super(ObjectInitializer)
+{
+	bCreateNew = true;
+	bEditAfterNew = true;
+	SupportedClass = UEditorUtilityWidgetBlueprint::StaticClass();
+	ParentClass = UEditorUtilityWidget::StaticClass();
+}
+
+bool UEditorUtilityWidgetBlueprintFactory::ConfigureProperties()
+{
+	return true;
+}
+
+UObject* UEditorUtilityWidgetBlueprintFactory::FactoryCreateNew(UClass* Class, UObject* InParent, FName Name, EObjectFlags Flags, UObject* Context, FFeedbackContext* Warn)
+{
+	// Make sure we are trying to factory a blueprint, then create and init one
+	check(Class->IsChildOf(UEditorUtilityWidgetBlueprint::StaticClass()));
+
+	FString ParentPath = InParent->GetPathName();
+
+	if ((ParentClass == nullptr) || !FKismetEditorUtilities::CanCreateBlueprintOfClass(ParentClass))
+	{
+		FFormatNamedArguments Args;
+		Args.Add(TEXT("ClassName"), (ParentClass != nullptr) ? FText::FromString(ParentClass->GetName()) : NSLOCTEXT("UnrealEd", "Null", "(null)"));
+		FMessageDialog::Open(EAppMsgType::Ok, FText::Format(NSLOCTEXT("UnrealEd", "CannotCreateBlueprintFromClass", "Cannot create a blueprint based on the class '{0}'."), Args));
+		return nullptr;
+	}
+	else
+	{
+		UEditorUtilityWidgetBlueprint* NewBP = CastChecked<UEditorUtilityWidgetBlueprint>(FKismetEditorUtilities::CreateBlueprint(ParentClass, InParent, Name, BlueprintType, UEditorUtilityWidgetBlueprint::StaticClass(), UWidgetBlueprintGeneratedClass::StaticClass(), NAME_None));
+
+		// Create a CanvasPanel to use as the default root widget
+		if (NewBP->WidgetTree->RootWidget == nullptr)
+		{
+			UWidget* Root = NewBP->WidgetTree->ConstructWidget<UCanvasPanel>(UCanvasPanel::StaticClass());
+			NewBP->WidgetTree->RootWidget = Root;
+		}
+
+		return NewBP;
+	}
+}
+
+bool UEditorUtilityWidgetBlueprintFactory::CanCreateNew() const
+{
+	return GetDefault<UEditorExperimentalSettings>()->bEnableEditorUtilityBlueprints;
+}
diff --git a/Engine/Source/Editor/Blutility/Private/GlobalBlutilityDialog.cpp b/Engine/Source/Editor/Blutility/Private/GlobalBlutilityDialog.cpp
deleted file mode 100644
diff --git a/Engine/Source/Editor/Blutility/Private/GlobalBlutilityDialog.h b/Engine/Source/Editor/Blutility/Private/GlobalBlutilityDialog.h
deleted file mode 100644
diff --git a/Engine/Source/Editor/Blutility/Public/EditorUtilityLibrary.h b/Engine/Source/Editor/Blutility/Public/EditorUtilityLibrary.h
new file mode 100644
--- /dev/null
+++ b/Engine/Source/Editor/Blutility/Public/EditorUtilityLibrary.h
@@ -0,0 +1,44 @@
+// Copyright 1998-2018 Epic Games, Inc. All Rights Reserved.
+
+#pragma once
+
+#include "CoreMinimal.h"
+#include "UObject/ObjectMacros.h"
+#include "UObject/Object.h"
+#include "UObject/ScriptMacros.h"
+#include "Kismet/BlueprintFunctionLibrary.h"
+#include "EditorUtilityLibrary.generated.h"
+
+class AActor;
+class UEditorPerProjectUserSettings;
+
+// Expose editor utility functions to Blutilities 
+UCLASS()
+class BLUTILITY_API UEditorUtilityLibrary : public UBlueprintFunctionLibrary
+{
+	GENERATED_UCLASS_BODY()
+
+public:
+	UFUNCTION(BlueprintCallable, Category = "Development|Editor")
+	static TArray<AActor*> GetSelectionSet();
+
+	UFUNCTION(BlueprintCallable, Category = "Development|Editor")
+	static void GetSelectionBounds(FVector& Origin, FVector& BoxExtent, float& SphereRadius);
+
+	// Gets the set of currently selected assets
+	UFUNCTION(BlueprintCallable, Category = "Development|Editor")
+	static TArray<UObject*> GetSelectedAssets();
+
+	// Renames an asset (cannot move folders)
+	UFUNCTION(BlueprintCallable, Category = "Development|Editor")
+	static void RenameAsset(UObject* Asset, const FString& NewName);
+
+	/**
+	* Attempts to find the actor specified by PathToActor in the current editor world
+	* @param	PathToActor	The path to the actor (e.g. PersistentLevel.PlayerStart)
+	* @return	A reference to the actor, or none if it wasn't found
+	*/
+	UFUNCTION(BlueprintPure, Category = "Development|Editor")
+	AActor* GetActorReference(FString PathToActor);
+
+};
diff --git a/Engine/Source/Editor/Blutility/Public/IBlutilityModule.h b/Engine/Source/Editor/Blutility/Public/IBlutilityModule.h
--- a/Engine/Source/Editor/Blutility/Public/IBlutilityModule.h
+++ b/Engine/Source/Editor/Blutility/Public/IBlutilityModule.h
@@ -4,6 +4,7 @@
 
 #include "CoreMinimal.h"
 #include "Modules/ModuleInterface.h"
+#include "AssetTypeCategories.h"
 
 class UBlueprint;
 
@@ -17,5 +18,13 @@ public:
 	/** Returns if the blueprint is blutility based */
 	virtual bool IsBlutility( const UBlueprint* Blueprint ) const = 0;
 
+	/** Global Find Results workspace menu item */
+	virtual TSharedPtr<class FWorkspaceItem> GetMenuGroup() const = 0;
+
+	virtual EAssetTypeCategories::Type GetAssetCategory() const = 0;
+
+	virtual void AddLoadedScriptUI(class UEditorUtilityWidgetBlueprint* InBlueprint) = 0;
+
+	virtual void RemoveLoadedScriptUI(class UEditorUtilityWidgetBlueprint* InBlueprint) = 0;
 };
 
diff --git a/Engine/Source/Editor/EditorStyle/Private/SlateEditorStyle.cpp b/Engine/Source/Editor/EditorStyle/Private/SlateEditorStyle.cpp
--- a/Engine/Source/Editor/EditorStyle/Private/SlateEditorStyle.cpp
+++ b/Engine/Source/Editor/EditorStyle/Private/SlateEditorStyle.cpp
@@ -5217,6 +5217,8 @@ void FSlateEditorStyle::FStyle::SetupLevelEditorStyle()
 	{
 		Set( "Scalability.ScalabilitySettings", new IMAGE_BRUSH("Scalability/ScalabilitySettings", FVector2D(473.0f, 266.0f) ) );
 	}
+
+	Set("WorkspaceMenu.AdditionalUI", new IMAGE_BRUSH("Icons/icon_ViewMode_LODColoration_16px", Icon16x16));
 }
 
 void FSlateEditorStyle::FStyle::SetupPersonaStyle()
diff --git a/Engine/Source/Editor/UMGEditor/Private/UMGEditorModule.cpp b/Engine/Source/Editor/UMGEditor/Private/UMGEditorModule.cpp
--- a/Engine/Source/Editor/UMGEditor/Private/UMGEditorModule.cpp
+++ b/Engine/Source/Editor/UMGEditor/Private/UMGEditorModule.cpp
@@ -23,8 +23,6 @@
 #include "Animation/WidgetMaterialTrackEditor.h"
 #include "Animation/MovieSceneSequenceEditor_WidgetAnimation.h"
 #include "IUMGModule.h"
-#include "ComponentReregisterContext.h"
-#include "Components/WidgetComponent.h"
 #include "Designer/DesignerCommands.h"
 
 #include "ClassIconFinder.h"
@@ -42,14 +40,12 @@
 
 const FName UMGEditorAppIdentifier = FName(TEXT("UMGEditorApp"));
 
-class FUMGEditorModule : public IUMGEditorModule, public IBlueprintCompiler, public FGCObject
+class FUMGEditorModule : public IUMGEditorModule, public FGCObject
 {
 public:
 	/** Constructor, set up console commands and variables **/
 	FUMGEditorModule()
-		: ReRegister(nullptr)
-		, CompileCount(0)
-		, Settings(nullptr)
+		: Settings(nullptr)
 	{
 	}
 
@@ -68,7 +64,7 @@ public:
 
 		// Register widget blueprint compiler we do this no matter what.
 		IKismetCompilerInterface& KismetCompilerModule = FModuleManager::LoadModuleChecked<IKismetCompilerInterface>("KismetCompiler");
-		KismetCompilerModule.GetCompilers().Add(this);
+		KismetCompilerModule.GetCompilers().Add(&WidgetBlueprintCompiler);
 
 		// Register asset types
 		IAssetTools& AssetTools = FModuleManager::LoadModuleChecked<FAssetToolsModule>("AssetTools").Get();
@@ -97,6 +93,9 @@ public:
 		MenuExtensibilityManager.Reset();
 		ToolBarExtensibilityManager.Reset();
 
+		IKismetCompilerInterface& KismetCompilerModule = FModuleManager::LoadModuleChecked<IKismetCompilerInterface>("KismetCompiler");
+		KismetCompilerModule.GetCompilers().Remove(&WidgetBlueprintCompiler);
+
 		// Unregister all the asset types that we registered
 		if ( FModuleManager::Get().IsModuleLoaded("AssetTools") )
 		{
@@ -131,61 +130,6 @@ public:
 		//}
 	}
 
-	bool CanCompile(const UBlueprint* Blueprint) override
-	{
-		return Cast<UWidgetBlueprint>(Blueprint) != nullptr;
-	}
-
-	void PreCompile(UBlueprint* Blueprint, const FKismetCompilerOptions& CompileOptions) override
-	{
-		if ( ReRegister == nullptr 
-			&& CanCompile(Blueprint) 
-			&& (CompileOptions.CompileType == EKismetCompileType::Full || CompileOptions.CompileType == EKismetCompileType::Cpp))
-		{
-			ReRegister = new TComponentReregisterContext<UWidgetComponent>();
-		}
-
-		CompileCount++;
-	}
-
-	void Compile(UBlueprint* Blueprint, const FKismetCompilerOptions& CompileOptions, FCompilerResultsLog& Results, TArray<UObject*>* ObjLoaded) override
-	{
-		if ( UWidgetBlueprint* WidgetBlueprint = CastChecked<UWidgetBlueprint>(Blueprint) )
-		{
-			FWidgetBlueprintCompiler Compiler(WidgetBlueprint, Results, CompileOptions, ObjLoaded);
-			Compiler.Compile();
-			check(Compiler.NewClass);
-		}
-	}
-
-	void PostCompile(UBlueprint* Blueprint, const FKismetCompilerOptions& CompileOptions) override
-	{
-		CompileCount--;
-
-		if ( CompileCount == 0 && ReRegister )
-		{
-			delete ReRegister;
-			ReRegister = nullptr;
-
-			if ( GIsEditor && GEditor )
-			{
-				GEditor->RedrawAllViewports(true);
-			}
-		}
-	}
-
-	bool GetBlueprintTypesForClass(UClass* ParentClass, UClass*& OutBlueprintClass, UClass*& OutBlueprintGeneratedClass) const override
-	{
-		if ( ParentClass == UUserWidget::StaticClass() || ParentClass->IsChildOf(UUserWidget::StaticClass()) )
-		{
-			OutBlueprintClass = UWidgetBlueprint::StaticClass();
-			OutBlueprintGeneratedClass = UWidgetBlueprintGeneratedClass::StaticClass();
-			return true;
-		}
-
-		return false;
-	}
-
 	/** Gets the extensibility managers for outside entities to extend gui page editor's menus and toolbars */
 	virtual TSharedPtr<FExtensibilityManager> GetMenuExtensibilityManager() override { return MenuExtensibilityManager; }
 	virtual TSharedPtr<FExtensibilityManager> GetToolBarExtensibilityManager() override { return ToolBarExtensibilityManager; }
@@ -226,6 +170,11 @@ public:
 		}
 	}
 
+	virtual FWidgetBlueprintCompiler* GetRegisteredCompiler() override
+	{
+		return &WidgetBlueprintCompiler;
+	}
+
 private:
 	void RegisterAssetTypeAction(IAssetTools& AssetTools, TSharedRef<IAssetTypeActions> Action)
 	{
@@ -245,16 +194,10 @@ private:
 	/** All created asset type actions.  Cached here so that we can unregister it during shutdown. */
 	TArray< TSharedPtr<IAssetTypeActions> > CreatedAssetTypeActions;
 
-	/** The temporary variable that captures and reinstances components after compiling finishes. */
-	TComponentReregisterContext<UWidgetComponent>* ReRegister;
-
-	/**
-	 * The current count on the number of compiles that have occurred.  We don't want to re-register components until all
-	 * compiling has stopped.
-	 */
-	int32 CompileCount;
-
 	USequencerSettings* Settings;
+
+	/** Compiler customization for Widgets */
+	FWidgetBlueprintCompiler WidgetBlueprintCompiler;
 };
 
 IMPLEMENT_MODULE(FUMGEditorModule, UMGEditor);
diff --git a/Engine/Source/Editor/UMGEditor/Private/WidgetBlueprint.cpp b/Engine/Source/Editor/UMGEditor/Private/WidgetBlueprint.cpp
--- a/Engine/Source/Editor/UMGEditor/Private/WidgetBlueprint.cpp
+++ b/Engine/Source/Editor/UMGEditor/Private/WidgetBlueprint.cpp
@@ -718,7 +718,7 @@ bool UWidgetBlueprint::ValidateGeneratedClass(const UClass* InClass)
 
 TSharedPtr<FKismetCompilerContext> UWidgetBlueprint::GetCompilerForWidgetBP(UBlueprint* BP, FCompilerResultsLog& InMessageLog, const FKismetCompilerOptions& InCompileOptions)
 {
-	return TSharedPtr<FKismetCompilerContext>(new FWidgetBlueprintCompiler(CastChecked<UWidgetBlueprint>(BP), InMessageLog, InCompileOptions, nullptr));
+	return TSharedPtr<FKismetCompilerContext>(new FWidgetBlueprintCompilerContext(CastChecked<UWidgetBlueprint>(BP), InMessageLog, InCompileOptions, nullptr));
 }
 
 void UWidgetBlueprint::GetReparentingRules(TSet< const UClass* >& AllowedChildrenOfClasses, TSet< const UClass* >& DisallowedChildrenOfClasses) const
diff --git a/Engine/Source/Editor/UMGEditor/Private/WidgetBlueprintCompiler.cpp b/Engine/Source/Editor/UMGEditor/Private/WidgetBlueprintCompiler.cpp
--- a/Engine/Source/Editor/UMGEditor/Private/WidgetBlueprintCompiler.cpp
+++ b/Engine/Source/Editor/UMGEditor/Private/WidgetBlueprintCompiler.cpp
@@ -19,6 +19,7 @@
 #include "WidgetGraphSchema.h"
 #include "IUMGModule.h"
 #include "UMGEditorProjectSettings.h"
+#include "Editor.h"
 
 #define LOCTEXT_NAMESPACE "UMG"
 
@@ -26,24 +27,88 @@
 
 extern COREUOBJECT_API bool GMinimalCompileOnLoad;
 
-FWidgetBlueprintCompiler::FWidgetBlueprintCompiler(UWidgetBlueprint* SourceSketch, FCompilerResultsLog& InMessageLog, const FKismetCompilerOptions& InCompilerOptions, TArray<UObject*>* InObjLoaded)
+
+FWidgetBlueprintCompiler::FWidgetBlueprintCompiler()
+	: ReRegister(nullptr)
+	, CompileCount(0)
+{
+
+}
+
+bool FWidgetBlueprintCompiler::CanCompile(const UBlueprint* Blueprint)
+{
+	return Cast<UWidgetBlueprint>(Blueprint) != nullptr;
+}
+
+
+void FWidgetBlueprintCompiler::PreCompile(UBlueprint* Blueprint, const FKismetCompilerOptions& CompileOptions)
+{
+	if (ReRegister == nullptr
+		&& CanCompile(Blueprint)
+		&& (CompileOptions.CompileType == EKismetCompileType::Full || CompileOptions.CompileType == EKismetCompileType::Cpp))
+	{
+		ReRegister = new TComponentReregisterContext<UWidgetComponent>();
+	}
+
+	CompileCount++;
+}
+
+void FWidgetBlueprintCompiler::Compile(UBlueprint * Blueprint, const FKismetCompilerOptions & CompileOptions, FCompilerResultsLog & Results, TArray<UObject*>* ObjLoaded)
+{
+	if (UWidgetBlueprint* WidgetBlueprint = CastChecked<UWidgetBlueprint>(Blueprint))
+	{
+		FWidgetBlueprintCompilerContext Compiler(WidgetBlueprint, Results, CompileOptions, ObjLoaded);
+		Compiler.Compile();
+		check(Compiler.NewClass);
+	}
+}
+
+void FWidgetBlueprintCompiler::PostCompile(UBlueprint* Blueprint, const FKismetCompilerOptions& CompileOptions)
+{
+	CompileCount--;
+
+	if (CompileCount == 0 && ReRegister)
+	{
+		delete ReRegister;
+		ReRegister = nullptr;
+
+		if (GIsEditor && GEditor)
+		{
+			GEditor->RedrawAllViewports(true);
+		}
+	}
+}
+
+bool FWidgetBlueprintCompiler::GetBlueprintTypesForClass(UClass* ParentClass, UClass*& OutBlueprintClass, UClass*& OutBlueprintGeneratedClass) const
+{
+	if (ParentClass == UUserWidget::StaticClass() || ParentClass->IsChildOf(UUserWidget::StaticClass()))
+	{
+		OutBlueprintClass = UWidgetBlueprint::StaticClass();
+		OutBlueprintGeneratedClass = UWidgetBlueprintGeneratedClass::StaticClass();
+		return true;
+	}
+
+	return false;
+}
+
+FWidgetBlueprintCompilerContext::FWidgetBlueprintCompilerContext(UWidgetBlueprint* SourceSketch, FCompilerResultsLog& InMessageLog, const FKismetCompilerOptions& InCompilerOptions, TArray<UObject*>* InObjLoaded)
 	: Super(SourceSketch, InMessageLog, InCompilerOptions, InObjLoaded)
 	, NewWidgetBlueprintClass(nullptr)
 	, WidgetSchema(nullptr)
 {
 }
 
-FWidgetBlueprintCompiler::~FWidgetBlueprintCompiler()
+FWidgetBlueprintCompilerContext::~FWidgetBlueprintCompilerContext()
 {
 }
 
-UEdGraphSchema_K2* FWidgetBlueprintCompiler::CreateSchema()
+UEdGraphSchema_K2* FWidgetBlueprintCompilerContext::CreateSchema()
 {
 	WidgetSchema = NewObject<UWidgetGraphSchema>();
 	return WidgetSchema;
 }
 
-void FWidgetBlueprintCompiler::CreateFunctionList()
+void FWidgetBlueprintCompilerContext::CreateFunctionList()
 {
 	Super::CreateFunctionList();
 
@@ -112,7 +177,7 @@ void FWidgetBlueprintCompiler::CreateFunctionList()
 	}
 }
 
-void FWidgetBlueprintCompiler::ValidateWidgetNames()
+void FWidgetBlueprintCompilerContext::ValidateWidgetNames()
 {
 	UWidgetBlueprint* WidgetBP = WidgetBlueprint();
 
@@ -143,7 +208,7 @@ struct FCullTemplateObjectsHelper
 };
 
 
-void FWidgetBlueprintCompiler::CleanAndSanitizeClass(UBlueprintGeneratedClass* ClassToClean, UObject*& InOutOldCDO)
+void FWidgetBlueprintCompilerContext::CleanAndSanitizeClass(UBlueprintGeneratedClass* ClassToClean, UObject*& InOutOldCDO)
 {
 	UWidgetBlueprint* WidgetBP = WidgetBlueprint();
 
@@ -191,7 +256,7 @@ void FWidgetBlueprintCompiler::CleanAndSanitizeClass(UBlueprintGeneratedClass* C
 	NewWidgetBlueprintClass->Bindings.Empty();
 }
 
-void FWidgetBlueprintCompiler::SaveSubObjectsFromCleanAndSanitizeClass(FSubobjectCollection& SubObjectsToSave, UBlueprintGeneratedClass* ClassToClean)
+void FWidgetBlueprintCompilerContext::SaveSubObjectsFromCleanAndSanitizeClass(FSubobjectCollection& SubObjectsToSave, UBlueprintGeneratedClass* ClassToClean)
 {
 	Super::SaveSubObjectsFromCleanAndSanitizeClass(SubObjectsToSave, ClassToClean);
 
@@ -206,7 +271,7 @@ void FWidgetBlueprintCompiler::SaveSubObjectsFromCleanAndSanitizeClass(FSubobjec
 	SubObjectsToSave.AddObject(WidgetBP->WidgetTree);
 }
 
-void FWidgetBlueprintCompiler::CreateClassVariablesFromBlueprint()
+void FWidgetBlueprintCompilerContext::CreateClassVariablesFromBlueprint()
 {
 	Super::CreateClassVariablesFromBlueprint();
 
@@ -318,7 +383,7 @@ void FWidgetBlueprintCompiler::CreateClassVariablesFromBlueprint()
 	}
 }
 
-void FWidgetBlueprintCompiler::CopyTermDefaultsToDefaultObject(UObject* DefaultObject)
+void FWidgetBlueprintCompilerContext::CopyTermDefaultsToDefaultObject(UObject* DefaultObject)
 {
 	FKismetCompilerContext::CopyTermDefaultsToDefaultObject(DefaultObject);
 
@@ -388,7 +453,7 @@ void FWidgetBlueprintCompiler::CopyTermDefaultsToDefaultObject(UObject* DefaultO
 
 }
 
-bool FWidgetBlueprintCompiler::CanAllowTemplate(FCompilerResultsLog& MessageLog, UWidgetBlueprintGeneratedClass* InClass)
+bool FWidgetBlueprintCompilerContext::CanAllowTemplate(FCompilerResultsLog& MessageLog, UWidgetBlueprintGeneratedClass* InClass)
 {
 	if ( InClass == nullptr )
 	{
@@ -437,7 +502,7 @@ bool FWidgetBlueprintCompiler::CanAllowTemplate(FCompilerResultsLog& MessageLog,
 	return true;
 }
 
-bool FWidgetBlueprintCompiler::CanTemplateWidget(FCompilerResultsLog& MessageLog, UUserWidget* ThisWidget, TArray<FText>& OutErrors)
+bool FWidgetBlueprintCompilerContext::CanTemplateWidget(FCompilerResultsLog& MessageLog, UUserWidget* ThisWidget, TArray<FText>& OutErrors)
 {
 	UWidgetBlueprintGeneratedClass* WidgetClass = Cast<UWidgetBlueprintGeneratedClass>(ThisWidget->GetClass());
 	if ( WidgetClass == nullptr )
@@ -456,7 +521,7 @@ bool FWidgetBlueprintCompiler::CanTemplateWidget(FCompilerResultsLog& MessageLog
 	return ThisWidget->VerifyTemplateIntegrity(OutErrors);
 }
 
-void FWidgetBlueprintCompiler::FinishCompilingClass(UClass* Class)
+void FWidgetBlueprintCompilerContext::FinishCompilingClass(UClass* Class)
 {
 	UWidgetBlueprint* WidgetBP = WidgetBlueprint();
 	UWidgetBlueprintGeneratedClass* BPGClass = CastChecked<UWidgetBlueprintGeneratedClass>(Class);
@@ -602,7 +667,7 @@ void FWidgetBlueprintCompiler::FinishCompilingClass(UClass* Class)
 	Super::FinishCompilingClass(Class);
 }
 
-void FWidgetBlueprintCompiler::PostCompile()
+void FWidgetBlueprintCompilerContext::PostCompile()
 {
 	Super::PostCompile();
 
@@ -649,7 +714,7 @@ void FWidgetBlueprintCompiler::PostCompile()
 	}
 }
 
-void FWidgetBlueprintCompiler::EnsureProperGeneratedClass(UClass*& TargetUClass)
+void FWidgetBlueprintCompilerContext::EnsureProperGeneratedClass(UClass*& TargetUClass)
 {
 	if ( TargetUClass && !( (UObject*)TargetUClass )->IsA(UWidgetBlueprintGeneratedClass::StaticClass()) )
 	{
@@ -658,7 +723,7 @@ void FWidgetBlueprintCompiler::EnsureProperGeneratedClass(UClass*& TargetUClass)
 	}
 }
 
-void FWidgetBlueprintCompiler::SpawnNewClass(const FString& NewClassName)
+void FWidgetBlueprintCompilerContext::SpawnNewClass(const FString& NewClassName)
 {
 	NewWidgetBlueprintClass = FindObject<UWidgetBlueprintGeneratedClass>(Blueprint->GetOutermost(), *NewClassName);
 
@@ -674,19 +739,19 @@ void FWidgetBlueprintCompiler::SpawnNewClass(const FString& NewClassName)
 	NewClass = NewWidgetBlueprintClass;
 }
 
-void FWidgetBlueprintCompiler::OnNewClassSet(UBlueprintGeneratedClass* ClassToUse)
+void FWidgetBlueprintCompilerContext::OnNewClassSet(UBlueprintGeneratedClass* ClassToUse)
 {
 	NewWidgetBlueprintClass = CastChecked<UWidgetBlueprintGeneratedClass>(ClassToUse);
 }
 
-void FWidgetBlueprintCompiler::PrecompileFunction(FKismetFunctionContext& Context, EInternalCompilerFlags InternalFlags)
+void FWidgetBlueprintCompilerContext::PrecompileFunction(FKismetFunctionContext& Context, EInternalCompilerFlags InternalFlags)
 {
 	Super::PrecompileFunction(Context, InternalFlags);
 
 	VerifyEventReplysAreNotEmpty(Context);
 }
 
-void FWidgetBlueprintCompiler::VerifyEventReplysAreNotEmpty(FKismetFunctionContext& Context)
+void FWidgetBlueprintCompilerContext::VerifyEventReplysAreNotEmpty(FKismetFunctionContext& Context)
 {
 	TArray<UK2Node_FunctionResult*> FunctionResults;
 	Context.SourceGraph->GetNodesOfClass<UK2Node_FunctionResult>(FunctionResults);
@@ -710,7 +775,7 @@ void FWidgetBlueprintCompiler::VerifyEventReplysAreNotEmpty(FKismetFunctionConte
 	}
 }
 
-bool FWidgetBlueprintCompiler::ValidateGeneratedClass(UBlueprintGeneratedClass* Class)
+bool FWidgetBlueprintCompilerContext::ValidateGeneratedClass(UBlueprintGeneratedClass* Class)
 {
 	bool SuperResult = Super::ValidateGeneratedClass(Class);
 	bool Result = UWidgetBlueprint::ValidateGeneratedClass(Class);
diff --git a/Engine/Source/Editor/UMGEditor/Private/WidgetBlueprintCompiler.h b/Engine/Source/Editor/UMGEditor/Private/WidgetBlueprintCompiler.h
deleted file mode 100644
diff --git a/Engine/Source/Editor/UMGEditor/Private/WidgetBlueprintEditor.h b/Engine/Source/Editor/UMGEditor/Private/WidgetBlueprintEditor.h
deleted file mode 100644
diff --git a/Engine/Source/Editor/UMGEditor/Public/UMGEditorModule.h b/Engine/Source/Editor/UMGEditor/Public/UMGEditorModule.h
--- a/Engine/Source/Editor/UMGEditor/Public/UMGEditorModule.h
+++ b/Engine/Source/Editor/UMGEditor/Public/UMGEditorModule.h
@@ -9,10 +9,11 @@
 extern const FName UMGEditorAppIdentifier;
 
 class FUMGEditor;
+class FWidgetBlueprintCompiler;
 
 /** The public interface of the UMG editor module. */
 class IUMGEditorModule : public IModuleInterface, public IHasMenuExtensibility, public IHasToolBarExtensibility
 {
 public:
-
+	virtual FWidgetBlueprintCompiler* GetRegisteredCompiler() = 0;
 };
diff --git a/Engine/Source/Editor/UMGEditor/Public/WidgetBlueprintCompiler.h b/Engine/Source/Editor/UMGEditor/Public/WidgetBlueprintCompiler.h
new file mode 100644
--- /dev/null
+++ b/Engine/Source/Editor/UMGEditor/Public/WidgetBlueprintCompiler.h
@@ -0,0 +1,94 @@
+// Copyright 1998-2018 Epic Games, Inc. All Rights Reserved.
+
+#pragma once
+
+#include "CoreMinimal.h"
+#include "WidgetBlueprint.h"
+#include "KismetCompiler.h"
+#include "KismetCompilerModule.h"
+#include "ComponentReregisterContext.h"
+#include "Components/WidgetComponent.h"
+
+
+//////////////////////////////////////////////////////////////////////////
+// FWidgetBlueprintCompiler 
+
+class UMGEDITOR_API FWidgetBlueprintCompiler : public IBlueprintCompiler
+{
+
+public:
+	FWidgetBlueprintCompiler();
+
+	bool CanCompile(const UBlueprint* Blueprint) override;
+	void PreCompile(UBlueprint* Blueprint, const FKismetCompilerOptions& CompileOptions) override;
+	void Compile(UBlueprint* Blueprint, const FKismetCompilerOptions& CompileOptions, FCompilerResultsLog& Results, TArray<UObject*>* ObjLoaded) override;
+	void PostCompile(UBlueprint* Blueprint, const FKismetCompilerOptions& CompileOptions) override;
+	bool GetBlueprintTypesForClass(UClass* ParentClass, UClass*& OutBlueprintClass, UClass*& OutBlueprintGeneratedClass) const override;
+
+private:
+
+	/** The temporary variable that captures and reinstances components after compiling finishes. */
+	TComponentReregisterContext<UWidgetComponent>* ReRegister;
+
+	/**
+	* The current count on the number of compiles that have occurred.  We don't want to re-register components until all
+	* compiling has stopped.
+	*/
+	int32 CompileCount;
+
+};
+
+//////////////////////////////////////////////////////////////////////////
+// FWidgetBlueprintCompilerContext
+
+
+class UMGEDITOR_API FWidgetBlueprintCompilerContext : public FKismetCompilerContext
+{
+protected:
+	typedef FKismetCompilerContext Super;
+
+public:
+	FWidgetBlueprintCompilerContext(UWidgetBlueprint* SourceSketch, FCompilerResultsLog& InMessageLog, const FKismetCompilerOptions& InCompilerOptions, TArray<UObject*>* InObjLoaded);
+	virtual ~FWidgetBlueprintCompilerContext();
+
+	static bool CanAllowTemplate(FCompilerResultsLog& MessageLog, UWidgetBlueprintGeneratedClass* InClass);
+	static bool CanTemplateWidget(FCompilerResultsLog& MessageLog, UUserWidget* ThisWidget, TArray<FText>& OutErrors);
+
+protected:
+	UWidgetBlueprint* WidgetBlueprint() const { return Cast<UWidgetBlueprint>(Blueprint); }
+
+	void ValidateWidgetNames();
+
+	// FKismetCompilerContext
+	virtual UEdGraphSchema_K2* CreateSchema() override;
+	virtual void CreateFunctionList() override;
+	virtual void SpawnNewClass(const FString& NewClassName) override;
+	virtual void OnNewClassSet(UBlueprintGeneratedClass* ClassToUse) override;
+	virtual void PrecompileFunction(FKismetFunctionContext& Context, EInternalCompilerFlags InternalFlags) override;
+	virtual void CleanAndSanitizeClass(UBlueprintGeneratedClass* ClassToClean, UObject*& InOutOldCDO) override;
+	virtual void SaveSubObjectsFromCleanAndSanitizeClass(FSubobjectCollection& SubObjectsToSave, UBlueprintGeneratedClass* ClassToClean) override;
+	virtual void EnsureProperGeneratedClass(UClass*& TargetClass) override;
+	virtual void CreateClassVariablesFromBlueprint() override;
+	virtual void CopyTermDefaultsToDefaultObject(UObject* DefaultObject);
+	virtual void FinishCompilingClass(UClass* Class) override;
+	virtual bool ValidateGeneratedClass(UBlueprintGeneratedClass* Class) override;
+	virtual void PostCompile() override;
+	// End FKismetCompilerContext
+
+	void VerifyEventReplysAreNotEmpty(FKismetFunctionContext& Context);
+
+protected:
+
+	UWidgetBlueprintGeneratedClass* NewWidgetBlueprintClass;
+
+	class UWidgetGraphSchema* WidgetSchema;
+
+	// Map of properties created for widgets; to aid in debug data generation
+	TMap<class UWidget*, class UProperty*> WidgetToMemberVariableMap;
+
+	// Map of properties created for widget animations; to aid in debug data generation
+	TMap<class UWidgetAnimation*, class UProperty*> WidgetAnimToMemberVariableMap;
+
+	///----------------------------------------------------------------
+};
+
diff --git a/Engine/Source/Editor/UMGEditor/Public/WidgetBlueprintEditor.h b/Engine/Source/Editor/UMGEditor/Public/WidgetBlueprintEditor.h
new file mode 100644
--- /dev/null
+++ b/Engine/Source/Editor/UMGEditor/Public/WidgetBlueprintEditor.h
@@ -0,0 +1,356 @@
+// Copyright 1998-2018 Epic Games, Inc. All Rights Reserved.
+
+#pragma once
+
+#include "CoreMinimal.h"
+#include "Misc/Guid.h"
+#include "Widgets/SWidget.h"
+#include "Framework/Commands/UICommandList.h"
+#include "Framework/MultiBox/MultiBoxExtender.h"
+#include "AssetData.h"
+#include "PreviewScene.h"
+#include "GraphEditor.h"
+#include "BlueprintEditor.h"
+#include "ISequencer.h"
+#include "WidgetReference.h"
+#include "Blueprint/UserWidget.h"
+
+class FMenuBuilder;
+class FWidgetBlueprintEditorToolbar;
+class IMessageLogListing;
+class STextBlock;
+class UPanelSlot;
+class UWidgetAnimation;
+class UWidgetBlueprint;
+
+struct FNamedSlotSelection
+{
+	FWidgetReference NamedSlotHostWidget;
+	FName SlotName;
+};
+
+/**
+ * widget blueprint editor (extends Blueprint editor)
+ */
+class UMGEDITOR_API FWidgetBlueprintEditor : public FBlueprintEditor
+{
+public:
+	DECLARE_MULTICAST_DELEGATE_OneParam(FOnHoveredWidgetSet, const FWidgetReference&)
+	DECLARE_MULTICAST_DELEGATE(FOnHoveredWidgetCleared);
+
+	DECLARE_MULTICAST_DELEGATE(FOnSelectedWidgetsChanging)
+	DECLARE_MULTICAST_DELEGATE(FOnSelectedWidgetsChanged)
+
+	/** Called after the widget preview has been updated */
+	DECLARE_MULTICAST_DELEGATE(FOnWidgetPreviewUpdated)
+
+	DECLARE_EVENT(FWidgetBlueprintEditor, FOnEnterWidgetDesigner)
+
+public:
+	FWidgetBlueprintEditor();
+
+	virtual ~FWidgetBlueprintEditor();
+
+	void InitWidgetBlueprintEditor(const EToolkitMode::Type Mode, const TSharedPtr< IToolkitHost >& InitToolkitHost, const TArray<UBlueprint*>& InBlueprints, bool bShouldOpenInDefaultsMode);
+
+	/** FBlueprintEditor interface */
+	virtual void Tick(float DeltaTime) override;
+	virtual void PostUndo(bool bSuccessful) override;
+	virtual void PostRedo(bool bSuccessful) override;
+	virtual void Compile() override;
+
+	/** FGCObjectInterface */
+	virtual void AddReferencedObjects( FReferenceCollector& Collector ) override;
+
+	/** @return The widget blueprint currently being edited in this editor */
+	class UWidgetBlueprint* GetWidgetBlueprintObj() const;
+
+	/** @return The preview widget. */
+	UUserWidget* GetPreview() const;
+
+	/** @return The preview scene that owns the preview widget. */
+	FPreviewScene* GetPreviewScene();
+
+	/**  */
+	bool IsSimulating() const;
+
+	/**  */
+	void SetIsSimulating(bool bSimulating);
+
+	/** Causes the preview to be destroyed and a new one to be created next tick */
+	void InvalidatePreview(bool bViewOnly = false);
+
+	/** Immediately rebuilds the preview widget. */
+	void RefreshPreview();
+
+	/** Creates a widget reference using the template. */
+	FWidgetReference GetReferenceFromTemplate(UWidget* TemplateWidget);
+
+	/** Creates a widget reference using the preview.  Which is used to lookup the stable template pointer. */
+	FWidgetReference GetReferenceFromPreview(UWidget* PreviewWidget);
+
+	/** @return The sequencer used to create widget animations */
+	TSharedPtr<ISequencer>& GetSequencer();
+
+	/** Changes the currently viewed animation in Sequencer to the new one*/
+	void ChangeViewedAnimation( UWidgetAnimation& InAnimationToView );
+
+	/** Get the current animation*/
+	UWidgetAnimation* GetCurrentAnimation() { return CurrentAnimation.Get(); }
+
+	/** Updates the current animation if it is invalid */
+	const UWidgetAnimation* RefreshCurrentAnimation();
+
+	/** Sets the currently selected set of widgets */
+	void SelectWidgets(const TSet<FWidgetReference>& Widgets, bool bAppendOrToggle);
+
+	/** Sets the currently selected set of objects */
+	void SelectObjects(const TSet<UObject*>& Objects);
+
+	/** Sets the selected named slot */
+	void SetSelectedNamedSlot(TOptional<FNamedSlotSelection> SelectedNamedSlot);
+
+	/** Removes removed widgets from the selection set. */
+	void CleanSelection();
+
+	/** @return The selected set of widgets */
+	const TSet<FWidgetReference>& GetSelectedWidgets() const;
+
+	/** @return The selected named slot */
+	TOptional<FNamedSlotSelection> GetSelectedNamedSlot() const;
+
+	/** @return The selected set of Objects */
+	const TSet< TWeakObjectPtr<UObject> >& GetSelectedObjects() const;
+
+	/** @return the selected template widget */
+	const TWeakObjectPtr<UClass> GetSelectedTemplate() const { return SelectedTemplate; }
+
+	/** @return the selected user widget */
+	const FAssetData GetSelectedUserWidget() const { return SelectedUserWidget; }
+
+	/** Set the selected template widget */
+	void SetSelectedTemplate(TWeakObjectPtr<UClass> TemplateClass) { SelectedTemplate = TemplateClass; }
+
+	/** Set the selected user widget */
+	void SetSelectedUserWidget(FAssetData InSelectedUserWidget) { SelectedUserWidget = InSelectedUserWidget; }
+
+	TSharedPtr<class FWidgetBlueprintEditorToolbar> GetWidgetToolbarBuilder() { return WidgetToolbar; }
+
+	/** Migrate a property change from the preview GUI to the template GUI. */
+	void MigrateFromChain(FEditPropertyChain* PropertyThatChanged, bool bIsModify);
+
+	/** Event called when an undo/redo transaction occurs */
+	DECLARE_EVENT(FWidgetBlueprintEditor, FOnWidgetBlueprintTransaction)
+	FOnWidgetBlueprintTransaction& GetOnWidgetBlueprintTransaction() { return OnWidgetBlueprintTransaction; }
+
+	/** Creates a sequencer widget */
+	TSharedRef<SWidget> CreateSequencerWidget();
+
+	/**
+	 * The widget we're now hovering over in any particular context, allows multiple views to 
+	 * synchronize feedback on where that widget is in their representation.
+	 */
+	void SetHoveredWidget(FWidgetReference& InHoveredWidget);
+
+	void ClearHoveredWidget();
+
+	/** @return The widget that is currently being hovered over (either in the designer or hierarchy) */
+	const FWidgetReference& GetHoveredWidget() const;
+
+	void AddPostDesignerLayoutAction(TFunction<void()> Action);
+
+	void OnEnteringDesigner();
+
+	TArray< TFunction<void()> >& GetQueuedDesignerActions();
+
+	/** Get the current designer flags that are in effect for the current user widget we're editing. */
+	EWidgetDesignFlags::Type GetCurrentDesignerFlags() const;
+
+	bool GetShowDashedOutlines() const;
+	void SetShowDashedOutlines(bool Value);
+
+	bool GetIsRespectingLocks() const;
+	void SetIsRespectingLocks(bool Value);
+
+public:
+	/** Fires whenever a new widget is being hovered over */
+	FOnHoveredWidgetSet OnHoveredWidgetSet;
+
+	/** Fires when there is no longer any widget being hovered over */
+	FOnHoveredWidgetCleared OnHoveredWidgetCleared;	
+
+	/** Fires whenever the selected set of widgets changing */
+	FOnSelectedWidgetsChanged OnSelectedWidgetsChanging;
+
+	/** Fires whenever the selected set of widgets changes */
+	FOnSelectedWidgetsChanged OnSelectedWidgetsChanged;
+
+	/** Notification for when the preview widget has been updated */
+	FOnWidgetPreviewUpdated OnWidgetPreviewUpdated;
+
+	/** Fires after the mode change to Designer*/
+	FOnEnterWidgetDesigner OnEnterWidgetDesigner;
+
+	/** Command list for handling widget actions in the WidgetBlueprintEditor */
+	TSharedPtr< FUICommandList > DesignerCommandList;
+
+	/** Paste Metadata */
+	FVector2D PasteDropLocation;
+
+protected:
+	// Begin FBlueprintEditor
+	virtual void RegisterApplicationModes(const TArray<UBlueprint*>& InBlueprints, bool bShouldOpenInDefaultsMode, bool bNewlyCreated = false) override;
+	virtual FGraphAppearanceInfo GetGraphAppearance(class UEdGraph* InGraph) const override;
+	virtual void AppendExtraCompilerResults(TSharedPtr<class IMessageLogListing> ResultsListing) override;
+	virtual TSubclassOf<UEdGraphSchema> GetDefaultSchemaClass() const override;
+	// End FBlueprintEditor
+
+private:
+	bool CanDeleteSelectedWidgets();
+	void DeleteSelectedWidgets();
+
+	bool CanCopySelectedWidgets();
+	void CopySelectedWidgets();
+
+	bool CanPasteWidgets();
+	void PasteWidgets();
+
+	bool CanCutSelectedWidgets();
+	void CutSelectedWidgets();
+
+	bool CanDuplicateSelectedWidgets();
+	void DuplicateSelectedWidgets();
+
+private:
+	/** Called whenever the blueprint is structurally changed. */
+	virtual void OnBlueprintChangedImpl(UBlueprint* InBlueprint, bool bIsJustBeingCompiled = false ) override;
+
+	/** Called when objects need to be swapped out for new versions, like after a blueprint recompile. */
+	void OnObjectsReplaced(const TMap<UObject*, UObject*>& ReplacementMap);
+
+	/** Destroy the current preview GUI object */
+	void DestroyPreview();
+
+	/** Tick the current preview GUI object */
+	void UpdatePreview(UBlueprint* InBlueprint, bool bInForceFullUpdate);
+
+	/** Populates the sequencer add menu. */
+	void OnGetAnimationAddMenuContent(FMenuBuilder& MenuBuilder, TSharedRef<ISequencer> Sequencer);
+
+	/** Populates sequencer menu when added with objects. Used to handle case where widget is deleted so it's 
+	    Object * is null.*/
+	void OnBuildCustomContextMenuForGuid(FMenuBuilder& MenuBuilder,FGuid ObjectBinding);
+
+	/** Populates the sequencer add submenu for the big list of widgets. */
+	void OnGetAnimationAddMenuContentAllWidgets(FMenuBuilder& MenuBuilder);
+
+	/** Adds the supplied UObject to the current animation. */
+	void AddObjectToAnimation(UObject* ObjectToAnimate);
+
+	/** Gets the extender to use for sequencers context sensitive menus and toolbars. */
+	TSharedRef<FExtender> GetAddTrackSequencerExtender(const TSharedRef<FUICommandList> CommandList, const TArray<UObject*> ContextSensitiveObjects);
+
+	/** Extends the sequencer add track menu. */
+	void ExtendSequencerAddTrackMenu( FMenuBuilder& AddTrackMenuBuilder, const TArray<UObject*> ContextObjects );
+
+	/** Replace track with selected widget function */
+	void ReplaceTrackWithSelectedWidget(FWidgetReference SelectedWidget, UWidget* BoundWidget, FGuid ObjectBinding);
+
+	/** Extends the sequencer add track menu. */
+	void ExtendSequencerObjectBindingMenu(FMenuBuilder& ObjectBindingMenuBuilder, const TArray<UObject*> ContextObjects);
+
+	/** Add an animation track for the supplied slot to the current animation. */
+	void AddSlotTrack( UPanelSlot* Slot );
+
+	/** Add an animation track for the supplied material property path to the current animation. */
+	void AddMaterialTrack( UWidget* Widget, TArray<UProperty*> MaterialPropertyPath, FText MaterialPropertyDisplayName );
+
+	/** Handler which is called whenever sequencer movie scene data changes. */
+	void OnMovieSceneDataChanged(EMovieSceneDataChangeType DataChangeType);
+
+	/** Handler which is called whenever sequencer binding is pasted. */
+	void OnMovieSceneBindingsPasted(const TArray<FMovieSceneBinding>& BindingsPasted);
+
+	/** Fire off when sequencer selection changed */
+	void SyncSelectedWidgetsWithSequencerSelection(TArray<FGuid> ObjectGuids);
+
+	/** Get the animation playback context */
+	UObject* GetAnimationPlaybackContext() const { return GetPreview(); }
+
+	/** Get the animation playback event contexts */
+	TArray<UObject*> GetAnimationEventContexts() const { TArray<UObject*> EventContexts; EventContexts.Add(GetPreview()); return EventContexts; }
+	
+private:
+	/** The preview scene that owns the preview GUI */
+	FPreviewScene PreviewScene;
+
+	/** Sequencer for creating and previewing widget animations */
+	TSharedPtr<ISequencer> Sequencer;
+
+	/** Overlay used to display UI on top of sequencer */
+	TWeakPtr<SOverlay> SequencerOverlay;
+
+	/** A text block which is displayed in the overlay when no animation is selected. */
+	TWeakPtr<STextBlock> NoAnimationTextBlock;
+
+	/** The Blueprint associated with the current preview */
+	UWidgetBlueprint* PreviewBlueprint;
+
+	/** The currently selected preview widgets in the preview GUI */
+	TSet< FWidgetReference > SelectedWidgets;
+
+	/** The currently selected objects in the designer */
+	TSet< TWeakObjectPtr<UObject> > SelectedObjects;
+
+	/** The last selected template widget in the palette view */
+	TWeakObjectPtr<UClass> SelectedTemplate;
+
+	/** AssetData of Selected UserWidget */
+	FAssetData SelectedUserWidget;
+
+	/** The currently selected named slot */
+	TOptional<FNamedSlotSelection> SelectedNamedSlot;
+
+	/** The preview GUI object */
+	mutable TWeakObjectPtr<UUserWidget> PreviewWidgetPtr;
+
+	/** Delegate called when a undo/redo transaction happens */
+	FOnWidgetBlueprintTransaction OnWidgetBlueprintTransaction;
+
+	/** The toolbar builder associated with this editor */
+	TSharedPtr<class FWidgetBlueprintEditorToolbar> WidgetToolbar;
+
+	/** The widget references out in the ether that may need to be updated after being issued. */
+	TArray< TWeakPtr<FWidgetHandle> > WidgetHandlePool;
+
+	/** The widget currently being hovered over */
+	FWidgetReference HoveredWidget;
+
+	/** The preview becomes invalid and needs to be rebuilt on the next tick. */
+	bool bPreviewInvalidated;
+
+	/**  */
+	bool bIsSimulateEnabled;
+
+	/**  */
+	bool bIsRealTime;
+
+	/** Should the designer show outlines when it creates widgets? */
+	bool bShowDashedOutlines;
+
+	/**  */
+	bool bRespectLocks;
+
+	TArray< TFunction<void()> > QueuedDesignerActions;
+
+	/** The currently viewed animation, if any. */
+	TWeakObjectPtr<UWidgetAnimation> CurrentAnimation;
+
+	FDelegateHandle SequencerAddTrackExtenderHandle;
+
+	/** Messages we want to append to the compiler results. */
+	TArray< TSharedRef<class FTokenizedMessage> > DesignerCompilerMessages;
+
+	/** When true the animation data in the generated class should be replaced with the current animation data. */
+	bool bRefreshGeneratedClassAnimations;
+};
diff --git a/Engine/Source/Editor/UMGEditor/UMGEditor.Build.cs b/Engine/Source/Editor/UMGEditor/UMGEditor.Build.cs
--- a/Engine/Source/Editor/UMGEditor/UMGEditor.Build.cs
+++ b/Engine/Source/Editor/UMGEditor/UMGEditor.Build.cs
@@ -28,6 +28,12 @@ public class UMGEditor : ModuleRules
 				"UMG",
 			});
 
+        PublicDependencyModuleNames.AddRange(
+            new string[]
+            {
+                "Sequencer",
+            });
+
 		PrivateDependencyModuleNames.AddRange(
 			new string[] {
 				"Core",
@@ -52,7 +58,6 @@ public class UMGEditor : ModuleRules
 				"MovieScene",
 				"MovieSceneTools",
                 "MovieSceneTracks",
-				"Sequencer",
 				"DetailCustomizations",
                 "Settings",
 				"RenderCore",
