Fix for copy and paste causing multiple transations.

diff --git a/Engine/Source/Editor/UMGEditor/Private/WidgetBlueprintEditorUtils.cpp b/Engine/Source/Editor/UMGEditor/Private/WidgetBlueprintEditorUtils.cpp
--- a/Engine/Source/Editor/UMGEditor/Private/WidgetBlueprintEditorUtils.cpp
+++ b/Engine/Source/Editor/UMGEditor/Private/WidgetBlueprintEditorUtils.cpp
@@ -1160,29 +1160,34 @@ TArray<UWidget*> FWidgetBlueprintEditorUtils::PasteWidgetsInternal(TSharedRef<FW
 						}
 					}
 
-					BlueprintEditor->AddPostDesignerLayoutAction(
-						[=] {
-						FWidgetReference WidgetRef = BlueprintEditor->GetReferenceFromTemplate(NewWidget);
-						UPanelSlot* PreviewSlot = WidgetRef.GetPreview()->Slot;
-						UPanelSlot* TemplateSlot = WidgetRef.GetTemplate()->Slot;
+					BlueprintEditor->RefreshPreview();
+						
+					FWidgetReference WidgetRef = BlueprintEditor->GetReferenceFromTemplate(NewWidget);
 
-						if ( UCanvasPanelSlot* CanvasSlot = Cast<UCanvasPanelSlot>(PreviewSlot) )
+					UPanelSlot* PreviewSlot = WidgetRef.GetPreview()->Slot;
+					UPanelSlot* TemplateSlot = WidgetRef.GetTemplate()->Slot;
+
+					if ( UCanvasPanelSlot* CanvasSlot = Cast<UCanvasPanelSlot>(PreviewSlot) )
+					{
+						FVector2D PasteOffset = FVector2D(0, 0);
+						if (bShouldReproduceOffsets)
 						{
-							FVector2D PasteOffset = FVector2D(0, 0);
-							if (bShouldReproduceOffsets)
-							{
-								PasteOffset = CanvasSlot->GetPosition()- FirstWidgetPosition;
-							}
+							PasteOffset = CanvasSlot->GetPosition()- FirstWidgetPosition;
+						}
 
-							CanvasSlot->SaveBaseLayout();
-							CanvasSlot->SetDesiredPosition(PasteLocation + PasteOffset);
-							CanvasSlot->RebaseLayout();
+						if (UCanvasPanel* Canvas = Cast<UCanvasPanel>(CanvasSlot->Parent))
+						{
+							Canvas->TakeWidget(); // Generate the underlying widget so redoing the layout below works.
 						}
 
-						TMap<FName, FString> SlotProperties;
-						FWidgetBlueprintEditorUtils::ExportPropertiesToText(PreviewSlot, SlotProperties);
-						FWidgetBlueprintEditorUtils::ImportPropertiesFromText(TemplateSlot, SlotProperties);
-					});
+						CanvasSlot->SaveBaseLayout();
+						CanvasSlot->SetDesiredPosition(PasteLocation + PasteOffset);
+						CanvasSlot->RebaseLayout();
+					}
+
+					TMap<FName, FString> SlotProperties;
+					FWidgetBlueprintEditorUtils::ExportPropertiesToText(PreviewSlot, SlotProperties);
+					FWidgetBlueprintEditorUtils::ImportPropertiesFromText(TemplateSlot, SlotProperties);
 				}
 			}
 
