PR #4870: Exclude abstract classes in ClassPin of CreateWidget (Contributed by [REDACTED_PERSON])

diff --git a/Engine/Source/Editor/UMGEditor/Private/Nodes/K2Node_CreateWidget.cpp b/Engine/Source/Editor/UMGEditor/Private/Nodes/K2Node_CreateWidget.cpp
--- a/Engine/Source/Editor/UMGEditor/Private/Nodes/K2Node_CreateWidget.cpp
+++ b/Engine/Source/Editor/UMGEditor/Private/Nodes/K2Node_CreateWidget.cpp
@@ -13,6 +13,7 @@
 #include "KismetCompilerMisc.h"
 #include "KismetCompiler.h"
 #include "EditorCategoryUtils.h"
+#include "Widgets/SGraphNodeCreateWidget.h"
 
 
 #define LOCTEXT_NAMESPACE "UMG"
@@ -125,7 +126,7 @@ void UK2Node_CreateWidget::ExpandNode(class FKismetCompilerContext& CompilerCont
 	}
 	else
 	{
-		// Copy blueprint literal onto 'UWidgetBlueprintLibrary::Create' call 
+		// Copy blueprint literal onto 'UWidgetBlueprintLibrary::Create' call
 		CallCreateWidgetTypePin->DefaultObject = SpawnClass;
 	}
 
@@ -155,4 +156,9 @@ void UK2Node_CreateWidget::ExpandNode(class FKismetCompilerContext& CompilerCont
 	CreateWidgetNode->BreakAllNodeLinks();
 }
 
+TSharedPtr<SGraphNode> UK2Node_CreateWidget::CreateVisualWidget()
+{
+	return SNew(SGraphNodeCreateWidget, this);
+}
+
 #undef LOCTEXT_NAMESPACE
diff --git a/Engine/Source/Editor/UMGEditor/Private/Nodes/K2Node_CreateWidget.h b/Engine/Source/Editor/UMGEditor/Private/Nodes/K2Node_CreateWidget.h
--- a/Engine/Source/Editor/UMGEditor/Private/Nodes/K2Node_CreateWidget.h
+++ b/Engine/Source/Editor/UMGEditor/Private/Nodes/K2Node_CreateWidget.h
@@ -17,6 +17,7 @@ class UMGEDITOR_API UK2Node_CreateWidget : public UK2Node_ConstructObjectFromCla
 	//~ Begin UEdGraphNode Interface.
 	virtual void AllocateDefaultPins() override;
 	virtual void ExpandNode(class FKismetCompilerContext& CompilerContext, UEdGraph* SourceGraph) override;
+	virtual TSharedPtr<SGraphNode> CreateVisualWidget() override;
 	//~ End UEdGraphNode Interface.
 
 	//~ Begin UK2Node Interface
diff --git a/Engine/Source/Editor/UMGEditor/Private/Widgets/SGraphNodeCreateWidget.cpp b/Engine/Source/Editor/UMGEditor/Private/Widgets/SGraphNodeCreateWidget.cpp
new file mode 100644
--- /dev/null
+++ b/Engine/Source/Editor/UMGEditor/Private/Widgets/SGraphNodeCreateWidget.cpp
@@ -0,0 +1,126 @@
+// Copyright 1998-2018 Epic Games, Inc. All Rights Reserved.
+
+#include "Widgets/SGraphNodeCreateWidget.h"
+#include "Modules/ModuleManager.h"
+#include "Widgets/SBoxPanel.h"
+#include "Widgets/Layout/SBox.h"
+#include "Editor.h"
+#include "EdGraphSchema_K2.h"
+#include "Nodes/K2Node_CreateWidget.h"
+#include "KismetPins/SGraphPinObject.h"
+#include "NodeFactory.h"
+#include "ClassViewerModule.h"
+#include "ClassViewerFilter.h"
+#include "ScopedTransaction.h"
+#include "Blueprint/UserWidget.h"
+
+#define LOCTEXT_NAMESPACE "SGraphPinUserWidgetBasedClass"
+
+//////////////////////////////////////////////////////////////////////////
+// SGraphPinUserWidgetBasedClass
+
+/**
+ * GraphPin can select only UUserWidget classes.
+ * Instead of asset picker, a class viewer is used.
+ */
+class SGraphPinUserWidgetBasedClass : public SGraphPinObject
+{
+	class FUserWidgetBasedClassFilter : public IClassViewerFilter
+	{
+	public:
+
+		virtual bool IsClassAllowed(const FClassViewerInitializationOptions& InInitOptions, const UClass* InClass, TSharedRef< FClassViewerFilterFuncs > InFilterFuncs ) override
+		{
+			if(NULL != InClass)
+			{
+				const bool bUserWidgetBased = InClass->IsChildOf(UUserWidget::StaticClass());
+				const bool bBlueprintType = UEdGraphSchema_K2::IsAllowableBlueprintVariableType(InClass);
+				const bool bNotAbstract = !InClass->HasAnyClassFlags(CLASS_Abstract);
+				return bUserWidgetBased && bBlueprintType && bNotAbstract;
+			}
+			return false;
+		}
+
+		virtual bool IsUnloadedClassAllowed(const FClassViewerInitializationOptions& InInitOptions, const TSharedRef< const IUnloadedBlueprintData > InUnloadedClassData, TSharedRef< FClassViewerFilterFuncs > InFilterFuncs) override
+		{
+			const bool bUserWidgetBased = InUnloadedClassData->IsChildOf(UUserWidget::StaticClass());
+			const bool bNotAbstract = !InUnloadedClassData->HasAnyClassFlags(CLASS_Abstract);
+			return bUserWidgetBased && bNotAbstract;
+		}
+	};
+
+protected:
+
+	void OnPickedNewClass(UClass* ChosenClass)
+	{
+		if (GraphPinObj->DefaultObject != ChosenClass)
+		{
+			const FScopedTransaction Transaction(NSLOCTEXT("GraphEditor", "ChangeClassPinValue", "Change Class Pin Value"));
+			GraphPinObj->Modify();
+
+			AssetPickerAnchor->SetIsOpen(false);
+			GraphPinObj->GetSchema()->TrySetDefaultObject(*GraphPinObj, ChosenClass);
+		}
+	}
+
+	virtual TSharedRef<SWidget> GenerateAssetPicker() override
+	{
+		FClassViewerModule& ClassViewerModule = FModuleManager::LoadModuleChecked<FClassViewerModule>("ClassViewer");
+
+		FClassViewerInitializationOptions Options;
+		Options.Mode = EClassViewerMode::ClassPicker;
+		Options.bIsActorsOnly = false;
+		Options.DisplayMode = EClassViewerDisplayMode::DefaultView;
+		Options.bShowUnloadedBlueprints = true;
+		Options.bShowNoneOption = false;
+		Options.bShowObjectRootClass = true;
+		TSharedPtr< FUserWidgetBasedClassFilter > Filter = MakeShareable(new FUserWidgetBasedClassFilter);
+		Options.ClassFilter = Filter;
+
+		return
+			SNew(SBox)
+			.WidthOverride(280)
+			[
+				SNew(SVerticalBox)
+				+SVerticalBox::Slot()
+				.AutoHeight()
+				.MaxHeight(500)
+				[
+					SNew(SBorder)
+					.Padding(4)
+					.BorderImage( FEditorStyle::GetBrush("ToolPanel.GroupBorder") )
+					[
+						ClassViewerModule.CreateClassViewer(Options, FOnClassPicked::CreateSP(this, &SGraphPinUserWidgetBasedClass::OnPickedNewClass))
+					]
+				]
+			];
+	}
+};
+
+//////////////////////////////////////////////////////////////////////////
+// SGraphNodeCreateWidget
+
+void SGraphNodeCreateWidget::CreatePinWidgets()
+{
+	UK2Node_CreateWidget* CreateWidgetNode = CastChecked<UK2Node_CreateWidget>(GraphNode);
+	UEdGraphPin* ClassPin = CreateWidgetNode->GetClassPin();
+
+	for (auto PinIt = GraphNode->Pins.CreateConstIterator(); PinIt; ++PinIt)
+	{
+		UEdGraphPin* CurrentPin = *PinIt;
+		if ((!CurrentPin->bHidden) && (CurrentPin != ClassPin))
+		{
+			TSharedPtr<SGraphPin> NewPin = FNodeFactory::CreatePinWidget(CurrentPin);
+			check(NewPin.IsValid());
+			this->AddPin(NewPin.ToSharedRef());
+		}
+		else if ((ClassPin == CurrentPin) && (!ClassPin->bHidden || (ClassPin->LinkedTo.Num() > 0)))
+		{
+			TSharedPtr<SGraphPinUserWidgetBasedClass> NewPin = SNew(SGraphPinUserWidgetBasedClass, ClassPin);
+			check(NewPin.IsValid());
+			this->AddPin(NewPin.ToSharedRef());
+		}
+	}
+}
+
+#undef LOCTEXT_NAMESPACE
diff --git a/Engine/Source/Editor/UMGEditor/Private/Widgets/SGraphNodeCreateWidget.h b/Engine/Source/Editor/UMGEditor/Private/Widgets/SGraphNodeCreateWidget.h
new file mode 100644
--- /dev/null
+++ b/Engine/Source/Editor/UMGEditor/Private/Widgets/SGraphNodeCreateWidget.h
@@ -0,0 +1,14 @@
+// Copyright 1998-2018 Epic Games, Inc. All Rights Reserved.
+
+#pragma once
+
+#include "CoreMinimal.h"
+#include "KismetNodes/SGraphNodeK2Default.h"
+
+class SGraphNodeCreateWidget : public SGraphNodeK2Default
+{
+public:
+	// SGraphNode interface
+	virtual void CreatePinWidgets() override;
+	// End of SGraphNode interface
+};
