PR #4230: Layout changes to SWidgetDetailsView (Contributed by [REDACTED_PERSON])

diff --git a/Engine/Source/Editor/EditorStyle/Private/SlateEditorStyle.cpp b/Engine/Source/Editor/EditorStyle/Private/SlateEditorStyle.cpp
--- a/Engine/Source/Editor/EditorStyle/Private/SlateEditorStyle.cpp
+++ b/Engine/Source/Editor/EditorStyle/Private/SlateEditorStyle.cpp
@@ -7738,6 +7738,7 @@ void FSlateEditorStyle::FStyle::SetupUMGEditorStyles()
 	Set("UMGEditor.Mirror", new IMAGE_BRUSH("Icons/UMG/Icon_Mirror_v3", Icon16x16));
 
 	Set("UMGEditor.ResizeResolutionFont", DEFAULT_FONT("Bold", 10));
+	Set("UMGEditor.CategoryIcon", new IMAGE_BRUSH("Icons/hiererchy_16x", Icon16x16));
 }
 
 void FSlateEditorStyle::FStyle::SetupTranslationEditorStyles()
diff --git a/Engine/Source/Editor/PropertyEditor/Private/SDetailsView.cpp b/Engine/Source/Editor/PropertyEditor/Private/SDetailsView.cpp
--- a/Engine/Source/Editor/PropertyEditor/Private/SDetailsView.cpp
+++ b/Engine/Source/Editor/PropertyEditor/Private/SDetailsView.cpp
@@ -204,7 +204,7 @@ void SDetailsView::Construct(const FArguments& InArgs)
 				.Visibility_Lambda([&](){ return this->bHasActiveFilter ? EVisibility::Visible : EVisibility::Collapsed; })
 			]
 			+SOverlay::Slot()
-			.Padding(4.0f, 2.0f)
+			.Padding(0.f, 2.0f, 4.f, 2.f)
 			.VAlign( VAlign_Center )
 			[
 				// Create the search box
diff --git a/Engine/Source/Editor/UMGEditor/Private/Details/SWidgetDetailsView.cpp b/Engine/Source/Editor/UMGEditor/Private/Details/SWidgetDetailsView.cpp
--- a/Engine/Source/Editor/UMGEditor/Private/Details/SWidgetDetailsView.cpp
+++ b/Engine/Source/Editor/UMGEditor/Private/Details/SWidgetDetailsView.cpp
@@ -72,81 +72,102 @@ void SWidgetDetailsView::Construct(const FArguments& InArgs, TSharedPtr<FWidgetB
 
 		+ SVerticalBox::Slot()
 		.AutoHeight()
-		.Padding(0, 0, 0, 6)
+		.Padding(0, 2)
 		[
-			SNew(SHorizontalBox)
-			.Visibility(this, &SWidgetDetailsView::GetCategoryAreaVisibility)
-
-			+ SHorizontalBox::Slot()
-			.AutoWidth()
-			.Padding(0, 0, 6, 0)
-			[
-				SNew(SBox)
-				.WidthOverride(200.0f)
-				.VAlign(VAlign_Center)
-				[
-					SNew(SEditableTextBox)
-					.SelectAllTextWhenFocused(true)
-					.ToolTipText(LOCTEXT("CategoryToolTip", "Sets the category of the widget"))
-					.HintText(LOCTEXT("Category", "Category"))
-					.Text(this, &SWidgetDetailsView::GetCategoryText)
-					.OnTextCommitted(this, &SWidgetDetailsView::HandleCategoryTextCommitted)
-				]
-			]
-		]
-
-		+ SVerticalBox::Slot()
-		.AutoHeight()
-		.Padding(0, 0, 0, 6)
-		[
-			SNew(SHorizontalBox)
-			.Visibility(this, &SWidgetDetailsView::GetNameAreaVisibility)
-
-			+ SHorizontalBox::Slot()
-			.AutoWidth()
-			.Padding(0, 0, 3, 0)
-			.VAlign(VAlign_Center)
+			SNew(SBorder)
+			.BorderImage(FEditorStyle::GetBrush(TEXT("ToolPanel.GroupBorder")))
+			.Visibility(this, &SWidgetDetailsView::GetBorderAreaVisibility)
 			[
-				SNew(SImage)
-				.Image(this, &SWidgetDetailsView::GetNameIcon)
-			]
+				SNew(SVerticalBox)
 
-			+ SHorizontalBox::Slot()
-			.AutoWidth()
-			.Padding(0, 0, 6, 0)
-			[
-				SNew(SBox)
-				.WidthOverride(200.0f)
-				.VAlign(VAlign_Center)
+				+ SVerticalBox::Slot()
+				.AutoHeight()
+				.Padding(4)
 				[
-					SAssignNew(NameTextBox, SEditableTextBox)
-					.SelectAllTextWhenFocused(true)
-					.HintText(LOCTEXT("Name", "Name"))
-					.Text(this, &SWidgetDetailsView::GetNameText)
-					.OnTextChanged(this, &SWidgetDetailsView::HandleNameTextChanged)
-					.OnTextCommitted(this, &SWidgetDetailsView::HandleNameTextCommitted)
+					SNew(SHorizontalBox)
+					.Visibility(this, &SWidgetDetailsView::GetCategoryAreaVisibility)
+
+					+ SHorizontalBox::Slot()
+					.AutoWidth()
+					.Padding(0, 0, 3, 0)
+					.VAlign(VAlign_Center)
+					[
+						SNew(SImage)
+						.Image(FEditorStyle::GetBrush(TEXT("UMGEditor.CategoryIcon")))
+					]
+
+					+ SHorizontalBox::Slot()
+					.AutoWidth()
+					.Padding(0, 0, 6, 0)
+					[
+						SNew(SBox)
+						.WidthOverride(200.0f)
+						.VAlign(VAlign_Center)
+						[
+							SNew(SEditableTextBox)
+							.SelectAllTextWhenFocused(true)
+							.ToolTipText(LOCTEXT("CategoryToolTip", "Sets the category of the widget"))
+							.HintText(LOCTEXT("Category", "Category"))
+							.Text(this, &SWidgetDetailsView::GetCategoryText)
+							.OnTextCommitted(this, &SWidgetDetailsView::HandleCategoryTextCommitted)
+						]
+					]
 				]
-			]
 
-			+ SHorizontalBox::Slot()
-			.AutoWidth()
-			[
-				SNew(SCheckBox)
-				.IsChecked(this, &SWidgetDetailsView::GetIsVariable)
-				.OnCheckStateChanged(this, &SWidgetDetailsView::HandleIsVariableChanged)
-				.Padding(FMargin(3,1,3,1))
+				+ SVerticalBox::Slot()
+				.AutoHeight()
+				.Padding(4)
 				[
-					SNew(STextBlock)
-					.Text(LOCTEXT("IsVariable", "Is Variable"))
+					SNew(SHorizontalBox)
+					.Visibility(this, &SWidgetDetailsView::GetNameAreaVisibility)
+
+					+ SHorizontalBox::Slot()
+					.AutoWidth()
+					.Padding(0, 0, 3, 0)
+					.VAlign(VAlign_Center)
+					[
+						SNew(SImage)
+						.Image(this, &SWidgetDetailsView::GetNameIcon)
+					]
+
+					+ SHorizontalBox::Slot()
+					.AutoWidth()
+					.Padding(0, 0, 6, 0)
+					[
+						SNew(SBox)
+						.WidthOverride(200.0f)
+						.VAlign(VAlign_Center)
+						[
+							SAssignNew(NameTextBox, SEditableTextBox)
+							.SelectAllTextWhenFocused(true)
+							.HintText(LOCTEXT("Name", "Name"))
+							.Text(this, &SWidgetDetailsView::GetNameText)
+							.OnTextChanged(this, &SWidgetDetailsView::HandleNameTextChanged)
+							.OnTextCommitted(this, &SWidgetDetailsView::HandleNameTextCommitted)
+						]
+					]
+
+					+ SHorizontalBox::Slot()
+					.AutoWidth()
+					[
+						SNew(SCheckBox)
+						.IsChecked(this, &SWidgetDetailsView::GetIsVariable)
+						.OnCheckStateChanged(this, &SWidgetDetailsView::HandleIsVariableChanged)
+						.Padding(FMargin(3,1,3,1))
+						[
+							SNew(STextBlock)
+							.Text(LOCTEXT("IsVariable", "Is Variable"))
+						]
+					]
+
+					+ SHorizontalBox::Slot()
+					.AutoWidth()
+					.Padding(15,0,0,0)
+					[
+						SAssignNew(ClassLinkArea, SBox)
+					]
 				]
 			]
-
-			+ SHorizontalBox::Slot()
-			.AutoWidth()
-			.Padding(15,0,0,0)
-			[
-				SAssignNew(ClassLinkArea, SBox)
-			]
 		]
 
 		+ SVerticalBox::Slot()
@@ -317,6 +338,11 @@ bool SWidgetDetailsView::IsWidgetCDOSelected() const
 	return false;
 }
 
+EVisibility SWidgetDetailsView::GetBorderAreaVisibility() const
+{
+	return (SelectedObjects.Num() == 0) ? EVisibility::Collapsed : EVisibility::Visible;
+}
+
 EVisibility SWidgetDetailsView::GetNameAreaVisibility() const
 {
 	return IsWidgetCDOSelected() ? EVisibility::Collapsed : EVisibility::Visible;
@@ -324,11 +350,6 @@ EVisibility SWidgetDetailsView::GetNameAreaVisibility() const
 
 EVisibility SWidgetDetailsView::GetCategoryAreaVisibility() const
 {
-	if ( SelectedObjects.Num() == 0 )
-	{
-		return EVisibility::Collapsed;
-	}
-
 	return IsWidgetCDOSelected() ? EVisibility::Visible : EVisibility::Collapsed;
 }
 
diff --git a/Engine/Source/Editor/UMGEditor/Private/Details/SWidgetDetailsView.h b/Engine/Source/Editor/UMGEditor/Private/Details/SWidgetDetailsView.h
--- a/Engine/Source/Editor/UMGEditor/Private/Details/SWidgetDetailsView.h
+++ b/Engine/Source/Editor/UMGEditor/Private/Details/SWidgetDetailsView.h
@@ -50,6 +50,8 @@ private:
 
 	bool IsWidgetCDOSelected() const;
 
+	EVisibility GetBorderAreaVisibility() const;
+
 	EVisibility GetNameAreaVisibility() const;
 
 	const FSlateBrush* GetNameIcon() const;
