PR #4847: UE-61192: Missing key handling for zoom level control (Contributed by [REDACTED_PERSON])

diff --git a/Engine/Source/Editor/UMGEditor/Private/Designer/SDesignSurface.cpp b/Engine/Source/Editor/UMGEditor/Private/Designer/SDesignSurface.cpp
--- a/Engine/Source/Editor/UMGEditor/Private/Designer/SDesignSurface.cpp
+++ b/Engine/Source/Editor/UMGEditor/Private/Designer/SDesignSurface.cpp
@@ -262,7 +262,7 @@ FReply SDesignSurface::OnMouseMove(const FGeometry& MyGeometry, const FPointerEv
 	{
 		const FVector2D CursorDelta = MouseEvent.GetCursorDelta();
 
-		const bool bShouldZoom = bIsRightMouseButtonDown && FSlateApplication::Get().IsUsingTrackpad();
+		const bool bShouldZoom = bIsRightMouseButtonDown && (bIsLeftMouseButtonDown || bIsMiddleMouseButtonDown || ModifierKeysState.IsAltDown() || FSlateApplication::Get().IsUsingTrackpad());
 		if ( bShouldZoom )
 		{
 			FReply ReplyState = FReply::Handled();
