Changing pin override to be for a single pin instead of all pin widgets

diff --git a/Engine/Source/Editor/UMGEditor/Private/Widgets/SGraphNodeCreateWidget.cpp b/Engine/Source/Editor/UMGEditor/Private/Widgets/SGraphNodeCreateWidget.cpp
--- a/Engine/Source/Editor/UMGEditor/Private/Widgets/SGraphNodeCreateWidget.cpp
+++ b/Engine/Source/Editor/UMGEditor/Private/Widgets/SGraphNodeCreateWidget.cpp
@@ -100,27 +100,17 @@ protected:
 //////////////////////////////////////////////////////////////////////////
 // SGraphNodeCreateWidget
 
-void SGraphNodeCreateWidget::CreatePinWidgets()
+TSharedPtr<SGraphPin> SGraphNodeCreateWidget::CreatePinWidget(UEdGraphPin* Pin) const
 {
 	UK2Node_CreateWidget* CreateWidgetNode = CastChecked<UK2Node_CreateWidget>(GraphNode);
 	UEdGraphPin* ClassPin = CreateWidgetNode->GetClassPin();
-
-	for (auto PinIt = GraphNode->Pins.CreateConstIterator(); PinIt; ++PinIt)
+	if ((ClassPin == Pin) && (!ClassPin->bHidden || (ClassPin->LinkedTo.Num() > 0)))
 	{
-		UEdGraphPin* CurrentPin = *PinIt;
-		if ((!CurrentPin->bHidden) && (CurrentPin != ClassPin))
-		{
-			TSharedPtr<SGraphPin> NewPin = FNodeFactory::CreatePinWidget(CurrentPin);
-			check(NewPin.IsValid());
-			this->AddPin(NewPin.ToSharedRef());
-		}
-		else if ((ClassPin == CurrentPin) && (!ClassPin->bHidden || (ClassPin->LinkedTo.Num() > 0)))
-		{
-			TSharedPtr<SGraphPinUserWidgetBasedClass> NewPin = SNew(SGraphPinUserWidgetBasedClass, ClassPin);
-			check(NewPin.IsValid());
-			this->AddPin(NewPin.ToSharedRef());
-		}
+		TSharedPtr<SGraphPinUserWidgetBasedClass> NewPin = SNew(SGraphPinUserWidgetBasedClass, ClassPin);
+		check(NewPin.IsValid());
+		return NewPin;
 	}
+	return SGraphNodeK2Default::CreatePinWidget(Pin);
 }
 
 #undef LOCTEXT_NAMESPACE
diff --git a/Engine/Source/Editor/UMGEditor/Private/Widgets/SGraphNodeCreateWidget.h b/Engine/Source/Editor/UMGEditor/Private/Widgets/SGraphNodeCreateWidget.h
--- a/Engine/Source/Editor/UMGEditor/Private/Widgets/SGraphNodeCreateWidget.h
+++ b/Engine/Source/Editor/UMGEditor/Private/Widgets/SGraphNodeCreateWidget.h
@@ -9,6 +9,6 @@ class SGraphNodeCreateWidget : public SGraphNodeK2Default
 {
 public:
 	// SGraphNode interface
-	virtual void CreatePinWidgets() override;
+	virtual TSharedPtr<SGraphPin> CreatePinWidget(UEdGraphPin* Pin) const override;
 	// End of SGraphNode interface
 };
