Adding a new animation to a UMG widget now prompts for compile and shows the animation in the MyBlueprint panel before recompile

diff --git a/Engine/Source/Editor/UMGEditor/Private/TabFactory/AnimationTabSummoner.cpp b/Engine/Source/Editor/UMGEditor/Private/TabFactory/AnimationTabSummoner.cpp
--- a/Engine/Source/Editor/UMGEditor/Private/TabFactory/AnimationTabSummoner.cpp
+++ b/Engine/Source/Editor/UMGEditor/Private/TabFactory/AnimationTabSummoner.cpp
@@ -216,6 +216,7 @@ private:
 			Blueprint->Modify();
 			Blueprint->Animations.Add(WidgetAnimation);
 			ListItem.Pin()->bNewAnimation = false;
+			FBlueprintEditorUtils::MarkBlueprintAsStructurallyModified(Blueprint);
 		}
 	}
 private:
