Fix for designer view showing all android screen resolutions as checked if they shared a device profile but were not the same final resolution

diff --git a/Engine/Source/Editor/UMGEditor/Private/Designer/SDesignerView.cpp b/Engine/Source/Editor/UMGEditor/Private/Designer/SDesignerView.cpp
--- a/Engine/Source/Editor/UMGEditor/Private/Designer/SDesignerView.cpp
+++ b/Engine/Source/Editor/UMGEditor/Private/Designer/SDesignerView.cpp
@@ -918,6 +918,8 @@ bool SDesignerView::IsRespectingLocks() const
 
 void SDesignerView::SetStartupResolution()
 {
+	HeightReadFromSettings = 0;
+	WidthReadFromSettings = 0;
 	// Use previously set resolution (or create new entries using default values)
 	// Width
 	if (!GConfig->GetInt(*ConfigSectionName, TEXT("PreviewWidth"), PreviewWidth, GEditorPerProjectIni))
@@ -3223,6 +3225,8 @@ void SDesignerView::HandleOnCommonResolutionSelected(FPlayScreenResolution InRes
 {
 	bSafeZoneFlipped = false;
 	bCanPreviewSwapAspectRatio = InResolution.bCanSwapAspectRatio;
+	WidthReadFromSettings = InResolution.Width;
+	HeightReadFromSettings = InResolution.Height;
 	// Phone/tablet resolutions can be stored in either portrait or landscape mode, and may need to be flipped
 	if (bCanPreviewSwapAspectRatio && ((!bPreviewIsPortrait && InResolution.Width < InResolution.Height) ||
 		(bPreviewIsPortrait && InResolution.Width > InResolution.Height)))
@@ -3310,12 +3314,27 @@ bool SDesignerView::HandleIsCommonResolutionSelected(FPlayScreenResolution InRes
 		}
 	}
 
-	if (!InResolution.ProfileName.IsEmpty())
+	int32 TestHeight = InResolution.Height;
+	int32 TestWidth = InResolution.Width;
+
+	// Swap the width and height to test if the preview is currently flipped
+	if ((InResolution.bCanSwapAspectRatio) && (PreviewWidth > PreviewHeight))
+	{
+		TestHeight = InResolution.Width;
+		TestWidth = InResolution.Height;
+	}
+
+	// Compare to the size in the settings
+	const bool bSizeMatches = ((TestWidth == WidthReadFromSettings) && (TestHeight == HeightReadFromSettings))
+		|| (InResolution.bCanSwapAspectRatio) && (PreviewWidth > PreviewHeight) && (TestHeight == WidthReadFromSettings) && (TestWidth == HeightReadFromSettings); // flipped to landscape
+
+	if (!PreviewOverrideName.IsEmpty() || !InResolution.ProfileName.IsEmpty())
 	{
-		return InResolution.ProfileName.Equals(PreviewOverrideName);
+		// Would have the same r.MobileContentScaleFactor and original size
+		return InResolution.ProfileName.Equals(PreviewOverrideName) && bSizeMatches;
 	}
 
-	return ((InResolution.Width == PreviewWidth ) && (InResolution.Height == PreviewHeight )) || (InResolution.bCanSwapAspectRatio && (InResolution.Height == PreviewWidth) && (InResolution.Width == PreviewHeight));
+	return bSizeMatches;
 }
 
 void SDesignerView::AddScreenResolutionSection(FMenuBuilder& MenuBuilder, const TArray<FPlayScreenResolution> Resolutions, const FText SectionName)
@@ -3564,9 +3583,17 @@ FReply SDesignerView::HandleZoomToFitClicked()
 FReply SDesignerView::HandleSwapAspectRatioClicked()
 {
 	bSafeZoneFlipped = false;
-	int32 OldPreviewHeight = PreviewHeight;
-	PreviewHeight = PreviewWidth;
-	PreviewWidth = OldPreviewHeight;
+	// If in portrait
+	if (PreviewHeight > PreviewWidth)
+	{
+		PreviewHeight = WidthReadFromSettings;
+		PreviewWidth = HeightReadFromSettings;
+	}
+	else
+	{
+		PreviewHeight = HeightReadFromSettings;
+		PreviewWidth = WidthReadFromSettings;
+	}
 
 	ScaleFactor = 1.0f;
 	ULevelEditorPlaySettings* PlayInSettings = GetMutableDefault<ULevelEditorPlaySettings>();
