Fix for missing file

diff --git a/Engine/Source/Editor/UMGEditor/Private/Designer/SDesignerView.h b/Engine/Source/Editor/UMGEditor/Private/Designer/SDesignerView.h
--- a/Engine/Source/Editor/UMGEditor/Private/Designer/SDesignerView.h
+++ b/Engine/Source/Editor/UMGEditor/Private/Designer/SDesignerView.h
@@ -345,6 +345,12 @@ private:
 	/** The configured Height of the preview area, simulates screen size. */
 	int32 PreviewHeight;
 
+	/** The original Width of the preview area, read from the settings file. */
+	int32 WidthReadFromSettings;
+
+	/** The original Height of the preview area, read from the settings file. */
+	int32 HeightReadFromSettings;
+
 	/***/
 	bool bShowResolutionOutlines;
 
