Fixing Linux compile issue

diff --git a/Engine/Source/Editor/UMGEditor/Private/Designer/SDesignerView.cpp b/Engine/Source/Editor/UMGEditor/Private/Designer/SDesignerView.cpp
--- a/Engine/Source/Editor/UMGEditor/Private/Designer/SDesignerView.cpp
+++ b/Engine/Source/Editor/UMGEditor/Private/Designer/SDesignerView.cpp
@@ -3326,7 +3326,7 @@ bool SDesignerView::HandleIsCommonResolutionSelected(FPlayScreenResolution InRes
 
 	// Compare to the size in the settings
 	const bool bSizeMatches = ((TestWidth == WidthReadFromSettings) && (TestHeight == HeightReadFromSettings))
-		|| (InResolution.bCanSwapAspectRatio) && (PreviewWidth > PreviewHeight) && (TestHeight == WidthReadFromSettings) && (TestWidth == HeightReadFromSettings); // flipped to landscape
+		|| (InResolution.bCanSwapAspectRatio && (PreviewWidth > PreviewHeight) && (TestHeight == WidthReadFromSettings) && (TestWidth == HeightReadFromSettings)); // flipped to landscape
 
 	if (!PreviewOverrideName.IsEmpty() || !InResolution.ProfileName.IsEmpty())
 	{
