Fix for iphone swapping breaking resolution

diff --git a/Engine/Source/Editor/UMGEditor/Private/Designer/SDesignerView.cpp b/Engine/Source/Editor/UMGEditor/Private/Designer/SDesignerView.cpp
--- a/Engine/Source/Editor/UMGEditor/Private/Designer/SDesignerView.cpp
+++ b/Engine/Source/Editor/UMGEditor/Private/Designer/SDesignerView.cpp
@@ -3600,7 +3600,7 @@ FReply SDesignerView::HandleSwapAspectRatioClicked()
 	const UDeviceProfile* DeviceProfile = UDeviceProfileManager::Get().FindProfile(PreviewOverrideName, false);
 
 	// Rescale the swapped sizes if we are on Android
-	if (DeviceProfile && DeviceProfile->DeviceType == TEXT("Android"))
+	if (DeviceProfile)
 	{
 		PlayInSettings->RescaleForMobilePreview(DeviceProfile, PreviewWidth, PreviewHeight, ScaleFactor);
 	}
