Alt+RMB zoom in UMG Designer now centers on mouse cursor instead of top left corner
Alt+RMB zoom in UMG/Blueprints now zooms in/out when moving right/left instead of only down/up
Fixed a bug where the context menu in UMG Designer would appear when Alt+RMB zooming
Fixed a bug when Alt+RMB zooming where zooming out would trigger on every pixel moved instead of every x pixels

diff --git a/Engine/Source/Editor/GraphEditor/Private/SNodePanel.cpp b/Engine/Source/Editor/GraphEditor/Private/SNodePanel.cpp
--- a/Engine/Source/Editor/GraphEditor/Private/SNodePanel.cpp
+++ b/Engine/Source/Editor/GraphEditor/Private/SNodePanel.cpp
@@ -196,7 +196,7 @@ namespace NodePanelDefs
 	// Node Culling Guardband Area
 	static const float GuardBandArea = 0.25f;
 	// Scaling factor to reduce speed of mouse zooming
-	static const float MouseZoomScaling = 0.05f;
+	static const float MouseZoomScaling = 0.04f;
 };
 
 SNodePanel::SNodePanel()
@@ -291,7 +291,7 @@ void SNodePanel::Construct()
 
 	ViewOffset = FVector2D::ZeroVector;
 	TotalMouseDelta = 0;
-	TotalMouseDeltaY = 0;
+	TotalMouseDeltaXY = 0;
 	bDeferredZoomToSelection = false;
 	bDeferredZoomToNodeExtents = false;
 
@@ -534,7 +534,7 @@ FReply SNodePanel::OnMouseButtonDown( const FGeometry& MyGeometry, const FPointe
 
 		DeferredMovementTargetObject = nullptr; // clear any interpolation when you manually zoom
 		CancelZoomToFit();
-		TotalMouseDeltaY = 0;
+		TotalMouseDeltaXY = 0;
 
 		if (!FSlateApplication::Get().IsUsingTrackpad()) // on trackpad we don't know yet if user wants to zoom or bring up the context menu
 		{
@@ -673,14 +673,14 @@ FReply SNodePanel::OnMouseMove(const FGeometry& MyGeometry, const FPointerEvent&
 		{
 			FReply ReplyState = FReply::Handled();
 
-			TotalMouseDeltaY += CursorDelta.Y;
+			TotalMouseDeltaXY += CursorDelta.X + CursorDelta.Y;
 
-			const int32 ZoomLevelDelta = FMath::FloorToInt(TotalMouseDeltaY * NodePanelDefs::MouseZoomScaling);
+			const int32 ZoomLevelDelta = FMath::RoundToInt(TotalMouseDeltaXY * NodePanelDefs::MouseZoomScaling);
 
 			// Get rid of mouse movement that's been 'used up' by zooming
 			if (ZoomLevelDelta != 0)
 			{
-				TotalMouseDeltaY -= (ZoomLevelDelta / NodePanelDefs::MouseZoomScaling);
+				TotalMouseDeltaXY -= (ZoomLevelDelta / NodePanelDefs::MouseZoomScaling);
 			}
 
 			// Perform zoom centered on the cached start offset
diff --git a/Engine/Source/Editor/GraphEditor/Public/SNodePanel.h b/Engine/Source/Editor/GraphEditor/Public/SNodePanel.h
--- a/Engine/Source/Editor/GraphEditor/Public/SNodePanel.h
+++ b/Engine/Source/Editor/GraphEditor/Public/SNodePanel.h
@@ -874,8 +874,8 @@ protected:
 	/** The total distance that the mouse has been dragged while down */
 	float TotalMouseDelta;
 
-	/** The Y component of mouse drag (used when zooming) */
-	float TotalMouseDeltaY;
+	/** The additive X and Y components of mouse drag (used when zooming) */
+	float TotalMouseDeltaXY;
 
 	/** Offset in the panel the user started the LMB+RMB zoom from */
 	FVector2D ZoomStartOffset;
diff --git a/Engine/Source/Editor/UMGEditor/Private/Designer/SDesignSurface.cpp b/Engine/Source/Editor/UMGEditor/Private/Designer/SDesignSurface.cpp
--- a/Engine/Source/Editor/UMGEditor/Private/Designer/SDesignSurface.cpp
+++ b/Engine/Source/Editor/UMGEditor/Private/Designer/SDesignSurface.cpp
@@ -122,7 +122,7 @@ void SDesignSurface::Construct(const FArguments& InArgs)
 	PostChangedZoom();
 	AllowContinousZoomInterpolation = InArgs._AllowContinousZoomInterpolation;
 	bIsPanning = false;
-	bIsZoomingWithTrackpad = false;
+	bIsZooming = false;
 
 	ViewOffset = FVector2D::ZeroVector;
 	bDrawGridLines = true;
@@ -146,7 +146,7 @@ void SDesignSurface::Construct(const FArguments& InArgs)
 	ZoomToFitPadding = FVector2D(100, 100);
 	TotalGestureMagnify = 0.0f;
 
-	TotalMouseDeltaY = 0.0f;
+	TotalMouseDelta = 0.0f;
 	ZoomStartOffset = FVector2D::ZeroVector;
 
 	ChildSlot
@@ -229,9 +229,9 @@ FReply SDesignSurface::OnMouseButtonDown(const FGeometry& MyGeometry, const FPoi
 		MouseDownPositionAbsolute = MouseEvent.GetLastScreenSpacePosition();
 	}
 
-	if (FSlateApplication::Get().IsUsingTrackpad())
+	if (MouseEvent.GetEffectingButton() == EKeys::RightMouseButton || FSlateApplication::Get().IsUsingTrackpad())
 	{
-		TotalMouseDeltaY = 0.0f;
+		TotalMouseDelta = 0.0f;
 		ZoomStartOffset = MyGeometry.AbsoluteToLocal(MouseEvent.GetLastScreenSpacePosition());
 	}
 
@@ -245,7 +245,7 @@ FReply SDesignSurface::OnMouseButtonUp(const FGeometry& MyGeometry, const FPoint
 	if ( MouseEvent.GetEffectingButton() == EKeys::RightMouseButton || MouseEvent.GetEffectingButton() == EKeys::MiddleMouseButton )
 	{
 		bIsPanning = false;
-		bIsZoomingWithTrackpad = false;
+		bIsZooming = false;
 	}
 
 	return FReply::Unhandled();
@@ -265,16 +265,18 @@ FReply SDesignSurface::OnMouseMove(const FGeometry& MyGeometry, const FPointerEv
 		const bool bShouldZoom = bIsRightMouseButtonDown && (bIsLeftMouseButtonDown || bIsMiddleMouseButtonDown || ModifierKeysState.IsAltDown() || FSlateApplication::Get().IsUsingTrackpad());
 		if ( bShouldZoom )
 		{
+			const float MouseZoomScaling = 0.04f;
 			FReply ReplyState = FReply::Handled();
 
-			TotalMouseDeltaY += CursorDelta.Y;
+			TotalMouseDelta += CursorDelta.X + CursorDelta.Y;
 
-			const int32 ZoomLevelDelta = FMath::FloorToInt(TotalMouseDeltaY * 0.05f);
+			const int32 ZoomLevelDelta = FMath::RoundToInt(TotalMouseDelta * MouseZoomScaling);
 
 			// Get rid of mouse movement that's been 'used up' by zooming
 			if (ZoomLevelDelta != 0)
 			{
-				TotalMouseDeltaY -= (ZoomLevelDelta / 0.05f);
+				TotalMouseDelta -= (ZoomLevelDelta / MouseZoomScaling);
+				bIsZooming = true;
 			}
 
 			// Perform zoom centered on the cached start offset
@@ -282,11 +284,6 @@ FReply SDesignSurface::OnMouseMove(const FGeometry& MyGeometry, const FPointerEv
 
 			bIsPanning = false;
 
-			if (FSlateApplication::Get().IsUsingTrackpad() && ZoomLevelDelta != 0)
-			{
-				bIsZoomingWithTrackpad = true;
-			}
-
 			return ReplyState;
 		}
 		else if ( bIsRightMouseButtonDown || bIsMiddleMouseButtonDown )
@@ -294,6 +291,7 @@ FReply SDesignSurface::OnMouseMove(const FGeometry& MyGeometry, const FPointerEv
 			FReply ReplyState = FReply::Handled();
 
 			bIsPanning = true;
+			bIsZooming = false;
 			ViewOffset = ViewOffsetStart + ( (MouseDownPositionAbsolute - MouseEvent.GetScreenSpacePosition()) / MyGeometry.Scale) / GetZoomAmount();
 
 			return ReplyState;
@@ -417,7 +415,7 @@ void SDesignSurface::ChangeZoomLevel(int32 ZoomLevelDelta, const FVector2D& Widg
 			// Update view offset to where ever we scrolled towards.
 			ViewOffset = NewViewOffset;
 
-			TotalMouseDeltaY = 0.0f;
+			TotalMouseDelta = 0.0f;
 		}
 	}
 }
diff --git a/Engine/Source/Editor/UMGEditor/Private/Designer/SDesignSurface.h b/Engine/Source/Editor/UMGEditor/Private/Designer/SDesignSurface.h
--- a/Engine/Source/Editor/UMGEditor/Private/Designer/SDesignSurface.h
+++ b/Engine/Source/Editor/UMGEditor/Private/Designer/SDesignSurface.h
@@ -102,7 +102,7 @@ protected:
 	bool bIsPanning;
 
 	/** Are we zooming the view with trackpad at the moment? */
-	bool bIsZoomingWithTrackpad;
+	bool bIsZooming;
 
 	/** Allow continuous zoom interpolation? */
 	TAttribute<bool> AllowContinousZoomInterpolation;
@@ -125,7 +125,7 @@ protected:
 	FVector2D ZoomToFitPadding;
 
 	/** The Y component of mouse drag (used when zooming) */
-	float TotalMouseDeltaY;
+	float TotalMouseDelta;
 
 	/** Offset in the panel the user started the LMB+RMB zoom from */
 	FVector2D ZoomStartOffset;
diff --git a/Engine/Source/Editor/UMGEditor/Private/Designer/SDesignerView.cpp b/Engine/Source/Editor/UMGEditor/Private/Designer/SDesignerView.cpp
--- a/Engine/Source/Editor/UMGEditor/Private/Designer/SDesignerView.cpp
+++ b/Engine/Source/Editor/UMGEditor/Private/Designer/SDesignerView.cpp
@@ -1756,7 +1756,7 @@ FReply SDesignerView::OnMouseButtonUp(const FGeometry& MyGeometry, const FPointe
 	}
 	else if ( MouseEvent.GetEffectingButton() == EKeys::RightMouseButton )
 	{
-		if ( !bIsPanning && !bIsZoomingWithTrackpad )
+		if ( !bIsPanning && !bIsZooming )
 		{
 			ResolvePendingSelectedWidgets();
 
