Don't export UUserWidget::WidgetTree when exporting widgets to text

diff --git a/Engine/Source/Editor/UMGEditor/Private/WidgetBlueprintEditorUtils.cpp b/Engine/Source/Editor/UMGEditor/Private/WidgetBlueprintEditorUtils.cpp
--- a/Engine/Source/Editor/UMGEditor/Private/WidgetBlueprintEditorUtils.cpp
+++ b/Engine/Source/Editor/UMGEditor/Private/WidgetBlueprintEditorUtils.cpp
@@ -992,18 +992,34 @@ void FWidgetBlueprintEditorUtils::ExportWidgetsToText(TArray<UWidget*> WidgetsTo
 	UnMarkAllObjects(EObjectMark(OBJECTMARK_TagExp | OBJECTMARK_TagImp));
 
 	FStringOutputDevice Archive;
-	const FExportObjectInnerContext Context;
 
-	// Export each of the selected nodes
+	// Validate all nodes are from the same scope and set all UUserWidget::WidgetTrees (and things outered to it) to be ignored
+	TArray<UObject*> WidgetsToIgnore;
 	UObject* LastOuter = nullptr;
 	for ( UWidget* Widget : WidgetsToExport )
 	{
 		// The nodes should all be from the same scope
 		UObject* ThisOuter = Widget->GetOuter();
-		check(( LastOuter == ThisOuter ) || ( LastOuter == nullptr ));
+		check((LastOuter == ThisOuter) || (LastOuter == nullptr));
 		LastOuter = ThisOuter;
 
-		UExporter::ExportToOutputDevice(&Context, Widget, nullptr, Archive, TEXT("copy"), 0, PPF_ExportsNotFullyQualified | PPF_Copy | PPF_Delimited, false, ThisOuter);
+		if ( UUserWidget* UserWidget = Cast<UUserWidget>(Widget) )
+		{
+			if ( UserWidget->WidgetTree )
+			{
+				WidgetsToIgnore.Add(UserWidget->WidgetTree);
+				// FExportObjectInnerContext does not automatically ignore UObjects if their outer is ignored
+				GetObjectsWithOuter(UserWidget->WidgetTree, WidgetsToIgnore);
+			}
+		}
+	}
+
+	const FExportObjectInnerContext Context(WidgetsToIgnore);
+	// Export each of the selected nodes
+	for ( UWidget* Widget : WidgetsToExport )
+	{
+
+		UExporter::ExportToOutputDevice(&Context, Widget, nullptr, Archive, TEXT("copy"), 0, PPF_ExportsNotFullyQualified | PPF_Copy | PPF_Delimited, false, LastOuter);
 
 		// Check to see if this widget was content of another widget holding it in a named slot.
 		if ( Widget->GetParent() == nullptr )
diff --git a/Engine/Source/Runtime/UMG/Public/Blueprint/UserWidget.h b/Engine/Source/Runtime/UMG/Public/Blueprint/UserWidget.h
--- a/Engine/Source/Runtime/UMG/Public/Blueprint/UserWidget.h
+++ b/Engine/Source/Runtime/UMG/Public/Blueprint/UserWidget.h
@@ -1131,7 +1131,7 @@ private:
 
 public:
 	/** The widget tree contained inside this user widget initialized by the blueprint */
-	UPROPERTY(Instanced)
+	UPROPERTY(Instanced, TextExportTransient)
 	UWidgetTree* WidgetTree;
 
 public:
