Exposed named slots in child widgets can now have their contents wrapped properly

diff --git a/Engine/Source/Editor/UMGEditor/Private/WidgetBlueprintEditorUtils.cpp b/Engine/Source/Editor/UMGEditor/Private/WidgetBlueprintEditorUtils.cpp
--- a/Engine/Source/Editor/UMGEditor/Private/WidgetBlueprintEditorUtils.cpp
+++ b/Engine/Source/Editor/UMGEditor/Private/WidgetBlueprintEditorUtils.cpp
@@ -500,6 +500,11 @@ void FWidgetBlueprintEditorUtils::FindAllAncestorNamedSlotHostWidgetsForContent(
 }
 
 bool FWidgetBlueprintEditorUtils::RemoveNamedSlotHostContent(UWidget* WidgetTemplate, INamedSlotInterface* NamedSlotHost)
+{
+	return ReplaceNamedSlotHostContent(WidgetTemplate, NamedSlotHost, nullptr);
+}
+
+bool FWidgetBlueprintEditorUtils::ReplaceNamedSlotHostContent(UWidget* WidgetTemplate, INamedSlotInterface* NamedSlotHost, UWidget* NewContentWidget)
 {
 	TArray<FName> SlotNames;
 	NamedSlotHost->GetSlotNames(SlotNames);
@@ -510,7 +515,7 @@ bool FWidgetBlueprintEditorUtils::RemoveNamedSlotHostContent(UWidget* WidgetTemp
 		{
 			if (SlotContent == WidgetTemplate)
 			{
-				NamedSlotHost->SetContentForSlot(SlotName, nullptr);
+				NamedSlotHost->SetContentForSlot(SlotName, NewContentWidget);
 				return true;
 			}
 		}
@@ -579,37 +584,60 @@ void FWidgetBlueprintEditorUtils::WrapWidgets(TSharedRef<FWidgetBlueprintEditor>
 		int32 OutIndex;
 		UWidget* Widget = Item.GetTemplate();
 		UPanelWidget* CurrentParent = BP->WidgetTree->FindWidgetParent(Widget, OutIndex);
+		UWidget* CurrentSlot = FindNamedSlotHostWidgetForContent(Widget, BP->WidgetTree);
 
-		// If the widget doesn't currently have a parent, and isn't the root, ignore it.
-		if (CurrentParent == nullptr && Widget != BP->WidgetTree->RootWidget)
+		// If the widget doesn't currently have a slot or parent, and isn't the root, ignore it.
+		if (CurrentSlot == nullptr && CurrentParent == nullptr && Widget != BP->WidgetTree->RootWidget)
 		{
 			continue;
 		}
 
 		Widget->Modify();
+		BP->WidgetTree->SetFlags(RF_Transactional);
+		BP->WidgetTree->Modify();
 
-		UPanelWidget*& NewWrapperWidget = OldParentToNewParent.FindOrAdd(CurrentParent);
-		if (NewWrapperWidget == nullptr || !NewWrapperWidget->CanAddMoreChildren())
+		if (CurrentSlot)
 		{
-			NewWrapperWidget = CastChecked<UPanelWidget>(Template->Create(BP->WidgetTree));
-			NewWrapperWidget->SetDesignerFlags(BlueprintEditor->GetCurrentDesignerFlags());
+			// If this is a named slot, we need to properly remove and reassign the slot content
+			INamedSlotInterface* NamedSlotHost = Cast<INamedSlotInterface>(CurrentSlot);
+			if (NamedSlotHost != nullptr)
+			{
+				CurrentSlot->SetFlags(RF_Transactional);
+				CurrentSlot->Modify();
 
-			BP->WidgetTree->SetFlags(RF_Transactional);
-			BP->WidgetTree->Modify();
+				UPanelWidget* NewSlotContents = CastChecked<UPanelWidget>(Template->Create(BP->WidgetTree));
+				NewSlotContents->SetDesignerFlags(BlueprintEditor->GetCurrentDesignerFlags());
+
+				FWidgetBlueprintEditorUtils::ReplaceNamedSlotHostContent(Widget, NamedSlotHost, NewSlotContents);
+
+				NewSlotContents->AddChild(Widget);
 
-			if (CurrentParent)
+			}
+		}
+		else if (CurrentParent)
+		{
+			UPanelWidget*& NewWrapperWidget = OldParentToNewParent.FindOrAdd(CurrentParent);
+			if (NewWrapperWidget == nullptr || !NewWrapperWidget->CanAddMoreChildren())
 			{
+				NewWrapperWidget = CastChecked<UPanelWidget>(Template->Create(BP->WidgetTree));
+				NewWrapperWidget->SetDesignerFlags(BlueprintEditor->GetCurrentDesignerFlags());				
+
 				CurrentParent->SetFlags(RF_Transactional);
 				CurrentParent->Modify();
 				CurrentParent->ReplaceChildAt(OutIndex, NewWrapperWidget);
+
+				NewWrapperWidget->AddChild(Widget);
 			}
-			else // Root Widget
-			{
-				BP->WidgetTree->RootWidget = NewWrapperWidget;
-			}
+		}
+		else
+		{
+			UPanelWidget* NewRootContents = CastChecked<UPanelWidget>(Template->Create(BP->WidgetTree));
+			NewRootContents->SetDesignerFlags(BlueprintEditor->GetCurrentDesignerFlags());
+
+			BP->WidgetTree->RootWidget = NewRootContents;
+			NewRootContents->AddChild(Widget);
 		}
 
-		NewWrapperWidget->AddChild(Widget);
 	}
 
 	FBlueprintEditorUtils::MarkBlueprintAsStructurallyModified(BP);
diff --git a/Engine/Source/Editor/UMGEditor/Private/WidgetBlueprintEditorUtils.h b/Engine/Source/Editor/UMGEditor/Private/WidgetBlueprintEditorUtils.h
--- a/Engine/Source/Editor/UMGEditor/Private/WidgetBlueprintEditorUtils.h
+++ b/Engine/Source/Editor/UMGEditor/Private/WidgetBlueprintEditorUtils.h
@@ -59,6 +59,8 @@ public:
 
 	static bool RemoveNamedSlotHostContent(UWidget* WidgetTemplate, INamedSlotInterface* NamedSlotHost);
 
+	static bool ReplaceNamedSlotHostContent(UWidget* WidgetTemplate, INamedSlotInterface* NamedSlotHost, UWidget* NewContentWidget);
+
 private:
 
 	static FString CopyWidgetsInternal(UWidgetBlueprint* BP, TSet<FWidgetReference> Widgets);
