Adding project setting option for the default root widget of new Widget Blueprints

#tests New widgets have the specified root (empty works too)

diff --git a/Engine/Source/Editor/UMGEditor/Private/Settings/UMGEditorProjectSettings.cpp b/Engine/Source/Editor/UMGEditor/Private/Settings/UMGEditorProjectSettings.cpp
--- a/Engine/Source/Editor/UMGEditor/Private/Settings/UMGEditorProjectSettings.cpp
+++ b/Engine/Source/Editor/UMGEditor/Private/Settings/UMGEditorProjectSettings.cpp
@@ -5,6 +5,7 @@
 #include "WidgetCompilerRule.h"
 #include "UObject/Package.h"
 #include "UObject/UObjectIterator.h"
+#include "Components/CanvasPanel.h"
 
 UUMGEditorProjectSettings::UUMGEditorProjectSettings()
 {
@@ -13,6 +14,8 @@ UUMGEditorProjectSettings::UUMGEditorProjectSettings()
 	bShowWidgetsFromEngineContent = false;
 	bShowWidgetsFromDeveloperContent = true;
 
+	DefaultRootWidget = UCanvasPanel::StaticClass();
+
 	// Deprecated
 	bCookSlowConstructionWidgetTree_DEPRECATED = true;
 	bWidgetSupportsDynamicCreation_DEPRECATED = true;
diff --git a/Engine/Source/Editor/UMGEditor/Private/WidgetBlueprintFactory.cpp b/Engine/Source/Editor/UMGEditor/Private/WidgetBlueprintFactory.cpp
--- a/Engine/Source/Editor/UMGEditor/Private/WidgetBlueprintFactory.cpp
+++ b/Engine/Source/Editor/UMGEditor/Private/WidgetBlueprintFactory.cpp
@@ -8,10 +8,8 @@
 #include "WidgetBlueprint.h"
 #include "Kismet2/KismetEditorUtilities.h"
 
-
-
-#include "Components/CanvasPanel.h"
 #include "Blueprint/WidgetTree.h"
+#include "UMGEditorProjectSettings.h"
 
 #define LOCTEXT_NAMESPACE "UWidgetBlueprintFactory"
 
@@ -53,19 +51,22 @@ UObject* UWidgetBlueprintFactory::FactoryCreateNew(UClass* Class, UObject* InPar
 	if ( ( ParentClass == NULL ) || !FKismetEditorUtilities::CanCreateBlueprintOfClass(ParentClass) || !ParentClass->IsChildOf(UUserWidget::StaticClass()) )
 	{
 		FFormatNamedArguments Args;
-		Args.Add( TEXT("ClassName"), (ParentClass != NULL) ? FText::FromString( ParentClass->GetName() ) : LOCTEXT("Null", "(null)") );
+		Args.Add( TEXT("ClassName"), ParentClass ? FText::FromString( ParentClass->GetName() ) : LOCTEXT("Null", "(null)") );
 		FMessageDialog::Open( EAppMsgType::Ok, FText::Format( LOCTEXT("CannotCreateWidgetBlueprint", "Cannot create a Widget Blueprint based on the class '{ClassName}'."), Args ) );
-		return NULL;
+		return nullptr;
 	}
 	else
 	{
 		UWidgetBlueprint* NewBP = CastChecked<UWidgetBlueprint>(FKismetEditorUtilities::CreateBlueprint(ParentClass, InParent, Name, BlueprintType, UWidgetBlueprint::StaticClass(), UWidgetBlueprintGeneratedClass::StaticClass(), CallingContext));
 
-		// Create a CanvasPanel to use as the default root widget
+		// Create the desired root widget specified by the project
 		if ( NewBP->WidgetTree->RootWidget == nullptr )
 		{
-			UWidget* Root = NewBP->WidgetTree->ConstructWidget<UCanvasPanel>(UCanvasPanel::StaticClass());
-			NewBP->WidgetTree->RootWidget = Root;
+			if (TSubclassOf<UPanelWidget> RootWidgetClass = GetDefault<UUMGEditorProjectSettings>()->DefaultRootWidget)
+			{
+				UWidget* Root = NewBP->WidgetTree->ConstructWidget<UWidget>(RootWidgetClass);
+				NewBP->WidgetTree->RootWidget = Root;
+			}
 		}
 
 		return NewBP;
diff --git a/Engine/Source/Editor/UMGEditor/Public/UMGEditorProjectSettings.h b/Engine/Source/Editor/UMGEditor/Public/UMGEditorProjectSettings.h
--- a/Engine/Source/Editor/UMGEditor/Public/UMGEditorProjectSettings.h
+++ b/Engine/Source/Editor/UMGEditor/Public/UMGEditorProjectSettings.h
@@ -13,6 +13,7 @@
 class UWidgetCompilerRule;
 class UUserWidget;
 class UWidgetBlueprint;
+class UPanelWidget;
 
 USTRUCT()
 struct FDebugResolution
@@ -174,6 +175,10 @@ public:
 
 public:
 
+	/** The panel widget to place at the root of all newly constructed widget blueprints. Can be empty. */
+	UPROPERTY(EditAnywhere, config, Category = Designer)
+	TSubclassOf<UPanelWidget> DefaultRootWidget;
+
 	UPROPERTY(EditAnywhere, config, Category=Designer)
 	TArray<FDebugResolution> DebugResolutions;
 
