When duplicating widget animations also set their display label, so when renaming the old one is not still in use.

diff --git a/Engine/Source/Editor/UMGEditor/Private/TabFactory/AnimationTabSummoner.cpp b/Engine/Source/Editor/UMGEditor/Private/TabFactory/AnimationTabSummoner.cpp
--- a/Engine/Source/Editor/UMGEditor/Private/TabFactory/AnimationTabSummoner.cpp
+++ b/Engine/Source/Editor/UMGEditor/Private/TabFactory/AnimationTabSummoner.cpp
@@ -572,7 +572,7 @@ private:
 			);
 	
 		NewAnimation->MovieScene->Rename(*NewAnimation->GetName(), nullptr, REN_DontCreateRedirectors | REN_ForceNoResetLoaders);
-
+		NewAnimation->SetDisplayLabel(NewAnimation->GetName());
 
 		bool bRenameRequestPending = true;
 		bool bNewAnimation = true;
