UMG: Never recompile dirty blueprints on evaluation for UMG widgets

This setting is counterproductive for widgets which are blueprints themselves. Compilation of such blueprints is always down to the user.l

diff --git a/Engine/Source/Editor/UMGEditor/Private/WidgetBlueprintEditor.cpp b/Engine/Source/Editor/UMGEditor/Private/WidgetBlueprintEditor.cpp
--- a/Engine/Source/Editor/UMGEditor/Private/WidgetBlueprintEditor.cpp
+++ b/Engine/Source/Editor/UMGEditor/Private/WidgetBlueprintEditor.cpp
@@ -21,6 +21,7 @@
 
 #include "Tracks/MovieScenePropertyTrack.h"
 #include "ISequencerModule.h"
+#include "SequencerSettings.h"
 #include "ObjectEditorUtils.h"
 
 #include "PropertyCustomizationHelpers.h"
@@ -790,6 +791,8 @@ TSharedPtr<ISequencer>& FWidgetBlueprintEditor::GetSequencer()
 		};
 
 		Sequencer = FModuleManager::LoadModuleChecked<ISequencerModule>("Sequencer").CreateSequencer(SequencerInitParams);
+		// Never recompile the blueprint on evaluate as this can create an insidious loop
+		Sequencer->GetSequencerSettings()->SetCompileDirectorOnEvaluate(false);
 		Sequencer->OnMovieSceneDataChanged().AddSP( this, &FWidgetBlueprintEditor::OnMovieSceneDataChanged );
 		Sequencer->OnMovieSceneBindingsPasted().AddSP( this, &FWidgetBlueprintEditor::OnMovieSceneBindingsPasted );
 		// Change selected widgets in the sequencer tree view
