UMG: Fix replace crash when source binding doesn't exist - add the source binding.

diff --git a/Engine/Source/Editor/UMGEditor/Private/WidgetBlueprintEditor.cpp b/Engine/Source/Editor/UMGEditor/Private/WidgetBlueprintEditor.cpp
--- a/Engine/Source/Editor/UMGEditor/Private/WidgetBlueprintEditor.cpp
+++ b/Engine/Source/Editor/UMGEditor/Private/WidgetBlueprintEditor.cpp
@@ -1378,6 +1378,18 @@ void FWidgetBlueprintEditor::ReplaceTrackWithSelectedWidget(FWidgetReference Sel
 			}
 		);
 
+		if (!SourceBinding)
+		{
+			WidgetAnimation->BindPossessableObject(ObjectId, *PreviewWidget, GetAnimationPlaybackContext());
+
+			SourceBinding = WidgetAnimation->AnimationBindings.FindByPredicate(
+				[&](FWidgetAnimationBinding& In)
+			{
+				return In.AnimationGuid == ObjectId;
+			}
+			);
+		}
+
 		check(SourceBinding);
 
 		// Set binding names used for lookup
