Now validate renaming animations with SlugStringForValidName instead of MakeObjectNameFromDisplayLabel. The old method didn't account for numbering, so [Name] and [Name]_1 were treated as equivalent. Also added in feedback when attempting to use invalid characters, the UI will now list the invalid characters being used.

diff --git a/Engine/Source/Editor/UMGEditor/Private/TabFactory/AnimationTabSummoner.cpp b/Engine/Source/Editor/UMGEditor/Private/TabFactory/AnimationTabSummoner.cpp
--- a/Engine/Source/Editor/UMGEditor/Private/TabFactory/AnimationTabSummoner.cpp
+++ b/Engine/Source/Editor/UMGEditor/Private/TabFactory/AnimationTabSummoner.cpp
@@ -48,6 +48,27 @@ FAnimationTabSummoner::FAnimationTabSummoner(TSharedPtr<class FWidgetBlueprintEd
 
 bool VerifyAnimationRename( FWidgetBlueprintEditor& BlueprintEditor, UWidgetAnimation* Animation, FString NewAnimationName, FText& OutErrorMessage )
 {
+	if (NewAnimationName != SlugStringForValidName(NewAnimationName))
+	{
+		FString InvalidCharacters = INVALID_OBJECTNAME_CHARACTERS;
+		FString CurrentInvalidCharacter;
+		FString FoundInvalidCharacters;
+
+		// Create a string with all invalid characters found to output with the error message.
+		for (int32 StringIndex = 0; StringIndex < InvalidCharacters.Len(); ++StringIndex)
+		{
+			CurrentInvalidCharacter = InvalidCharacters.Mid(StringIndex, 1);
+
+			if (NewAnimationName.Contains(CurrentInvalidCharacter))
+			{
+				FoundInvalidCharacters += CurrentInvalidCharacter;
+			}
+		}
+		
+		OutErrorMessage = FText::Format(LOCTEXT("NameContainsInvalidCharacters", "The object name may not contain the following characters:  {0}"), FText::FromString(FoundInvalidCharacters));
+		return false;
+	}
+
 	UWidgetBlueprint* Blueprint = BlueprintEditor.GetWidgetBlueprintObj();
 	if (Blueprint && FindObject<UWidgetAnimation>( Blueprint, *NewAnimationName, true ) )
 	{
@@ -157,7 +178,7 @@ private:
 	{
 		UWidgetAnimation* Animation = ListItem.Pin()->Animation;
 
-		const FName NewName = MakeObjectNameFromDisplayLabel(InText.ToString(), Animation->GetFName());
+		const FName NewName = *InText.ToString();
 
 		if ( Animation->GetFName() != NewName )
 		{
@@ -173,8 +194,8 @@ private:
 		UWidgetAnimation* WidgetAnimation = ListItem.Pin()->Animation;
 		UWidgetBlueprint* Blueprint = BlueprintEditor.Pin()->GetWidgetBlueprintObj();
 
-		// Get the new FName slug from the given display name
-		const FName NewFName = MakeObjectNameFromDisplayLabel(InText.ToString(), WidgetAnimation->GetFName());
+		// Name has already been checked in VerifyAnimationRename
+		const FName NewFName = *InText.ToString();
 		const FName OldFName = WidgetAnimation->GetFName();
 
 		UObjectPropertyBase* ExistingProperty = Cast<UObjectPropertyBase>(Blueprint->ParentClass->FindPropertyByName(NewFName));
