When naming new animation, pass in empty name instead of base name to avoid object name collision. Removed unecassary valid name check when creating new animation.

diff --git a/Engine/Source/Editor/UMGEditor/Private/TabFactory/AnimationTabSummoner.cpp b/Engine/Source/Editor/UMGEditor/Private/TabFactory/AnimationTabSummoner.cpp
--- a/Engine/Source/Editor/UMGEditor/Private/TabFactory/AnimationTabSummoner.cpp
+++ b/Engine/Source/Editor/UMGEditor/Private/TabFactory/AnimationTabSummoner.cpp
@@ -430,7 +430,7 @@ private:
 		UWidgetBlueprint* WidgetBlueprint = Editor->GetWidgetBlueprintObj();
 
 		FString BaseName = "NewAnimation";
-		UWidgetAnimation* NewAnimation = NewObject<UWidgetAnimation>(WidgetBlueprint, *BaseName, RF_Transactional);
+		UWidgetAnimation* NewAnimation = NewObject<UWidgetAnimation>(WidgetBlueprint, FName(), RF_Transactional);
 
 		FString UniqueName = BaseName;
 		int32 NameIndex = 1;
@@ -441,7 +441,7 @@ private:
 			NameIndex++;
 		}
 
-		const FName NewFName = MakeObjectNameFromDisplayLabel(UniqueName, NewAnimation->GetFName());
+		const FName NewFName = FName(*UniqueName);
 		NewAnimation->SetDisplayLabel( UniqueName );
 		NewAnimation->Rename(*UniqueName);
 
