UMG - Fixing a crash in the hierarchy view, sometimes the template ptr in stale.

diff --git a/Engine/Source/Editor/UMGEditor/Private/Hierarchy/SHierarchyViewItem.h b/Engine/Source/Editor/UMGEditor/Private/Hierarchy/SHierarchyViewItem.h
--- a/Engine/Source/Editor/UMGEditor/Private/Hierarchy/SHierarchyViewItem.h
+++ b/Engine/Source/Editor/UMGEditor/Private/Hierarchy/SHierarchyViewItem.h
@@ -265,12 +265,20 @@ public:
 
 	virtual bool IsExpanded() const override
 	{
-		return Item.GetTemplate()->bExpandedInDesigner;
+		if (UWidget* Template = Item.GetTemplate())
+		{
+			return Template->bExpandedInDesigner;
+		}
+
+		return false;
 	}
 
 	virtual void SetExpanded(bool bIsExpanded) override
 	{
-		Item.GetTemplate()->bExpandedInDesigner = bIsExpanded;
+		if (UWidget* Template = Item.GetTemplate())
+		{
+			Template->bExpandedInDesigner = bIsExpanded;
+		}
 	}
 
 	virtual bool CanRename() const override;
