UMG - The UMG Runtime module should never depend on UnrealEd even at editor time, it creates a circular dependency in the modules and we don't want that.  Providing an intreface for widgets with custom compiler rules to provide errors, warnings and notes without relying on an UnrealEd class.

diff --git a/Engine/Source/Editor/UMGEditor/Private/WidgetBlueprintCompiler.cpp b/Engine/Source/Editor/UMGEditor/Private/WidgetBlueprintCompiler.cpp
--- a/Engine/Source/Editor/UMGEditor/Private/WidgetBlueprintCompiler.cpp
+++ b/Engine/Source/Editor/UMGEditor/Private/WidgetBlueprintCompiler.cpp
@@ -20,6 +20,7 @@
 #include "IUMGModule.h"
 #include "UMGEditorProjectSettings.h"
 #include "WidgetCompilerRule.h"
+#include "Editor/WidgetCompilerLog.h"
 #include "Editor.h"
 
 #define LOCTEXT_NAMESPACE "UMG"
@@ -763,6 +764,26 @@ void FWidgetBlueprintCompilerContext::FinishCompilingClass(UClass* Class)
 	Super::FinishCompilingClass(Class);
 }
 
+
+class FBlueprintCompilerLog : public IWidgetCompilerLog
+{
+public:
+	FBlueprintCompilerLog(FCompilerResultsLog& InMessageLog)
+		: MessageLog(InMessageLog)
+	{
+	}
+
+	virtual void InternalLogMessage(TSharedRef<FTokenizedMessage>& InMessage) override
+	{
+		MessageLog.AddTokenizedMessage(InMessage);
+	}
+
+private:
+	// Compiler message log (errors, warnings, notes)
+	FCompilerResultsLog& MessageLog;
+};
+
+
 void FWidgetBlueprintCompilerContext::PostCompile()
 {
 	Super::PostCompile();
@@ -786,7 +807,8 @@ void FWidgetBlueprintCompilerContext::PostCompile()
 
 	if (!Blueprint->bIsRegeneratingOnLoad && bIsFullCompile)
 	{
-		WidgetClass->GetDefaultObject<UUserWidget>()->ValidateBlueprint(*WidgetBP->WidgetTree, MessageLog);
+		FBlueprintCompilerLog BlueprintLog(MessageLog);
+		WidgetClass->GetDefaultObject<UUserWidget>()->ValidateBlueprint(*WidgetBP->WidgetTree, BlueprintLog);
 
 		if (MessageLog.NumErrors == 0 && WidgetClass->bAllowTemplate)
 		{
diff --git a/Engine/Source/Editor/UnrealEd/Public/Kismet2/CompilerResultsLog.h b/Engine/Source/Editor/UnrealEd/Public/Kismet2/CompilerResultsLog.h
--- a/Engine/Source/Editor/UnrealEd/Public/Kismet2/CompilerResultsLog.h
+++ b/Engine/Source/Editor/UnrealEd/Public/Kismet2/CompilerResultsLog.h
@@ -237,6 +237,26 @@ public:
 		return Line;
 	}
 
+	/**
+	 * Store an already tokenized message.
+	 */
+	void AddTokenizedMessage(TSharedRef<FTokenizedMessage> InMessage)
+	{
+		switch (InMessage->GetSeverity())
+		{
+		case EMessageSeverity::Error:
+		case EMessageSeverity::CriticalError:
+			++NumErrors;
+			break;
+		case EMessageSeverity::Warning:
+		case EMessageSeverity::PerformanceWarning:
+			++NumWarnings;
+			break;
+		}
+
+		Messages.Add(InMessage);
+	}
+
 	/**
 	 * Commit all stored potential messages for a given node. Returns true if any messages were written.
 	 */
diff --git a/Engine/Source/Runtime/UMG/Private/Components/DynamicEntryBox.cpp b/Engine/Source/Runtime/UMG/Private/Components/DynamicEntryBox.cpp
--- a/Engine/Source/Runtime/UMG/Private/Components/DynamicEntryBox.cpp
+++ b/Engine/Source/Runtime/UMG/Private/Components/DynamicEntryBox.cpp
@@ -4,6 +4,7 @@
 #include "UMGPrivate.h"
 #include "Widgets/Layout/SWrapBox.h"
 #include "Widgets/SBoxPanel.h"
+#include "Editor/WidgetCompilerLog.h"
 
 #define LOCTEXT_NAMESPACE "UMG"
 
@@ -180,11 +181,11 @@ void UDynamicEntryBox::SetEntrySpacing(const FVector2D& InEntrySpacing)
 }
 
 #if WITH_EDITOR
-void UDynamicEntryBox::ValidateCompiledDefaults(class FCompilerResultsLog& CompileLog) const
+void UDynamicEntryBox::ValidateCompiledDefaults(IWidgetCompilerLog& CompileLog) const
 {
 	if (!EntryWidgetClass)
 	{
-		CompileLog.Error(*FText::Format(LOCTEXT("Error_DynamicEntryBox_MissingEntryClass", "{0} has no EntryWidgetClass specified - required for any Dynamic Entry Box to function."), FText::FromString(GetName())).ToString());
+		CompileLog.Error(FText::Format(LOCTEXT("Error_DynamicEntryBox_MissingEntryClass", "{0} has no EntryWidgetClass specified - required for any Dynamic Entry Box to function."), FText::FromString(GetName())));
 	}
 }
 #endif
diff --git a/Engine/Source/Runtime/UMG/Private/Components/ListViewBase.cpp b/Engine/Source/Runtime/UMG/Private/Components/ListViewBase.cpp
--- a/Engine/Source/Runtime/UMG/Private/Components/ListViewBase.cpp
+++ b/Engine/Source/Runtime/UMG/Private/Components/ListViewBase.cpp
@@ -5,6 +5,7 @@
 #include "Widgets/Text/STextBlock.h"
 #include "TimerManager.h"
 #include "Engine/Blueprint.h"
+#include "Editor/WidgetCompilerLog.h"
 
 #define LOCTEXT_NAMESPACE "UMG"
 
@@ -21,11 +22,11 @@ const FText UListViewBase::GetPaletteCategory()
 	return LOCTEXT("Lists", "Lists");
 }
 
-void UListViewBase::ValidateCompiledDefaults(FCompilerResultsLog& CompileLog) const
+void UListViewBase::ValidateCompiledDefaults(IWidgetCompilerLog& CompileLog) const
 {
 	if (!EntryWidgetClass)
 	{
-		CompileLog.Error(*FText::Format(LOCTEXT("Error_ListViewBase_MissingEntryClass", "{0} has no EntryWidgetClass specified - required for any UListViewBase to function."), FText::FromString(GetName())).ToString());
+		CompileLog.Error(FText::Format(LOCTEXT("Error_ListViewBase_MissingEntryClass", "{0} has no EntryWidgetClass specified - required for any UListViewBase to function."), FText::FromString(GetName())));
 	}
 }
 #endif
diff --git a/Engine/Source/Runtime/UMG/Private/UserWidget.cpp b/Engine/Source/Runtime/UMG/Private/UserWidget.cpp
--- a/Engine/Source/Runtime/UMG/Private/UserWidget.cpp
+++ b/Engine/Source/Runtime/UMG/Private/UserWidget.cpp
@@ -26,6 +26,7 @@
 #include "Compilation/MovieSceneCompiler.h"
 #include "TimerManager.h"
 #include "UObject/Package.h"
+#include "Editor/WidgetCompilerLog.h"
 
 #define LOCTEXT_NAMESPACE "UMG"
 
@@ -1366,7 +1367,7 @@ void UUserWidget::OnDesignerChanged(const FDesignerChangedEventArgs& EventArgs)
 	}
 }
 
-void UUserWidget::ValidateBlueprint(const UWidgetTree& BlueprintWidgetTree, FCompilerResultsLog& CompileLog) const
+void UUserWidget::ValidateBlueprint(const UWidgetTree& BlueprintWidgetTree, IWidgetCompilerLog& CompileLog) const
 {
 	ValidateCompiledDefaults(CompileLog);
 	ValidateCompiledWidgetTree(BlueprintWidgetTree, CompileLog);
diff --git a/Engine/Source/Runtime/UMG/Public/Blueprint/UserWidget.h b/Engine/Source/Runtime/UMG/Public/Blueprint/UserWidget.h
--- a/Engine/Source/Runtime/UMG/Public/Blueprint/UserWidget.h
+++ b/Engine/Source/Runtime/UMG/Public/Blueprint/UserWidget.h
@@ -1081,7 +1081,7 @@ public:
 	 * @see ValidateCompiledDefaults
 	 * @see ValidateCompiledWidgetTree
 	 */
-	void ValidateBlueprint(const UWidgetTree& BlueprintWidgetTree, FCompilerResultsLog& CompileLog) const;
+	void ValidateBlueprint(const UWidgetTree& BlueprintWidgetTree, class IWidgetCompilerLog& CompileLog) const;
 
 	/**
 	 * Override to perform any custom inspections of the default widget tree at the end of compilation.
@@ -1091,7 +1091,7 @@ public:
 	 *
 	 * Tip: If you need to validate properties of BindWidget members, you can search for them by property name within the widget tree.
 	 */
-	virtual void ValidateCompiledWidgetTree(const UWidgetTree& BlueprintWidgetTree, FCompilerResultsLog& CompileLog) const {};
+	virtual void ValidateCompiledWidgetTree(const UWidgetTree& BlueprintWidgetTree, class IWidgetCompilerLog& CompileLog) const {};
 #endif
 
 	static UUserWidget* CreateWidgetInstance(UWidget& OwningWidget, TSubclassOf<UUserWidget> UserWidgetClass, FName WidgetName);
diff --git a/Engine/Source/Runtime/UMG/Public/Components/DynamicEntryBox.h b/Engine/Source/Runtime/UMG/Public/Components/DynamicEntryBox.h
--- a/Engine/Source/Runtime/UMG/Public/Components/DynamicEntryBox.h
+++ b/Engine/Source/Runtime/UMG/Public/Components/DynamicEntryBox.h
@@ -93,7 +93,7 @@ public:
 	void SetEntrySpacing(const FVector2D& InEntrySpacing);
 
 #if WITH_EDITOR
-	virtual void ValidateCompiledDefaults(class FCompilerResultsLog& CompileLog) const override;
+	virtual void ValidateCompiledDefaults(class IWidgetCompilerLog& CompileLog) const override;
 #endif
 
 #if WITH_EDITORONLY_DATA
diff --git a/Engine/Source/Runtime/UMG/Public/Components/ListViewBase.h b/Engine/Source/Runtime/UMG/Public/Components/ListViewBase.h
--- a/Engine/Source/Runtime/UMG/Public/Components/ListViewBase.h
+++ b/Engine/Source/Runtime/UMG/Public/Components/ListViewBase.h
@@ -412,7 +412,7 @@ public:
 
 #if WITH_EDITOR
 	virtual const FText GetPaletteCategory() override;
-	virtual void ValidateCompiledDefaults(FCompilerResultsLog& CompileLog) const override;
+	virtual void ValidateCompiledDefaults(IWidgetCompilerLog& CompileLog) const override;
 #endif
 
 	TSubclassOf<UUserWidget> GetEntryWidgetClass() const { return EntryWidgetClass; }
diff --git a/Engine/Source/Runtime/UMG/Public/Components/Widget.h b/Engine/Source/Runtime/UMG/Public/Components/Widget.h
--- a/Engine/Source/Runtime/UMG/Public/Components/Widget.h
+++ b/Engine/Source/Runtime/UMG/Public/Components/Widget.h
@@ -735,7 +735,7 @@ public:
 	 * Allows UMG elements to evaluate their default states and determine whether they are acceptable.
 	 * To trigger compilation failure, add an error to the log. Warnings and notes will be visible, but will not cause compiles to fail.
 	 */
-	virtual void ValidateCompiledDefaults(class FCompilerResultsLog& CompileLog) const {}
+	virtual void ValidateCompiledDefaults(class IWidgetCompilerLog& CompileLog) const {}
 #else
 	FORCEINLINE bool IsDesignTime() const { return false; }
 #endif
diff --git a/Engine/Source/Runtime/UMG/Public/Editor/WidgetCompilerLog.h b/Engine/Source/Runtime/UMG/Public/Editor/WidgetCompilerLog.h
new file mode 100644
--- /dev/null
+++ b/Engine/Source/Runtime/UMG/Public/Editor/WidgetCompilerLog.h
@@ -0,0 +1,42 @@
+// Copyright 1998-2019 Epic Games, Inc. All Rights Reserved.
+
+#pragma once
+
+#include "CoreMinimal.h"
+#include "UObject/ObjectMacros.h"
+#include "UObject/Object.h"
+
+#if WITH_EDITOR
+
+class IWidgetCompilerLog
+{
+public:
+	TSharedRef<FTokenizedMessage> Error(const FText& Message)
+	{
+		TSharedRef<FTokenizedMessage> Line = FTokenizedMessage::Create(EMessageSeverity::Error);
+		Line->AddToken(FTextToken::Create(Message));
+		InternalLogMessage(Line);
+		return Line;
+	}
+
+	TSharedRef<FTokenizedMessage> Warning(const FText& Message)
+	{
+		TSharedRef<FTokenizedMessage> Line = FTokenizedMessage::Create(EMessageSeverity::Warning);
+		Line->AddToken(FTextToken::Create(Message));
+		InternalLogMessage(Line);
+		return Line;
+	}
+
+	TSharedRef<FTokenizedMessage> Note(const FText& Message)
+	{
+		TSharedRef<FTokenizedMessage> Line = FTokenizedMessage::Create(EMessageSeverity::Info);
+		Line->AddToken(FTextToken::Create(Message));
+		InternalLogMessage(Line);
+		return Line;
+	}
+	
+protected:
+	virtual void InternalLogMessage(TSharedRef<FTokenizedMessage>& Message) = 0;
+};
+
+#endif
\ No newline at end of file
diff --git a/Engine/Source/Runtime/UMG/UMG.Build.cs b/Engine/Source/Runtime/UMG/UMG.Build.cs
--- a/Engine/Source/Runtime/UMG/UMG.Build.cs
+++ b/Engine/Source/Runtime/UMG/UMG.Build.cs
@@ -57,14 +57,5 @@ public class UMG : ModuleRules
 			    }
             );
 		}
-
-		if (Target.Type == TargetType.Editor)
-		{
-			PublicDependencyModuleNames.AddRange(
-				new string[] {
-					"UnrealEd",
-				}
-			);
-		}
 	}
 }
