Fixing logic for displaying matching resolutions to the selected resolution in UMG

diff --git a/Engine/Source/Editor/UMGEditor/Private/Designer/SDesignerView.cpp b/Engine/Source/Editor/UMGEditor/Private/Designer/SDesignerView.cpp
--- a/Engine/Source/Editor/UMGEditor/Private/Designer/SDesignerView.cpp
+++ b/Engine/Source/Editor/UMGEditor/Private/Designer/SDesignerView.cpp
@@ -3325,8 +3325,8 @@ bool SDesignerView::HandleIsCommonResolutionSelected(FPlayScreenResolution InRes
 	}
 
 	// Compare to the size in the settings
-	const bool bSizeMatches = ((TestWidth == WidthReadFromSettings) && (TestHeight == HeightReadFromSettings))
-		|| (InResolution.bCanSwapAspectRatio && (PreviewWidth > PreviewHeight) && (TestHeight == WidthReadFromSettings) && (TestWidth == HeightReadFromSettings)); // flipped to landscape
+	const bool bSizeMatches = ((TestWidth == WidthReadFromSettings) && (TestHeight == HeightReadFromSettings)
+		|| ((bCanPreviewSwapAspectRatio && (PreviewWidth > PreviewHeight)) || InResolution.bCanSwapAspectRatio) && (TestHeight == WidthReadFromSettings) && (TestWidth == HeightReadFromSettings)); // flipped to landscape
 
 	if (!PreviewOverrideName.IsEmpty() || !InResolution.ProfileName.IsEmpty())
 	{
