Adding clearer messaging when custom rules aren't set up correctly

diff --git a/Engine/Source/Editor/UMGEditor/Private/Designer/SDesignerView.cpp b/Engine/Source/Editor/UMGEditor/Private/Designer/SDesignerView.cpp
--- a/Engine/Source/Editor/UMGEditor/Private/Designer/SDesignerView.cpp
+++ b/Engine/Source/Editor/UMGEditor/Private/Designer/SDesignerView.cpp
@@ -69,6 +69,7 @@
 #include "UMGEditorProjectSettings.h"
 #include "DeviceProfiles/DeviceProfile.h"
 #include "DeviceProfiles/DeviceProfileManager.h"
+#include "Engine/DPICustomScalingRule.h"
 
 #define LOCTEXT_NAMESPACE "UMG"
 
@@ -772,7 +773,7 @@ TSharedRef<SWidget> SDesignerView::CreateOverlayUI()
 				SNew(STextBlock)
 				.TextStyle(FEditorStyle::Get(), "Graph.ZoomText")
 				.Text(this, &SDesignerView::GetCurrentDPIScaleText)
-				.ColorAndOpacity(FLinearColor(1, 1, 1, 0.25f))
+				.ColorAndOpacity(this, &SDesignerView::GetCurrentDPIScaleColor)
 			]
 
 			+ SHorizontalBox::Slot()
@@ -3112,10 +3113,37 @@ FText SDesignerView::GetCurrentDPIScaleText() const
 	Options.MaximumFractionalDigits = 2;
 	Options.MinimumFractionalDigits = 1;
 
+	const UUserInterfaceSettings* UISettings = GetDefault<UUserInterfaceSettings>(UUserInterfaceSettings::StaticClass());
+	if (UISettings && UISettings->UIScaleRule == EUIScalingRule::Custom)
+	{
+		UClass* CustomScalingRuleClassInstance = UISettings->CustomScalingRuleClass.TryLoadClass<UDPICustomScalingRule>();
+
+		if (CustomScalingRuleClassInstance == nullptr)
+		{
+			return LOCTEXT("NoCustomRuleWarning", "Warning: Using Custom DPI Rule with no rules class set. Set a class in User Interface Project Settings. ");
+		}
+	}
+
 	FText DPIString = FText::AsNumber(GetPreviewDPIScale(), &Options, I18N.GetInvariantCulture());
 	return FText::Format(LOCTEXT("CurrentDPIScaleFormat", "DPI Scale {0}"), DPIString);
 }
 
+FSlateColor SDesignerView::GetCurrentDPIScaleColor() const
+{
+	const UUserInterfaceSettings* UISettings = GetDefault<UUserInterfaceSettings>(UUserInterfaceSettings::StaticClass());
+	if (UISettings && UISettings->UIScaleRule == EUIScalingRule::Custom)
+	{
+		UClass* CustomScalingRuleClassInstance = UISettings->CustomScalingRuleClass.TryLoadClass<UDPICustomScalingRule>();
+
+		if (CustomScalingRuleClassInstance == nullptr)
+		{
+			return FSlateColor(FLinearColor::Yellow);
+		}
+	}
+
+	return FSlateColor(FLinearColor(1, 1, 1, 0.25f));
+}
+
 FText SDesignerView::GetCurrentScaleFactorText() const
 {
 	FInternationalization& I18N = FInternationalization::Get();
diff --git a/Engine/Source/Editor/UMGEditor/Private/Designer/SDesignerView.h b/Engine/Source/Editor/UMGEditor/Private/Designer/SDesignerView.h
--- a/Engine/Source/Editor/UMGEditor/Private/Designer/SDesignerView.h
+++ b/Engine/Source/Editor/UMGEditor/Private/Designer/SDesignerView.h
@@ -182,6 +182,7 @@ private:
 
 	FText GetCurrentResolutionText() const;
 	FText GetCurrentDPIScaleText() const;
+	FSlateColor GetCurrentDPIScaleColor() const;
 	FText GetCurrentScaleFactorText() const;
 	FText GetCurrentSafeZoneText() const;
 	FSlateColor GetResolutionTextColorAndOpacity() const;
