Ambiguous logical operators, warning/error on Clang

Changed statement from/to:
A && B || ((C && D) || E) && F && G -> (A && B) || (((C && D || E) && F && G)

diff --git a/Engine/Source/Editor/UMGEditor/Private/Designer/SDesignerView.cpp b/Engine/Source/Editor/UMGEditor/Private/Designer/SDesignerView.cpp
--- a/Engine/Source/Editor/UMGEditor/Private/Designer/SDesignerView.cpp
+++ b/Engine/Source/Editor/UMGEditor/Private/Designer/SDesignerView.cpp
@@ -3353,8 +3353,8 @@ bool SDesignerView::HandleIsCommonResolutionSelected(FPlayScreenResolution InRes
 	}
 
 	// Compare to the size in the settings
-	const bool bSizeMatches = ((TestWidth == WidthReadFromSettings) && (TestHeight == HeightReadFromSettings)
-		|| ((bCanPreviewSwapAspectRatio && (PreviewWidth > PreviewHeight)) || InResolution.bCanSwapAspectRatio) && (TestHeight == WidthReadFromSettings) && (TestWidth == HeightReadFromSettings)); // flipped to landscape
+	const bool bSizeMatches = (((TestWidth == WidthReadFromSettings) && (TestHeight == HeightReadFromSettings))
+		|| (((bCanPreviewSwapAspectRatio && (PreviewWidth > PreviewHeight)) || InResolution.bCanSwapAspectRatio) && (TestHeight == WidthReadFromSettings) && (TestWidth == HeightReadFromSettings))); // flipped to landscape
 
 	if (!PreviewOverrideName.IsEmpty() || !InResolution.ProfileName.IsEmpty())
 	{
