Hide UMG widget extensions if widget is hidden

diff --git a/Engine/Source/Editor/UMGEditor/Private/Designer/SDesignerView.cpp b/Engine/Source/Editor/UMGEditor/Private/Designer/SDesignerView.cpp
--- a/Engine/Source/Editor/UMGEditor/Private/Designer/SDesignerView.cpp
+++ b/Engine/Source/Editor/UMGEditor/Private/Designer/SDesignerView.cpp
@@ -469,7 +469,7 @@ void SDesignerView::Construct(const FArguments& InArgs, TSharedPtr<FWidgetBluepr
 				.VAlign(VAlign_Fill)
 				[
 					SAssignNew(ExtensionWidgetCanvas, SCanvas)
-					.Visibility(EVisibility::SelfHitTestInvisible)
+					.Visibility(this, &SDesignerView::GetExtensionCanvasVisibility)
 				]
 
 				// Designer overlay UI, toolbar, status messages, zoom level...etc
@@ -497,6 +497,23 @@ void SDesignerView::Construct(const FArguments& InArgs, TSharedPtr<FWidgetBluepr
 	//RegisterActiveTimer(0.f, FWidgetActiveTimerDelegate::CreateSP(this, &SDesignerView::EnsureTick));
 }
 
+EVisibility SDesignerView::GetExtensionCanvasVisibility() const
+{
+	// If any selected widgets are hidden, then don't show widget extensions.
+	// If we want to support extensions on mixed-visibility in the future,
+	// every existing widget extension will probably need to be updated, as
+	// most do not check widget visibility before performing their function.
+	for (const FWidgetReference& Widget : GetSelectedWidgets())
+	{
+		UWidget* Preview = Widget.GetPreview();
+		if (!Preview || !Preview->IsVisibleInDesigner())
+		{
+			return EVisibility::Hidden;
+		}
+	}
+	return EVisibility::SelfHitTestInvisible;
+}
+
 EActiveTimerReturnType SDesignerView::EnsureTick(double InCurrentTime, float InDeltaTime)
 {
 	return EActiveTimerReturnType::Continue;
diff --git a/Engine/Source/Editor/UMGEditor/Private/Designer/SDesignerView.h b/Engine/Source/Editor/UMGEditor/Private/Designer/SDesignerView.h
--- a/Engine/Source/Editor/UMGEditor/Private/Designer/SDesignerView.h
+++ b/Engine/Source/Editor/UMGEditor/Private/Designer/SDesignerView.h
@@ -228,6 +228,8 @@ private:
 
 	UWidget* GetWidgetInDesignScopeFromSlateWidget(TSharedRef<SWidget>& InWidget);
 
+	EVisibility GetExtensionCanvasVisibility() const;
+
 private:
 	struct FWidgetHitResult
 	{
