UMG Compiler - During compilation, we now check to see if the widget tree on widget blueprints is valid.  Through some unknown means, the widget tree on the WidgetBlueprint was being detatched as a DefaultSubobject, and WidgetBlueprints ended up with two WidgetTrees.  This bit of code now scans to make sure there is no abandoned tree, and fixes up the correct widget tree to be the proper default subobject with the correct name.

diff --git a/Engine/Source/Editor/UMGEditor/Private/WidgetBlueprintCompiler.cpp b/Engine/Source/Editor/UMGEditor/Private/WidgetBlueprintCompiler.cpp
--- a/Engine/Source/Editor/UMGEditor/Private/WidgetBlueprintCompiler.cpp
+++ b/Engine/Source/Editor/UMGEditor/Private/WidgetBlueprintCompiler.cpp
@@ -553,6 +553,27 @@ void FWidgetBlueprintCompilerContext::SanitizeBindings(UBlueprintGeneratedClass*
 	WidgetBP->PropertyBindings = AttributeBindings;
 }
 
+void FWidgetBlueprintCompilerContext::FixAbandonedWidgetTree(UWidgetBlueprint* WidgetBP)
+{
+	UWidgetTree* WidgetTree = WidgetBP->WidgetTree;
+
+	if (ensure(WidgetTree))
+	{
+		if (WidgetTree->GetName() != TEXT("WidgetTree"))
+		{
+			if (UWidgetTree* AbandonedWidgetTree = static_cast<UWidgetTree*>(FindObjectWithOuter(WidgetBP, UWidgetTree::StaticClass(), TEXT("WidgetTree"))))
+			{
+				AbandonedWidgetTree->ClearFlags(RF_DefaultSubObject);
+				AbandonedWidgetTree->SetFlags(RF_Transient);
+				AbandonedWidgetTree->Rename(nullptr, GetTransientPackage(), REN_DontCreateRedirectors | REN_ForceNoResetLoaders | REN_NonTransactional | REN_DoNotDirty);
+			}
+
+			WidgetTree->Rename(TEXT("WidgetTree"), nullptr, REN_DontCreateRedirectors | REN_ForceNoResetLoaders | REN_NonTransactional | REN_DoNotDirty);
+			WidgetTree->SetFlags(RF_DefaultSubObject);
+		}
+	}
+}
+
 void FWidgetBlueprintCompilerContext::FinishCompilingClass(UClass* Class)
 {
 	UWidgetBlueprint* WidgetBP = WidgetBlueprint();
@@ -567,6 +588,8 @@ void FWidgetBlueprintCompilerContext::FinishCompilingClass(UClass* Class)
 			UBlueprint::ForceLoadMembers(WidgetBP->WidgetTree);
 		}
 
+		FixAbandonedWidgetTree(WidgetBP);
+
 		BPGClass->bCookSlowConstructionWidgetTree = GetDefault<UUMGEditorProjectSettings>()->CompilerOption_CookSlowConstructionWidgetTree(WidgetBP);
 
 		BPGClass->WidgetTree = Cast<UWidgetTree>(StaticDuplicateObject(WidgetBP->WidgetTree, BPGClass, NAME_None, RF_AllFlags & ~RF_DefaultSubObject));
diff --git a/Engine/Source/Editor/UMGEditor/Public/WidgetBlueprintCompiler.h b/Engine/Source/Editor/UMGEditor/Public/WidgetBlueprintCompiler.h
--- a/Engine/Source/Editor/UMGEditor/Public/WidgetBlueprintCompiler.h
+++ b/Engine/Source/Editor/UMGEditor/Public/WidgetBlueprintCompiler.h
@@ -80,6 +80,7 @@ protected:
 	void VerifyEventReplysAreNotEmpty(FKismetFunctionContext& Context);
 
 protected:
+	void FixAbandonedWidgetTree(UWidgetBlueprint* WidgetBP);
 
 	UWidgetBlueprintGeneratedClass* NewWidgetBlueprintClass;
 
