UMG: Fix replace updating the animation binding that matches the name of the widget, rather than using the animation guid.

diff --git a/Engine/Source/Editor/UMGEditor/Private/WidgetBlueprintEditor.cpp b/Engine/Source/Editor/UMGEditor/Private/WidgetBlueprintEditor.cpp
--- a/Engine/Source/Editor/UMGEditor/Private/WidgetBlueprintEditor.cpp
+++ b/Engine/Source/Editor/UMGEditor/Private/WidgetBlueprintEditor.cpp
@@ -1393,10 +1393,9 @@ void FWidgetBlueprintEditor::ReplaceTrackWithSelectedWidget(FWidgetReference Sel
 		check(SourceBinding);
 
 		// Set binding names used for lookup
-		FName PredicateName = SourceBinding->WidgetName;
 		for (FWidgetAnimationBinding& Binding : WidgetAnimation->AnimationBindings)
 		{
-			if (Binding.WidgetName != PredicateName)
+			if (Binding.AnimationGuid != ObjectId)
 			{
 				continue;
 			}
