Editgrate CL 4923001 from Dev-Editor to 4.22

diff --git a/Engine/Source/Editor/UMGEditor/Private/Designer/SDesignerView.cpp b/Engine/Source/Editor/UMGEditor/Private/Designer/SDesignerView.cpp
--- a/Engine/Source/Editor/UMGEditor/Private/Designer/SDesignerView.cpp
+++ b/Engine/Source/Editor/UMGEditor/Private/Designer/SDesignerView.cpp
@@ -3071,25 +3071,6 @@ FReply SDesignerView::OnDrop(const FGeometry& MyGeometry, const FDragDropEvent&
 		// Regenerate extension widgets now that we've finished moving or placing the widget.
 		CreateExtensionWidgetsForSelection();
 
-		UPanelSlot* Slot;
-		UWidgetTree* WidgetTree = BP->WidgetTree;
-		if (WidgetTree && SelectedDragDropOp.IsValid())
-		{
-			for (auto& DraggedWidget : SelectedDragDropOp->DraggedWidgets)
-			{
-				FWidgetReference Reference = BlueprintEditor.Pin()->GetReferenceFromTemplate(DraggedWidget.Template);
-				UWidget* LocalPreviewWidget = Reference.GetPreview();
-				if (LocalPreviewWidget && LocalPreviewWidget->GetParent())
-				{
-					Slot = LocalPreviewWidget->GetParent()->GetSlots()[LocalPreviewWidget->GetParent()->GetChildIndex(LocalPreviewWidget)];
-					if (Slot != nullptr)
-					{
-						FWidgetBlueprintEditorUtils::ImportPropertiesFromText(Slot, DraggedWidget.ExportedSlotProperties);
-					}
-				}
-			}
-		}
-
 		DropPreviews.Empty();
 		return FReply::Handled().SetUserFocus(SharedThis(this));
 	}
