Fix for circular dependencies in widget BPs

diff --git a/Engine/Source/Editor/UMGEditor/Private/Hierarchy/SHierarchyViewItem.cpp b/Engine/Source/Editor/UMGEditor/Private/Hierarchy/SHierarchyViewItem.cpp
--- a/Engine/Source/Editor/UMGEditor/Private/Hierarchy/SHierarchyViewItem.cpp
+++ b/Engine/Source/Editor/UMGEditor/Private/Hierarchy/SHierarchyViewItem.cpp
@@ -609,6 +609,26 @@ void FHierarchyRoot::UpdateSelection()
 
 TOptional<EItemDropZone> FHierarchyRoot::HandleCanAcceptDrop(const FDragDropEvent& DragDropEvent, EItemDropZone DropZone)
 {
+	TSharedPtr<FWidgetTemplateDragDropOp> TemplateDragDropOp = DragDropEvent.GetOperationAs<FWidgetTemplateDragDropOp>();
+	if (TemplateDragDropOp.IsValid())
+	{
+		UWidgetBlueprint* Blueprint = BlueprintEditor.Pin()->GetWidgetBlueprintObj();
+		if(Blueprint)
+		{
+			UWidget* Widget = TemplateDragDropOp->Template->Create(Blueprint->WidgetTree);
+	
+			if (Widget)
+			{
+				if (!Blueprint->IsWidgetFreeFromCircularReferences(Cast<UUserWidget>(Widget)))
+				{
+					TemplateDragDropOp->CurrentIconBrush = FEditorStyle::GetBrush(TEXT("Graph.ConnectorFeedback.Error"));
+					TemplateDragDropOp->CurrentHoverText = LOCTEXT("CircularReference", "This would cause a circular reference.");
+					return TOptional<EItemDropZone>();
+				}
+			}
+		}
+	}
+
 	bool bIsDrop = false;
 	return ProcessHierarchyDragDrop(DragDropEvent, DropZone, bIsDrop, BlueprintEditor.Pin(), FWidgetReference());
 }
@@ -738,6 +758,18 @@ TOptional<EItemDropZone> FNamedSlotModel::HandleCanAcceptDrop(const FDragDropEve
 				return TOptional<EItemDropZone>();
 			}
 
+			UWidget* Widget = TemplateDragDropOp->Template->Create(Blueprint->WidgetTree);
+
+			if (Widget)
+			{
+				if (!Blueprint->IsWidgetFreeFromCircularReferences(Cast<UUserWidget>(Widget)))
+				{
+					TemplateDragDropOp->CurrentIconBrush = FEditorStyle::GetBrush(TEXT("Graph.ConnectorFeedback.Error"));
+					TemplateDragDropOp->CurrentHoverText = LOCTEXT("CircularReference", "This would cause a circular reference.");
+					return TOptional<EItemDropZone>();
+				}
+			}
+
 			TemplateDragDropOp->CurrentIconBrush = FEditorStyle::GetBrush(TEXT("Graph.ConnectorFeedback.OK"));
 			return EItemDropZone::OntoItem;
 		}
@@ -954,6 +986,25 @@ FSlateFontInfo FHierarchyWidget::GetFont() const
 
 TOptional<EItemDropZone> FHierarchyWidget::HandleCanAcceptDrop(const FDragDropEvent& DragDropEvent, EItemDropZone DropZone)
 {
+	TSharedPtr<FWidgetTemplateDragDropOp> TemplateDragDropOp = DragDropEvent.GetOperationAs<FWidgetTemplateDragDropOp>();
+	if (TemplateDragDropOp.IsValid())
+	{
+		UWidgetBlueprint* Blueprint = BlueprintEditor.Pin()->GetWidgetBlueprintObj();
+		if(Blueprint)
+		{
+			UWidget* Widget = TemplateDragDropOp->Template->Create(Blueprint->WidgetTree);
+
+			if (Widget)
+			{
+				if (!Blueprint->IsWidgetFreeFromCircularReferences(Cast<UUserWidget>(Widget)))
+				{
+					TemplateDragDropOp->CurrentIconBrush = FEditorStyle::GetBrush(TEXT("Graph.ConnectorFeedback.Error"));
+					TemplateDragDropOp->CurrentHoverText = LOCTEXT("CircularReference", "This would cause a circular reference.");
+					return TOptional<EItemDropZone>();
+				}
+			}
+		}
+	}
 	bool bIsDrop = false;
 	return ProcessHierarchyDragDrop(DragDropEvent, DropZone, bIsDrop, BlueprintEditor.Pin(), Item);
 }
diff --git a/Engine/Source/Editor/UMGEditor/Private/WidgetBlueprint.cpp b/Engine/Source/Editor/UMGEditor/Private/WidgetBlueprint.cpp
--- a/Engine/Source/Editor/UMGEditor/Private/WidgetBlueprint.cpp
+++ b/Engine/Source/Editor/UMGEditor/Private/WidgetBlueprint.cpp
@@ -786,8 +786,28 @@ bool UWidgetBlueprint::IsWidgetFreeFromCircularReferences(UUserWidget* UserWidge
 			// will cause a circular reference within the blueprint.
 			return false;
 		}
+		else if (UWidgetBlueprint* GeneratedByBlueprint = Cast<UWidgetBlueprint>(UserWidget->WidgetGeneratedBy))
+		{
+			// Check the generated by blueprints - this will catch even cases where one has the other in the widget tree but hasn't compiled yet
+			if (GeneratedByBlueprint->WidgetTree && GeneratedByBlueprint->WidgetTree->RootWidget)
+			{
+				TArray<UWidget*> ChildWidgets;
+				GeneratedByBlueprint->WidgetTree->GetChildWidgets(GeneratedByBlueprint->WidgetTree->RootWidget, ChildWidgets);
+				for (UWidget* ChildWidget : ChildWidgets)
+				{
+					if (UWidgetBlueprint* ChildGeneratedBlueprint = Cast<UWidgetBlueprint>(ChildWidget->WidgetGeneratedBy))
+					{
+						if (this == ChildGeneratedBlueprint)
+						{
+							return false;
+						}
+					}
+				}
+			}
+		}
 		else if (UserWidget->WidgetTree)
 		{
+			// This loop checks for references that existed in the compiled blueprint, in case it's changed since then
 			TArray<UWidget*> ChildWidgets;
 			UserWidget->WidgetTree->GetAllWidgets(ChildWidgets);
 
