Making the hierarchy view show if you have flow preference changes in the view.

diff --git a/Engine/Source/Editor/UMGEditor/Private/Hierarchy/SHierarchyViewItem.cpp b/Engine/Source/Editor/UMGEditor/Private/Hierarchy/SHierarchyViewItem.cpp
--- a/Engine/Source/Editor/UMGEditor/Private/Hierarchy/SHierarchyViewItem.cpp
+++ b/Engine/Source/Editor/UMGEditor/Private/Hierarchy/SHierarchyViewItem.cpp
@@ -607,6 +607,17 @@ void FHierarchyRoot::UpdateSelection()
 	}
 }
 
+bool FHierarchyRoot::DoesWidgetOverrideFlowDirection() const
+{
+	TSharedPtr<FWidgetBlueprintEditor> BPEd = BlueprintEditor.Pin();
+	if (UWidget* Default = BPEd->GetWidgetBlueprintObj()->GeneratedClass->GetDefaultObject<UWidget>())
+	{
+		return Default->FlowDirectionPreference != EFlowDirectionPreference::Inherit;
+	}
+
+	return false;
+}
+
 TOptional<EItemDropZone> FHierarchyRoot::HandleCanAcceptDrop(const FDragDropEvent& DragDropEvent, EItemDropZone DropZone)
 {
 	TSharedPtr<FWidgetTemplateDragDropOp> TemplateDragDropOp = DragDropEvent.GetOperationAs<FWidgetTemplateDragDropOp>();
@@ -1165,6 +1176,19 @@ void SHierarchyViewItem::Construct(const FArguments& InArgs, const TSharedRef< S
 				.IsSelected(this, &SHierarchyViewItem::IsSelectedExclusively)
 			]
 
+			// Flow Direction Icon
+			+ SHorizontalBox::Slot()
+			.AutoWidth()
+			.VAlign(VAlign_Center)
+			[
+				SNew(STextBlock)
+				.ToolTipText(LOCTEXT("FlowDirectionHierarchyToolTip", "This widget overrides the flow direction preference."))
+				.Visibility_Lambda([InModel] { return InModel->DoesWidgetOverrideFlowDirection() ? EVisibility::Visible : EVisibility::Collapsed; })
+				.ColorAndOpacity(FCoreStyle::Get().GetSlateColor("Foreground"))
+				.Font(FEditorStyle::Get().GetFontStyle("FontAwesome.10"))
+				.Text(FEditorFontGlyphs::Exchange)
+			]
+
 			// Locked Icon
 			+ SHorizontalBox::Slot()
 			.AutoWidth()
diff --git a/Engine/Source/Editor/UMGEditor/Private/Hierarchy/SHierarchyViewItem.h b/Engine/Source/Editor/UMGEditor/Private/Hierarchy/SHierarchyViewItem.h
--- a/Engine/Source/Editor/UMGEditor/Private/Hierarchy/SHierarchyViewItem.h
+++ b/Engine/Source/Editor/UMGEditor/Private/Hierarchy/SHierarchyViewItem.h
@@ -82,6 +82,8 @@ public:
 		}
 	}
 
+	virtual bool DoesWidgetOverrideFlowDirection() const { return false; }
+
 	virtual bool IsExpanded() const { return true; }
 	virtual void SetExpanded(bool bIsExpanded) { }
 
@@ -130,6 +132,8 @@ public:
 
 	virtual void OnSelection() override;
 
+	virtual bool DoesWidgetOverrideFlowDirection() const override;
+
 	virtual TOptional<EItemDropZone> HandleCanAcceptDrop(const FDragDropEvent& DragDropEvent, EItemDropZone DropZone) override;
 	virtual FReply HandleAcceptDrop(FDragDropEvent const& DragDropEvent, EItemDropZone DropZone) override;
 
@@ -263,6 +267,17 @@ public:
 		}
 	}
 
+	virtual bool DoesWidgetOverrideFlowDirection() const override
+	{
+		UWidget* TemplateWidget = Item.GetTemplate();
+		if (TemplateWidget)
+		{
+			return TemplateWidget->FlowDirectionPreference != EFlowDirectionPreference::Inherit;
+		}
+
+		return false;
+	}
+
 	virtual bool IsExpanded() const override
 	{
 		if (UWidget* Template = Item.GetTemplate())
