Widgets no longer will reorder when being moved in the widget editor

diff --git a/Engine/Source/Editor/UMGEditor/Private/Designer/SDesignerView.cpp b/Engine/Source/Editor/UMGEditor/Private/Designer/SDesignerView.cpp
--- a/Engine/Source/Editor/UMGEditor/Private/Designer/SDesignerView.cpp
+++ b/Engine/Source/Editor/UMGEditor/Private/Designer/SDesignerView.cpp
@@ -2963,7 +2963,7 @@ void SDesignerView::ProcessDropAndAddWidget(const FGeometry& MyGeometry, const F
 					}
 					else
 					{
-						Slot = ParentWidget->AddChild(Widget);
+						Slot = ParentWidget->InsertChildAt(ParentWidget->GetChildIndex(Widget), Widget);
 					}
 
 					if (Slot != nullptr)
