UMG: Add new binding to selected folder

diff --git a/Engine/Source/Editor/Sequencer/Private/Sequencer.cpp b/Engine/Source/Editor/Sequencer/Private/Sequencer.cpp
--- a/Engine/Source/Editor/Sequencer/Private/Sequencer.cpp
+++ b/Engine/Source/Editor/Sequencer/Private/Sequencer.cpp
@@ -5455,6 +5455,12 @@ void FSequencer::GetSelectedSections(TArray<UMovieSceneSection*>& OutSelectedSec
 	}
 }
 
+void FSequencer::GetSelectedFolders(TArray<UMovieSceneFolder*>& OutSelectedFolders)
+{
+	FString OutNewNodePath;
+	CalculateSelectedFolderAndPath(OutSelectedFolders, OutNewNodePath);
+}
+
 void FSequencer::SelectObject(FGuid ObjectBinding)
 {
 	const TSharedPtr<FSequencerObjectBindingNode>* Node = NodeTree->GetObjectBindingMap().Find(ObjectBinding);
diff --git a/Engine/Source/Editor/Sequencer/Private/Sequencer.h b/Engine/Source/Editor/Sequencer/Private/Sequencer.h
--- a/Engine/Source/Editor/Sequencer/Private/Sequencer.h
+++ b/Engine/Source/Editor/Sequencer/Private/Sequencer.h
@@ -700,6 +700,7 @@ public:
 	virtual FSequencerSelectionPreview& GetSelectionPreview() override;
 	virtual void GetSelectedTracks(TArray<UMovieSceneTrack*>& OutSelectedTracks) override;
 	virtual void GetSelectedSections(TArray<UMovieSceneSection*>& OutSelectedSections) override;
+	virtual void GetSelectedFolders(TArray<UMovieSceneFolder*>& OutSelectedFolders) override;
 	virtual void SelectObject(FGuid ObjectBinding) override;
 	virtual void SelectTrack(UMovieSceneTrack* Track) override;
 	virtual void SelectSection(UMovieSceneSection* Section) override;
diff --git a/Engine/Source/Editor/Sequencer/Public/ISequencer.h b/Engine/Source/Editor/Sequencer/Public/ISequencer.h
--- a/Engine/Source/Editor/Sequencer/Public/ISequencer.h
+++ b/Engine/Source/Editor/Sequencer/Public/ISequencer.h
@@ -466,6 +466,9 @@ public:
 	/** Gets the currently selected sections. */
 	virtual void GetSelectedSections(TArray<UMovieSceneSection*>& OutSelectedSections) = 0;
 
+	/** Gets the currently selected folders. */
+	virtual void GetSelectedFolders(TArray<UMovieSceneFolder*>& OutSelectedFolders) = 0;
+
 	/** Selects an object by GUID */
 	virtual void SelectObject(FGuid ObjectBinding) = 0;
 
diff --git a/Engine/Source/Editor/UMGEditor/Private/WidgetBlueprintEditor.cpp b/Engine/Source/Editor/UMGEditor/Private/WidgetBlueprintEditor.cpp
--- a/Engine/Source/Editor/UMGEditor/Private/WidgetBlueprintEditor.cpp
+++ b/Engine/Source/Editor/UMGEditor/Private/WidgetBlueprintEditor.cpp
@@ -2,6 +2,7 @@
 
 #include "WidgetBlueprintEditor.h"
 #include "MovieSceneBinding.h"
+#include "MovieSceneFolder.h"
 #include "MovieScene.h"
 #include "Animation/WidgetAnimation.h"
 #include "Widgets/Text/STextBlock.h"
@@ -1200,7 +1201,15 @@ void FWidgetBlueprintEditor::AddObjectToAnimation(UObject* ObjectToAnimate)
 	Sequencer->GetFocusedMovieSceneSequence()->Modify();
 
 	// @todo Sequencer - Make this kind of adding more explicit, this current setup seem a bit brittle.
-	Sequencer->GetHandleToObject(ObjectToAnimate);
+	FGuid NewGuid = Sequencer->GetHandleToObject(ObjectToAnimate);
+
+	TArray<UMovieSceneFolder*> SelectedParentFolders;
+	Sequencer->GetSelectedFolders(SelectedParentFolders);
+
+	if (SelectedParentFolders.Num() > 0)
+	{
+		SelectedParentFolders[0]->AddChildObjectBinding(NewGuid);
+	}
 }
 
 TSharedRef<FExtender> FWidgetBlueprintEditor::GetAddTrackSequencerExtender( const TSharedRef<FUICommandList> CommandList, const TArray<UObject*> ContextSensitiveObjects )
