When pasting Canvas Panel Slots in UMG Animations, use get handle, and remove the existing binding to rebind, instead of rebinding to a parent.

diff --git a/Engine/Source/Editor/UMGEditor/Private/WidgetBlueprintEditor.cpp b/Engine/Source/Editor/UMGEditor/Private/WidgetBlueprintEditor.cpp
--- a/Engine/Source/Editor/UMGEditor/Private/WidgetBlueprintEditor.cpp
+++ b/Engine/Source/Editor/UMGEditor/Private/WidgetBlueprintEditor.cpp
@@ -1532,6 +1532,16 @@ void FWidgetBlueprintEditor::OnMovieSceneBindingsPasted(const TArray<FMovieScene
 						AnimationSequence->BindPossessableObject(BindingPasted.GetObjectGuid(), *BindableObject.Object, BindingContext);			
 						break;
 					}
+
+					// Special case for canvas slots, they need to be added again
+					if (BindableObject.Object->GetFName().ToString() == BindingPasted.GetName())
+					{
+						// Create handle, to rebind correctly
+						Sequencer->GetHandleToObject(BindableObject.Object);
+						// Remove the existing binding, as it is now replaced by the that was just added by getting the handle
+						AnimationSequence->GetMovieScene()->RemovePossessable(BindingPasted.GetObjectGuid());
+						break;
+					}
 				}
 			}
 		}
