Fix UMG Sequencer editor's "All Named Widgets" menu having erroneous "Search Start" header.

diff --git a/Engine/Source/Editor/UMGEditor/Private/WidgetBlueprintEditor.cpp b/Engine/Source/Editor/UMGEditor/Private/WidgetBlueprintEditor.cpp
--- a/Engine/Source/Editor/UMGEditor/Private/WidgetBlueprintEditor.cpp
+++ b/Engine/Source/Editor/UMGEditor/Private/WidgetBlueprintEditor.cpp
@@ -1170,8 +1170,6 @@ void FWidgetBlueprintEditor::OnGetAnimationAddMenuContent(FMenuBuilder& MenuBuil
 
 void FWidgetBlueprintEditor::OnGetAnimationAddMenuContentAllWidgets(FMenuBuilder& MenuBuilder)
 {
-	MenuBuilder.AddSearchWidget();
-
 	TArray<FObjectAndDisplayName> BindableObjects;
 	{
 		GetBindableObjects(GetPreview()->WidgetTree, BindableObjects);
