Add multiple widget bindings for UMG

diff --git a/Engine/Source/Editor/UMGEditor/Private/WidgetBlueprintEditor.cpp b/Engine/Source/Editor/UMGEditor/Private/WidgetBlueprintEditor.cpp
--- a/Engine/Source/Editor/UMGEditor/Private/WidgetBlueprintEditor.cpp
+++ b/Engine/Source/Editor/UMGEditor/Private/WidgetBlueprintEditor.cpp
@@ -1225,34 +1225,78 @@ void FWidgetBlueprintEditor::OnBuildCustomContextMenuForGuid(FMenuBuilder& MenuB
 {
 	if (CurrentAnimation.IsValid())
 	{
-		FWidgetReference SelectedWidget;
-		if (SelectedWidgets.Num() == 1)
+		TArray<FWidgetReference> ValidSelectedWidgets;
+		for (FWidgetReference SelectedWidget : SelectedWidgets)
 		{
-			for (FWidgetReference Widget : SelectedWidgets)
+			if (SelectedWidget.IsValid())
 			{
-				SelectedWidget = Widget;
+				//need to make sure it's a widget, if not bound assume it is.
+				UWidget* BoundWidget = nullptr;
+				bool bNotBound = true;
+				for (TWeakObjectPtr<> WeakObjectPtr : GetSequencer()->FindObjectsInCurrentSequence(ObjectBinding))
+				{
+					BoundWidget = Cast<UWidget>(WeakObjectPtr.Get());
+					bNotBound = false;
+					break;
+				}
+
+				if (bNotBound || (BoundWidget && SelectedWidget.GetPreview()->GetTypedOuter<UWidgetTree>() == BoundWidget->GetTypedOuter<UWidgetTree>()))
+				{
+					ValidSelectedWidgets.Add(SelectedWidget);
+				}
 			}
 		}
-		if (SelectedWidget.IsValid())
+		
+		if(ValidSelectedWidgets.Num() > 0)
 		{
-			//need to make sure it's a widget, if not bound assume it is.
-			UWidget* BoundWidget = nullptr;
-			bool bNotBound = true;
-			for (TWeakObjectPtr<> WeakObjectPtr : GetSequencer()->FindObjectsInCurrentSequence(ObjectBinding))
+			MenuBuilder.AddMenuSeparator();
+			
+			MenuBuilder.AddMenuEntry(
+				LOCTEXT("AddSelectedToBinding", "Add Selected"),
+				LOCTEXT("AddSelectedToBindingToolTip", "Add selected objects to this track"),
+				FSlateIcon(),
+				FExecuteAction::CreateRaw(this, &FWidgetBlueprintEditor::AddWidgetsToTrack, ValidSelectedWidgets, ObjectBinding)
+			);
+			
+			if (ValidSelectedWidgets.Num() > 1)
 			{
-				BoundWidget = Cast<UWidget>(WeakObjectPtr.Get());
-				bNotBound = false;
-				break;
+				MenuBuilder.AddMenuEntry(
+					LOCTEXT("ReplaceBindingWithSelected", "Replace with Selected"),
+					LOCTEXT("ReplaceBindingWithSelectedToolTip", "Replace the object binding with selected objects"),
+					FSlateIcon(),
+					FExecuteAction::CreateRaw(this, &FWidgetBlueprintEditor::ReplaceTrackWithWidgets, ValidSelectedWidgets, ObjectBinding)
+				);
 			}
-			if (bNotBound || (BoundWidget && SelectedWidget.GetPreview()->GetTypedOuter<UWidgetTree>() == BoundWidget->GetTypedOuter<UWidgetTree>()) )
+			else
 			{
-				FUIAction ReplaceWithMenuAction(FExecuteAction::CreateRaw(this, &FWidgetBlueprintEditor::ReplaceTrackWithSelectedWidget, SelectedWidget, (UWidget*)nullptr, ObjectBinding));
-
-				FText ReplaceWithLabel = FText::Format(LOCTEXT("ReplaceObject", "Replace with {0}"), FText::FromString(SelectedWidget.GetPreview()->GetName()));
-				FText ReplaceWithToolTip = FText::Format(LOCTEXT("ReplaceObjectToolTip", "Replace the bound widget in this animation with {0}"), FText::FromString(SelectedWidget.GetPreview()->GetName()));
-				MenuBuilder.AddMenuSeparator();
-				MenuBuilder.AddMenuEntry(ReplaceWithLabel, ReplaceWithToolTip, FSlateIcon(), ReplaceWithMenuAction);
+				MenuBuilder.AddMenuEntry(
+					FText::Format(LOCTEXT("ReplaceObject", "Replace with {0}"), FText::FromString(ValidSelectedWidgets[0].GetPreview()->GetName())),
+					FText::Format(LOCTEXT("ReplaceObjectToolTip", "Replace the bound widget in this animation with {0}"), FText::FromString(ValidSelectedWidgets[0].GetPreview()->GetName())),
+					FSlateIcon(),
+					FExecuteAction::CreateRaw(this, &FWidgetBlueprintEditor::ReplaceTrackWithWidgets, ValidSelectedWidgets, ObjectBinding)
+				);
 			}
+			
+			MenuBuilder.AddMenuEntry(
+				LOCTEXT("RemoveSelectedFromBinding", "Remove Selected"),
+				LOCTEXT("RemoveSelectedFromBindingToolTip", "Remove selected objects from this track"),
+				FSlateIcon(),
+				FExecuteAction::CreateRaw(this, &FWidgetBlueprintEditor::RemoveWidgetsFromTrack, ValidSelectedWidgets, ObjectBinding)
+			);
+
+			MenuBuilder.AddMenuEntry(
+				LOCTEXT("RemoveAllBindings", "Remove All"),
+				LOCTEXT("RemoveAllBindingsToolTip", "Remove all bound objects from this track"),
+				FSlateIcon(),
+				FExecuteAction::CreateRaw(this, &FWidgetBlueprintEditor::RemoveAllWidgetsFromTrack, ObjectBinding)
+			);
+			
+			MenuBuilder.AddMenuEntry(
+				LOCTEXT("RemoveMissing", "Remove Missing"),
+				LOCTEXT("RemoveMissingToolTip", "Remove missing objects bound to this track"),
+				FSlateIcon(),
+				FExecuteAction::CreateRaw(this, &FWidgetBlueprintEditor::RemoveMissingWidgetsFromTrack, ObjectBinding)
+			);
 		}
 	}
 }
@@ -1307,35 +1351,33 @@ void FWidgetBlueprintEditor::ExtendSequencerAddTrackMenu( FMenuBuilder& AddTrack
 	}
 }
 
-void FWidgetBlueprintEditor::ReplaceTrackWithSelectedWidget(FWidgetReference SelectedWidget, UWidget* BoundWidget, FGuid ObjectId)
+void FWidgetBlueprintEditor::AddWidgetsToTrack(const TArray<FWidgetReference> Widgets, FGuid ObjectId)
 {
-	const FScopedTransaction Transaction( LOCTEXT( "ReplaceTrackWithSelectedWidget", "Replace Track with Selected Widget" ) );
+	const FScopedTransaction Transaction(LOCTEXT("AddSelectedWidgetsToTrack", "Add Widgets to Track"));
 
 	UWidgetAnimation* WidgetAnimation = Cast<UWidgetAnimation>(Sequencer->GetFocusedMovieSceneSequence());
 	UMovieScene* MovieScene = WidgetAnimation->GetMovieScene();
-	UWidget* PreviewWidget = SelectedWidget.GetPreview();
-	UWidget* TemplateWidget = SelectedWidget.GetTemplate();
-	// Try find if the SelectedWidget is already bound, if so return
-	FGuid SelectedWidgetId = Sequencer->FindObjectId(*PreviewWidget, MovieSceneSequenceID::Root);
-	if (SelectedWidgetId.IsValid())
-	{
-		FNotificationInfo Info(LOCTEXT("SelectedWidgetalreadybound", "Selected Widget already bound"));
-		Info.FadeInDuration = 0.1f;
-		Info.FadeOutDuration = 0.5f;
-		Info.ExpireDuration = 2.5f;
-		auto NotificationItem = FSlateNotificationManager::Get().AddNotification(Info);
 
-		NotificationItem->SetCompletionState(SNotificationItem::CS_Success);
-		NotificationItem->ExpireAndFadeout();
-		return;
-	}
+	TArray<FWidgetReference> WidgetsToAdd(Widgets);
+
+	bool bIncompatibleTracks = false;
 
-	if (!BoundWidget || TemplateWidget->GetClass() != BoundWidget->GetClass())
+	for (const FWidgetReference Widget : Widgets)
 	{
+		UWidget* PreviewWidget = Widget.GetPreview();
+
+		// Try find if the SelectedWidget is already bound, if so return
+		FGuid SelectedWidgetId = Sequencer->FindObjectId(*PreviewWidget, MovieSceneSequenceID::Root);
+		if (SelectedWidgetId.IsValid())
+		{
+			WidgetsToAdd.Remove(Widget);
+			continue;
+		}
+
 		TArray<FMovieSceneBinding> MovieSceneBindings = MovieScene->GetBindings();
 		for (FMovieSceneBinding Binding : MovieSceneBindings)
 		{
-			if(ObjectId == Binding.GetObjectGuid())
+			if (ObjectId == Binding.GetObjectGuid())
 			{
 				TArray<UMovieSceneTrack*> MovieSceneTracks = Binding.GetTracks();
 				for (UMovieSceneTrack* Track : MovieSceneTracks)
@@ -1347,99 +1389,291 @@ void FWidgetBlueprintEditor::ReplaceTrackWithSelectedWidget(FWidgetReference Sel
 						PropertyName.RemoveFromStart("b", ESearchCase::CaseSensitive);
 						FString NameString = "Set" + PropertyName;
 						FName FunctionName = FName(*NameString);
-						if (!SelectedWidget.GetTemplate()->FindFunction(FunctionName))
+						if (!Widget.GetTemplate()->FindFunction(FunctionName))
 						{
-							// Exists a track that's not compatible 
-							FNotificationInfo Info(FText::Format(LOCTEXT("IncompatibleTrackToReplaceWith", "Selected Widget doesn't match to a Property '{0}' this track is bound to"), FText::FromString(PropertyName)));
-							Info.FadeInDuration = 0.1f;
-							Info.FadeOutDuration = 0.5f;
-							Info.ExpireDuration = 2.5f;
-							auto NotificationItem = FSlateNotificationManager::Get().AddNotification(Info);
-
-							NotificationItem->SetCompletionState(SNotificationItem::CS_Success);
-							NotificationItem->ExpireAndFadeout();
-							return;
+							bIncompatibleTracks = true;
+							WidgetsToAdd.Remove(Widget);
+							goto NextWidget;
 						}
 					}
 				}
 			}
 		}
+
+	NextWidget:
+		continue;
 	}
 
-	if (BoundWidget)
+	if (WidgetsToAdd.Num() == 0)
 	{
-		Sequencer->PreAnimatedState.RestorePreAnimatedState(*Sequencer, *BoundWidget);
+		if (bIncompatibleTracks)
+		{
+			// Exists a track that's not compatible 
+			FNotificationInfo Info(LOCTEXT("IncompatibleTracksToReplaceWith", "Selected Widgets do not match a Property this track is bound to"));
+			Info.FadeInDuration = 0.1f;
+			Info.FadeOutDuration = 0.5f;
+			Info.ExpireDuration = 2.5f;
+			auto NotificationItem = FSlateNotificationManager::Get().AddNotification(Info);
+
+			NotificationItem->SetCompletionState(SNotificationItem::CS_Success);
+			NotificationItem->ExpireAndFadeout();
+		}
+		else
+		{
+			FNotificationInfo Info(LOCTEXT("Widgetalreadybound", "Widget already bound"));
+			Info.FadeInDuration = 0.1f;
+			Info.FadeOutDuration = 0.5f;
+			Info.ExpireDuration = 2.5f;
+			auto NotificationItem = FSlateNotificationManager::Get().AddNotification(Info);
+
+			NotificationItem->SetCompletionState(SNotificationItem::CS_Success);
+			NotificationItem->ExpireAndFadeout();
+		}
+
+		return;
 	}
 
-	// else it's safe to modify
 	MovieScene->Modify();
-	MovieScene->SetObjectDisplayName(ObjectId, FText::FromString(PreviewWidget->GetName()));
-
-	// Replace bindings in WidgetAnimation
 	WidgetAnimation->Modify();
+
+	for (const FWidgetReference Widget : WidgetsToAdd)
 	{
-		FWidgetAnimationBinding* SourceBinding = WidgetAnimation->AnimationBindings.FindByPredicate(
-			[&](FWidgetAnimationBinding& In)
-			{
-				return In.AnimationGuid == ObjectId;
-			}
-		);
+		UWidget* PreviewWidget = Widget.GetPreview();
+		WidgetAnimation->BindPossessableObject(ObjectId, *PreviewWidget, GetAnimationPlaybackContext());
+	}
 
-		if (!SourceBinding)
-		{
-			WidgetAnimation->BindPossessableObject(ObjectId, *PreviewWidget, GetAnimationPlaybackContext());
+	UpdateTrackName(ObjectId);
 
-			SourceBinding = WidgetAnimation->AnimationBindings.FindByPredicate(
-				[&](FWidgetAnimationBinding& In)
-			{
-				return In.AnimationGuid == ObjectId;
-			}
-			);
+	Sequencer->NotifyMovieSceneDataChanged(EMovieSceneDataChangeType::MovieSceneStructureItemsChanged);
+}
+
+void FWidgetBlueprintEditor::RemoveWidgetsFromTrack(const TArray<FWidgetReference> Widgets, FGuid ObjectId)
+{
+	const FScopedTransaction Transaction(LOCTEXT("RemoveWidgetsFromTrack", "Remove Widgets from Track"));
+
+	UWidgetAnimation* WidgetAnimation = Cast<UWidgetAnimation>(Sequencer->GetFocusedMovieSceneSequence());
+	UMovieScene* MovieScene = WidgetAnimation->GetMovieScene();
+
+	TArray<FWidgetReference> WidgetsToRemove;
+
+	for (const FWidgetReference& Widget : Widgets)
+	{
+		UWidget* PreviewWidget = Widget.GetPreview();
+		FGuid WidgetId = Sequencer->FindObjectId(*PreviewWidget, MovieSceneSequenceID::Root);
+		if (WidgetId.IsValid() && WidgetId == ObjectId)
+		{
+			WidgetsToRemove.Add(Widget);
 		}
+	}
 
-		check(SourceBinding);
+	if (WidgetsToRemove.Num() == 0)
+	{
+		FNotificationInfo Info(LOCTEXT("SelectedWidgetNotBound", "Selected Widget not Bound to Track"));
+		Info.FadeInDuration = 0.1f;
+		Info.FadeOutDuration = 0.5f;
+		Info.ExpireDuration = 2.5f;
+		auto NotificationItem = FSlateNotificationManager::Get().AddNotification(Info);
+
+		NotificationItem->SetCompletionState(SNotificationItem::CS_Success);
+		NotificationItem->ExpireAndFadeout();
 
-		// Set binding names used for lookup
-		for (FWidgetAnimationBinding& Binding : WidgetAnimation->AnimationBindings)
+		return;
+	}
+
+	MovieScene->Modify();
+
+	for (const FWidgetReference& Widget : WidgetsToRemove)
+	{
+		UWidget* PreviewWidget = Widget.GetPreview();
+		WidgetAnimation->RemoveBinding(*PreviewWidget);
+
+		Sequencer->PreAnimatedState.RestorePreAnimatedState(*Sequencer, *PreviewWidget);
+	}
+
+	UpdateTrackName(ObjectId);
+
+	Sequencer->NotifyMovieSceneDataChanged(EMovieSceneDataChangeType::MovieSceneStructureItemsChanged);
+}
+
+void FWidgetBlueprintEditor::RemoveAllWidgetsFromTrack(FGuid ObjectId)
+{
+	const FScopedTransaction Transaction(LOCTEXT("RemoveAllWidgetsFromTrack", "Remove All Widgets from Track"));
+
+	UWidgetAnimation* WidgetAnimation = Cast<UWidgetAnimation>(Sequencer->GetFocusedMovieSceneSequence());
+	UMovieScene* MovieScene = WidgetAnimation->GetMovieScene();
+
+	TArray<FWidgetAnimationBinding> BindingsToRemove = WidgetAnimation->GetBindings();
+
+	UUserWidget* PreviewRoot = GetPreview();
+	check(PreviewRoot);
+
+	MovieScene->Modify();
+
+	for (const FWidgetAnimationBinding& Binding : BindingsToRemove)
+	{
+		UWidget* BoundWidget = PreviewRoot->GetWidgetFromName(Binding.WidgetName);
+		WidgetAnimation->RemoveBinding(Binding);
+		if (BoundWidget)
 		{
-			if (Binding.AnimationGuid != ObjectId)
-			{
-				continue;
-			}
-		
-			Binding.WidgetName = PreviewWidget->GetFName();
+			Sequencer->PreAnimatedState.RestorePreAnimatedState(*Sequencer, *BoundWidget);
 		}
 	}
 
 	Sequencer->NotifyMovieSceneDataChanged(EMovieSceneDataChangeType::MovieSceneStructureItemsChanged);
 }
 
+void FWidgetBlueprintEditor::RemoveMissingWidgetsFromTrack(FGuid ObjectId)
+{
+	const FScopedTransaction Transaction(LOCTEXT("RemoveMissingWidgetsFromTrack", "Remove Missing Widgets from Track"));
+
+	UWidgetAnimation* WidgetAnimation = Cast<UWidgetAnimation>(Sequencer->GetFocusedMovieSceneSequence());
+	UMovieScene* MovieScene = WidgetAnimation->GetMovieScene();
+
+	const TArray<FWidgetAnimationBinding> Bindings = WidgetAnimation->GetBindings();
+
+	UUserWidget* PreviewRoot = GetPreview();
+	check(PreviewRoot);
 
-void FWidgetBlueprintEditor::ExtendSequencerObjectBindingMenu(FMenuBuilder& ObjectBindingMenuBuilder, const TArray<UObject*> ContextObjects)
+	MovieScene->Modify();
+
+	for (const FWidgetAnimationBinding& Binding : Bindings)
+	{
+		if (Binding.FindRuntimeObject(*PreviewRoot->WidgetTree, *PreviewRoot) == nullptr)
+		{
+			WidgetAnimation->RemoveBinding(Binding);
+		}
+	}
+
+	UpdateTrackName(ObjectId);
+}
+
+void FWidgetBlueprintEditor::ReplaceTrackWithWidgets(const TArray<FWidgetReference> Widgets, FGuid ObjectId)
 {
-	FWidgetReference SelectedWidget;
-	if (SelectedWidgets.Num() == 1)
+	const FScopedTransaction Transaction( LOCTEXT( "ReplaceTrackWithSelectedWidgets", "Replace Track with Selected Widgets" ) );
+
+	UWidgetAnimation* WidgetAnimation = Cast<UWidgetAnimation>(Sequencer->GetFocusedMovieSceneSequence());
+	UMovieScene* MovieScene = WidgetAnimation->GetMovieScene();
+
+	TArray<FWidgetReference> WidgetsToAdd(Widgets);
+	TArray<FWidgetAnimationBinding> BindingsToRemove(WidgetAnimation->GetBindings());
+
+	TArray<FMovieSceneBinding> MovieSceneBindings = MovieScene->GetBindings();
+
+	bool bIncompatibleTracks = false;
+
+	for (const FWidgetReference& Widget : Widgets)
 	{
-		for (FWidgetReference Widget : SelectedWidgets)
+		UWidget* PreviewWidget = Widget.GetPreview();
+		
+		// Check if the widget is already bound
+		FGuid WidgetId = Sequencer->FindObjectId(*PreviewWidget, MovieSceneSequenceID::Root);
+		if (WidgetId.IsValid())
 		{
-			SelectedWidget = Widget;
+			// If bound to the same object, we don't want to remove it
+			if (ObjectId == WidgetId)
+			{
+				FName WidgetName = PreviewWidget->GetFName();
+				FName SlotWidgetName = NAME_None;
+
+				const UPanelSlot* WidgetSlot = Cast<UPanelSlot>(PreviewWidget);
+
+				if ((WidgetSlot != nullptr) && (WidgetSlot->Content != nullptr))
+				{
+					SlotWidgetName = WidgetSlot->GetFName();
+					WidgetName = WidgetSlot->Content->GetFName();
+				}
+
+				BindingsToRemove.RemoveAll([&](const FWidgetAnimationBinding& Binding) {
+					return Binding.WidgetName.IsEqual(WidgetName) && Binding.SlotWidgetName.IsEqual(SlotWidgetName);
+				});
+			}
+			WidgetsToAdd.Remove(Widget);
+			continue;
 		}
+
+		for (FMovieSceneBinding Binding : MovieSceneBindings)
+		{
+			if (ObjectId == Binding.GetObjectGuid())
+			{
+				TArray<UMovieSceneTrack*> MovieSceneTracks = Binding.GetTracks();
+				for (UMovieSceneTrack* Track : MovieSceneTracks)
+				{
+					UMovieScenePropertyTrack* PropertyTrack = Cast<UMovieScenePropertyTrack>(Track);
+					if (PropertyTrack)
+					{
+						FString PropertyName = PropertyTrack->GetPropertyName().ToString();
+						PropertyName.RemoveFromStart("b", ESearchCase::CaseSensitive);
+						FString NameString = "Set" + PropertyName;
+						FName FunctionName = FName(*NameString);
+						if (!Widget.GetTemplate()->FindFunction(FunctionName))
+						{
+							bIncompatibleTracks = true;
+							WidgetsToAdd.Remove(Widget);
+							goto NextWidget;
+						}
+					}
+				}
+			}
+		}
+
+	NextWidget:
+		continue;
 	}
-	if (SelectedWidget.IsValid())
+
+	if (WidgetsToAdd.Num() == 0)
 	{
-		UWidget* BoundWidget = Cast<UWidget>(ContextObjects[0]);
-		if (BoundWidget && SelectedWidget.GetPreview()->GetTypedOuter<UWidgetTree>() == BoundWidget->GetTypedOuter<UWidgetTree>() )
+		if (bIncompatibleTracks)
+		{
+			// Exists a track that's not compatible 
+			FNotificationInfo Info(LOCTEXT("IncompatibleTracksToReplaceWith", "Selected Widgets do not match a Property this track is bound to"));
+			Info.FadeInDuration = 0.1f;
+			Info.FadeOutDuration = 0.5f;
+			Info.ExpireDuration = 2.5f;
+			auto NotificationItem = FSlateNotificationManager::Get().AddNotification(Info);
+
+			NotificationItem->SetCompletionState(SNotificationItem::CS_Success);
+			NotificationItem->ExpireAndFadeout();
+		}
+		else
 		{
-			FGuid ObjectId = Sequencer->FindObjectId(*BoundWidget, MovieSceneSequenceID::Root);
-			FUIAction ReplaceWithMenuAction(FExecuteAction::CreateRaw(this, &FWidgetBlueprintEditor::ReplaceTrackWithSelectedWidget, SelectedWidget, BoundWidget, ObjectId));
+			FNotificationInfo Info(LOCTEXT("Widgetalreadybound", "Widget already bound"));
+			Info.FadeInDuration = 0.1f;
+			Info.FadeOutDuration = 0.5f;
+			Info.ExpireDuration = 2.5f;
+			auto NotificationItem = FSlateNotificationManager::Get().AddNotification(Info);
+
+			NotificationItem->SetCompletionState(SNotificationItem::CS_Success);
+			NotificationItem->ExpireAndFadeout();
+		}
+
+		return;
+	}
+
+	// else it's safe to modify
+	MovieScene->Modify();
 
-			FText ReplaceWithLabel = FText::Format(LOCTEXT("ReplaceObject", "Replace with {0}"), FText::FromString(SelectedWidget.GetPreview()->GetName()));
-			FText ReplaceWithToolTip = FText::Format(LOCTEXT("ReplaceObjectToolTip", "Replace the bound widget in this animation with {0}"), FText::FromString(SelectedWidget.GetPreview()->GetName()));
+	// Replace bindings in WidgetAnimation
+	WidgetAnimation->Modify();
+	{
+		UUserWidget* PreviewRoot = GetPreview();
+		check(PreviewRoot);
 
-			ObjectBindingMenuBuilder.AddMenuEntry(ReplaceWithLabel, ReplaceWithToolTip, FSlateIcon(), ReplaceWithMenuAction);
-			ObjectBindingMenuBuilder.AddMenuSeparator();
+		for (const FWidgetAnimationBinding& Binding : BindingsToRemove)
+		{
+			UWidget* BoundWidget = PreviewRoot->GetWidgetFromName(Binding.WidgetName);
+			WidgetAnimation->RemoveBinding(*BoundWidget);
+			Sequencer->PreAnimatedState.RestorePreAnimatedState(*Sequencer, *BoundWidget);
 		}
+
+		if (WidgetsToAdd.Num() > 0)
+		{
+			AddWidgetsToTrack(WidgetsToAdd, ObjectId);
+		}
+
+		UpdateTrackName(ObjectId);
 	}
+
+	Sequencer->NotifyMovieSceneDataChanged(EMovieSceneDataChangeType::MovieSceneStructureItemsChanged);
 }
 
 void FWidgetBlueprintEditor::AddSlotTrack( UPanelSlot* Slot )
@@ -1574,4 +1808,21 @@ void FWidgetBlueprintEditor::SyncSelectedWidgetsWithSequencerSelection(TArray<FG
 	}
 }
 
+void FWidgetBlueprintEditor::UpdateTrackName(FGuid ObjectId)
+{
+	UWidgetAnimation* WidgetAnimation = Cast<UWidgetAnimation>(Sequencer->GetFocusedMovieSceneSequence());
+	UMovieScene* MovieScene = WidgetAnimation->GetMovieScene();
+
+	const TArray<FWidgetAnimationBinding>& WidgetBindings = WidgetAnimation->GetBindings();
+	if (WidgetBindings.Num() > 0)
+	{
+		FString NewLabel = WidgetBindings[0].WidgetName.ToString();
+		if (WidgetBindings.Num() > 1)
+		{
+			NewLabel.Append(FString::Printf(TEXT(" (%d)"), WidgetBindings.Num()));
+		}
+
+		MovieScene->SetObjectDisplayName(ObjectId, FText::FromString(NewLabel));
+	}
+}
 #undef LOCTEXT_NAMESPACE
diff --git a/Engine/Source/Editor/UMGEditor/Public/WidgetBlueprintEditor.h b/Engine/Source/Editor/UMGEditor/Public/WidgetBlueprintEditor.h
--- a/Engine/Source/Editor/UMGEditor/Public/WidgetBlueprintEditor.h
+++ b/Engine/Source/Editor/UMGEditor/Public/WidgetBlueprintEditor.h
@@ -261,11 +261,20 @@ private:
 	/** Extends the sequencer add track menu. */
 	void ExtendSequencerAddTrackMenu( FMenuBuilder& AddTrackMenuBuilder, const TArray<UObject*> ContextObjects );
 
-	/** Replace track with selected widget function */
-	void ReplaceTrackWithSelectedWidget(FWidgetReference SelectedWidget, UWidget* BoundWidget, FGuid ObjectBinding);
+	/** Binds additional widgets to a track of the same type */
+	void AddWidgetsToTrack(const TArray<FWidgetReference> Widgets, FGuid ObjectId);
 
-	/** Extends the sequencer add track menu. */
-	void ExtendSequencerObjectBindingMenu(FMenuBuilder& ObjectBindingMenuBuilder, const TArray<UObject*> ContextObjects);
+	/** Unbind widgets from a track*/
+	void RemoveWidgetsFromTrack(const TArray<FWidgetReference> Widgets, FGuid ObjectId);
+
+	/** Remove all bindings from a track */
+	void RemoveAllWidgetsFromTrack(FGuid ObjectId);
+
+	/** Remove any missing bindings from a track */
+	void RemoveMissingWidgetsFromTrack(FGuid ObjectId);
+	
+	/** Replace current widget bindings on a track with new widget bindings */
+	void ReplaceTrackWithWidgets(const TArray<FWidgetReference> Widgets, FGuid ObjectId);
 
 	/** Add an animation track for the supplied slot to the current animation. */
 	void AddSlotTrack( UPanelSlot* Slot );
@@ -288,6 +297,9 @@ private:
 	/** Get the animation playback event contexts */
 	TArray<UObject*> GetAnimationEventContexts() const { TArray<UObject*> EventContexts; EventContexts.Add(GetPreview()); return EventContexts; }
 	
+	/** Update the name of a track to reflect changes in bindings */
+	void UpdateTrackName(FGuid ObjectId);
+
 private:
 	/** The preview scene that owns the preview GUI */
 	FPreviewScene PreviewScene;
diff --git a/Engine/Source/Runtime/UMG/Private/Animation/WidgetAnimation.cpp b/Engine/Source/Runtime/UMG/Private/Animation/WidgetAnimation.cpp
--- a/Engine/Source/Runtime/UMG/Private/Animation/WidgetAnimation.cpp
+++ b/Engine/Source/Runtime/UMG/Private/Animation/WidgetAnimation.cpp
@@ -215,19 +215,20 @@ void UWidgetAnimation::LocateBoundObjects(const FGuid& ObjectId, UObject* InCont
 		return;
 	}
 
-	const FWidgetAnimationBinding* Binding = AnimationBindings.FindByPredicate(
-		[ObjectId](const FWidgetAnimationBinding& In){
-			return In.AnimationGuid == ObjectId;
-		}
-	);
-
 	UUserWidget* PreviewWidget = CastChecked<UUserWidget>(InContext);
-	UObject* FoundObject = Binding ? Binding->FindRuntimeObject(*PreviewWidget->WidgetTree, *PreviewWidget) : nullptr;
-
-	if (FoundObject)
+	for (const FWidgetAnimationBinding& Binding : AnimationBindings)
 	{
-		OutObjects.Add(FoundObject);
+		if (Binding.AnimationGuid == ObjectId)
+		{
+			UObject* FoundObject = Binding.FindRuntimeObject(*PreviewWidget->WidgetTree, *PreviewWidget);
+
+			if (FoundObject)
+			{
+				OutObjects.Add(FoundObject);
+			}
+		}
 	}
+
 }
 
 
@@ -269,6 +270,33 @@ void UWidgetAnimation::UnbindPossessableObjects(const FGuid& ObjectId)
 	});
 }
 
+void UWidgetAnimation::RemoveBinding(const UObject& PossessedObject)
+{
+	Modify();
+
+	FName WidgetName = PossessedObject.GetFName();
+	FName SlotWidgetName = NAME_None;
+
+	const UPanelSlot* PossessedSlot = Cast<UPanelSlot>(&PossessedObject);
+
+	if ((PossessedSlot != nullptr) && (PossessedSlot->Content != nullptr))
+	{
+		SlotWidgetName = PossessedSlot->GetFName();
+		WidgetName = PossessedSlot->Content->GetFName();
+	}
+
+	AnimationBindings.RemoveAll([&](const FWidgetAnimationBinding& Binding) {
+		return Binding.WidgetName.IsEqual(WidgetName) && Binding.SlotWidgetName.IsEqual(SlotWidgetName);
+	});
+}
+
+void UWidgetAnimation::RemoveBinding(const FWidgetAnimationBinding& Binding)
+{
+	Modify();
+
+	AnimationBindings.Remove(Binding);
+}
+
 bool UWidgetAnimation::IsPostLoadThreadSafe() const
 {
 	return true;
diff --git a/Engine/Source/Runtime/UMG/Public/Animation/WidgetAnimation.h b/Engine/Source/Runtime/UMG/Public/Animation/WidgetAnimation.h
--- a/Engine/Source/Runtime/UMG/Public/Animation/WidgetAnimation.h
+++ b/Engine/Source/Runtime/UMG/Public/Animation/WidgetAnimation.h
@@ -119,6 +119,11 @@ public:
 	/** Get Animation bindings of the animation */
 	const TArray<FWidgetAnimationBinding>& GetBindings() const { return AnimationBindings; }
 
+	/** Remove Animation Binding */
+	UMG_API void RemoveBinding(const UObject& PossessedObject);
+	UMG_API void RemoveBinding(const FWidgetAnimationBinding& Binding);
+	
+
 	/** Whether to finish evaluation on stop */
 	bool GetLegacyFinishOnStop() const { return bLegacyFinishOnStop; }
 
