Adding UMG designer overlay element that shows the coordinates of the cursor in widget preview space
- Basically just showing numbers for what the ruler lines already show

diff --git a/Engine/Source/Editor/UMGEditor/Private/Designer/SDesignerView.cpp b/Engine/Source/Editor/UMGEditor/Private/Designer/SDesignerView.cpp
--- a/Engine/Source/Editor/UMGEditor/Private/Designer/SDesignerView.cpp
+++ b/Engine/Source/Editor/UMGEditor/Private/Designer/SDesignerView.cpp
@@ -591,6 +591,19 @@ TSharedRef<SWidget> SDesignerView::CreateOverlayUI()
 			.Visibility(EVisibility::SelfHitTestInvisible)
 		]
 
+		+ SHorizontalBox::Slot()
+		.AutoWidth()
+		.VAlign(VAlign_Center)
+		.Padding(40, 2, 0, 0)
+		[
+			SNew(STextBlock)
+			.TextStyle(FEditorStyle::Get(), "Graph.ZoomText")
+			.Font(FCoreStyle::GetDefaultFontStyle(TEXT("BoldCondensed"), 14))
+			.Text(this, &SDesignerView::GetCursorPositionText)
+			.ColorAndOpacity(FLinearColor(1.f, 1.f, 1.f, 0.25f))
+			.Visibility(this, &SDesignerView::GetCursorPositionTextVisibility)
+		]
+
 		+ SHorizontalBox::Slot()
 		.FillWidth(1.0f)
 		[
@@ -3247,6 +3260,23 @@ FText SDesignerView::GetDesignerOutlineText() const
 	return FText::GetEmpty();
 }
 
+FText SDesignerView::GetCursorPositionText() const
+{
+	if (const FArrangedWidget* CachedPreviewSurface = CachedWidgetGeometry.Find(PreviewSurface.ToSharedRef()))
+	{
+		const FGeometry& RootGeometry = CachedPreviewSurface->Geometry;
+		const FVector2D CursorPos = RootGeometry.AbsoluteToLocal(FSlateApplication::Get().GetCursorPos()) / GetPreviewDPIScale();
+
+		return FText::Format(LOCTEXT("CursorPositionFormat", "{0} x {1}"), FText::AsNumber(FMath::RoundToInt(CursorPos.X)), FText::AsNumber(FMath::RoundToInt(CursorPos.Y)));
+	}
+	return FText();
+}
+
+EVisibility SDesignerView::GetCursorPositionTextVisibility() const
+{
+	return IsHovered() ? EVisibility::SelfHitTestInvisible : EVisibility::Collapsed;
+}
+
 FReply SDesignerView::HandleDPISettingsClicked()
 {
 	FModuleManager::LoadModuleChecked<ISettingsModule>("Settings").ShowViewer("Project", "Engine", "UI");
diff --git a/Engine/Source/Editor/UMGEditor/Private/Designer/SDesignerView.h b/Engine/Source/Editor/UMGEditor/Private/Designer/SDesignerView.h
--- a/Engine/Source/Editor/UMGEditor/Private/Designer/SDesignerView.h
+++ b/Engine/Source/Editor/UMGEditor/Private/Designer/SDesignerView.h
@@ -213,6 +213,9 @@ private:
 	FSlateColor GetDesignerOutlineColor() const;
 	FText GetDesignerOutlineText() const;
 
+	FText GetCursorPositionText() const;
+	EVisibility GetCursorPositionTextVisibility() const;
+
 	// Handles drawing selection and other effects a SPaintSurface widget injected into the hierarchy.
 	int32 HandleEffectsPainting(const FOnPaintHandlerParams& PaintArgs);
 	void DrawSelectionAndHoverOutline(const FOnPaintHandlerParams& PaintArgs);
