UMGEditor - Cleaning up some code related to how handle named slots in the designer, using the more versatile TScriptInterface for exchanging the pointer, and updating the replace with child code to work  if you're targeting the root content in a named slot.

diff --git a/Engine/Source/Editor/UMGEditor/Private/Hierarchy/SHierarchyViewItem.cpp b/Engine/Source/Editor/UMGEditor/Private/Hierarchy/SHierarchyViewItem.cpp
--- a/Engine/Source/Editor/UMGEditor/Private/Hierarchy/SHierarchyViewItem.cpp
+++ b/Engine/Source/Editor/UMGEditor/Private/Hierarchy/SHierarchyViewItem.cpp
@@ -331,8 +331,7 @@ TOptional<EItemDropZone> ProcessHierarchyDragDrop(const FDragDropEvent& DragDrop
 					UWidget* NamedSlotHostWidget = FWidgetBlueprintEditorUtils::FindNamedSlotHostWidgetForContent(TemplateWidget, Blueprint->WidgetTree);
 					if (NamedSlotHostWidget != nullptr)
 					{
-						INamedSlotInterface* NamedSlotHost = Cast<INamedSlotInterface>(NamedSlotHostWidget);
-						if (NamedSlotHost != nullptr)
+						if (TScriptInterface<INamedSlotInterface> NamedSlotHost = TScriptInterface<INamedSlotInterface>(NamedSlotHostWidget))
 						{
 							NamedSlotHostWidget->SetFlags(RF_Transactional);
 							NamedSlotHostWidget->Modify();
@@ -992,8 +991,7 @@ FReply FNamedSlotModel::HandleAcceptDrop(FDragDropEvent const& DragDropEvent, EI
 		UWidget* SourceNamedSlotHostWidget = FWidgetBlueprintEditorUtils::FindNamedSlotHostWidgetForContent(DroppingWidget, Blueprint->WidgetTree);
 		if (SourceNamedSlotHostWidget != nullptr)
 		{
-			INamedSlotInterface* SourceNamedSlotHost = Cast<INamedSlotInterface>(SourceNamedSlotHostWidget);
-			if (SourceNamedSlotHost != nullptr)
+			if (TScriptInterface<INamedSlotInterface> SourceNamedSlotHost = TScriptInterface<INamedSlotInterface>(SourceNamedSlotHostWidget))
 			{
 				SourceNamedSlotHostWidget->SetFlags(RF_Transactional);
 				SourceNamedSlotHostWidget->Modify();
@@ -1018,7 +1016,7 @@ void FNamedSlotModel::DoDrop(UWidget* NamedSlotHostWidget, UWidget* DroppingWidg
 	NamedSlotHostWidget->SetFlags(RF_Transactional);
 	NamedSlotHostWidget->Modify();
 
-	INamedSlotInterface* NamedSlotInterface = Cast<INamedSlotInterface>(NamedSlotHostWidget);
+	TScriptInterface<INamedSlotInterface> NamedSlotInterface = TScriptInterface<INamedSlotInterface>(NamedSlotHostWidget);
 	NamedSlotInterface->SetContentForSlot(SlotName, DroppingWidget);
 
 	FBlueprintEditorUtils::MarkBlueprintAsStructurallyModified(Blueprint);
diff --git a/Engine/Source/Editor/UMGEditor/Private/WidgetBlueprintEditorUtils.cpp b/Engine/Source/Editor/UMGEditor/Private/WidgetBlueprintEditorUtils.cpp
--- a/Engine/Source/Editor/UMGEditor/Private/WidgetBlueprintEditorUtils.cpp
+++ b/Engine/Source/Editor/UMGEditor/Private/WidgetBlueprintEditorUtils.cpp
@@ -34,6 +34,7 @@
 #include "Widgets/Notifications/SNotificationList.h"
 #include "Components/Widget.h"
 #include "Blueprint/WidgetNavigation.h"
+#include "UObject/ScriptInterface.h"
 
 #define LOCTEXT_NAMESPACE "UMG"
 
@@ -429,9 +430,9 @@ void FWidgetBlueprintEditorUtils::DeleteWidgets(UWidgetBlueprint* Blueprint, TSe
 	}
 }
 
-INamedSlotInterface* FWidgetBlueprintEditorUtils::FindNamedSlotHostForContent(UWidget* WidgetTemplate, UWidgetTree* WidgetTree)
+TScriptInterface<INamedSlotInterface> FWidgetBlueprintEditorUtils::FindNamedSlotHostForContent(UWidget* WidgetTemplate, UWidgetTree* WidgetTree)
 {
-	return Cast<INamedSlotInterface>(FindNamedSlotHostWidgetForContent(WidgetTemplate, WidgetTree));
+	return TScriptInterface<INamedSlotInterface>(FindNamedSlotHostWidgetForContent(WidgetTemplate, WidgetTree));
 }
 
 UWidget* FWidgetBlueprintEditorUtils::FindNamedSlotHostWidgetForContent(UWidget* WidgetTemplate, UWidgetTree* WidgetTree)
@@ -514,12 +515,12 @@ void FWidgetBlueprintEditorUtils::FindAllAncestorNamedSlotHostWidgetsForContent(
 	}
 }
 
-bool FWidgetBlueprintEditorUtils::RemoveNamedSlotHostContent(UWidget* WidgetTemplate, INamedSlotInterface* NamedSlotHost)
+bool FWidgetBlueprintEditorUtils::RemoveNamedSlotHostContent(UWidget* WidgetTemplate, TScriptInterface<INamedSlotInterface> NamedSlotHost)
 {
 	return ReplaceNamedSlotHostContent(WidgetTemplate, NamedSlotHost, nullptr);
 }
 
-bool FWidgetBlueprintEditorUtils::ReplaceNamedSlotHostContent(UWidget* WidgetTemplate, INamedSlotInterface* NamedSlotHost, UWidget* NewContentWidget)
+bool FWidgetBlueprintEditorUtils::ReplaceNamedSlotHostContent(UWidget* WidgetTemplate, TScriptInterface<INamedSlotInterface> NamedSlotHost, UWidget* NewContentWidget)
 {
 	TArray<FName> SlotNames;
 	NamedSlotHost->GetSlotNames(SlotNames);
@@ -530,6 +531,8 @@ bool FWidgetBlueprintEditorUtils::ReplaceNamedSlotHostContent(UWidget* WidgetTem
 		{
 			if (SlotContent == WidgetTemplate)
 			{
+				NewContentWidget->Modify();
+				NamedSlotHost.GetObject()->Modify();
 				NamedSlotHost->SetContentForSlot(SlotName, NewContentWidget);
 				return true;
 			}
@@ -542,7 +545,7 @@ bool FWidgetBlueprintEditorUtils::ReplaceNamedSlotHostContent(UWidget* WidgetTem
 bool FWidgetBlueprintEditorUtils::FindAndRemoveNamedSlotContent(UWidget* WidgetTemplate, UWidgetTree* WidgetTree)
 {
 	UWidget* NamedSlotHostWidget = FindNamedSlotHostWidgetForContent(WidgetTemplate, WidgetTree);
-	if ( INamedSlotInterface* NamedSlotHost = Cast<INamedSlotInterface>(NamedSlotHostWidget) )
+	if (TScriptInterface<INamedSlotInterface> NamedSlotHost = TScriptInterface<INamedSlotInterface>(NamedSlotHostWidget) )
 	{
 		NamedSlotHostWidget->Modify();
 		return RemoveNamedSlotHostContent(WidgetTemplate, NamedSlotHost);
@@ -614,8 +617,7 @@ void FWidgetBlueprintEditorUtils::WrapWidgets(TSharedRef<FWidgetBlueprintEditor>
 		if (CurrentSlot)
 		{
 			// If this is a named slot, we need to properly remove and reassign the slot content
-			INamedSlotInterface* NamedSlotHost = Cast<INamedSlotInterface>(CurrentSlot);
-			if (NamedSlotHost != nullptr)
+			if (TScriptInterface<INamedSlotInterface> NamedSlotHost = TScriptInterface<INamedSlotInterface>(CurrentSlot))
 			{
 				CurrentSlot->SetFlags(RF_Transactional);
 				CurrentSlot->Modify();
@@ -892,6 +894,10 @@ void FWidgetBlueprintEditorUtils::ReplaceWidgetWithChildren(TSharedRef<FWidgetBl
 			BP->WidgetTree->Modify();
 			BP->WidgetTree->RootWidget = FirstChildTemplate;
 		}
+		else if (TScriptInterface<INamedSlotInterface> NamedSlotHost = FindNamedSlotHostForContent(ExistingPanelTemplate, BP->WidgetTree))
+		{
+			ReplaceNamedSlotHostContent(ExistingPanelTemplate, NamedSlotHost, FirstChildTemplate);
+		}
 		else
 		{
 			Transaction.Cancel();
diff --git a/Engine/Source/Editor/UMGEditor/Private/WidgetBlueprintEditorUtils.h b/Engine/Source/Editor/UMGEditor/Private/WidgetBlueprintEditorUtils.h
--- a/Engine/Source/Editor/UMGEditor/Private/WidgetBlueprintEditorUtils.h
+++ b/Engine/Source/Editor/UMGEditor/Private/WidgetBlueprintEditorUtils.h
@@ -51,15 +51,15 @@ public:
 	/** Attempts to import any property in the map and apply it to a property with the same name on the object. */
 	static void ImportPropertiesFromText(UObject* Object, const TMap<FName, FString>& ExportedProperties);
 
-	static INamedSlotInterface* FindNamedSlotHostForContent(UWidget* WidgetTemplate, UWidgetTree* WidgetTree);
+	static TScriptInterface<INamedSlotInterface> FindNamedSlotHostForContent(UWidget* WidgetTemplate, UWidgetTree* WidgetTree);
 
 	static UWidget* FindNamedSlotHostWidgetForContent(UWidget* WidgetTemplate, UWidgetTree* WidgetTree);
 
 	static void FindAllAncestorNamedSlotHostWidgetsForContent(TArray<FWidgetReference>& OutSlotHostWidgets, UWidget* WidgetTemplate, TSharedRef<FWidgetBlueprintEditor> BlueprintEditor);
 
-	static bool RemoveNamedSlotHostContent(UWidget* WidgetTemplate, INamedSlotInterface* NamedSlotHost);
+	static bool RemoveNamedSlotHostContent(UWidget* WidgetTemplate, TScriptInterface<INamedSlotInterface> NamedSlotHost);
 
-	static bool ReplaceNamedSlotHostContent(UWidget* WidgetTemplate, INamedSlotInterface* NamedSlotHost, UWidget* NewContentWidget);
+	static bool ReplaceNamedSlotHostContent(UWidget* WidgetTemplate, TScriptInterface<INamedSlotInterface> NamedSlotHost, UWidget* NewContentWidget);
 
 private:
 
