Fix hit-test that kept a hard reference on the widget under cursor which would result in a false widget leak detection when the designer preview widget is destroyed.

diff --git a/Engine/Source/Editor/UMGEditor/Private/Designer/SDesignerView.cpp b/Engine/Source/Editor/UMGEditor/Private/Designer/SDesignerView.cpp
--- a/Engine/Source/Editor/UMGEditor/Private/Designer/SDesignerView.cpp
+++ b/Engine/Source/Editor/UMGEditor/Private/Designer/SDesignerView.cpp
@@ -1738,13 +1738,21 @@ FReply SDesignerView::OnMouseButtonDown(const FGeometry& MyGeometry, const FPoin
 	SDesignSurface::OnMouseButtonDown(MyGeometry, MouseEvent);
 
 	//TODO UMG Undoable Selection
-	FWidgetHitResult HitResult;
-	if ( FindWidgetUnderCursor(MyGeometry, MouseEvent, UWidget::StaticClass(), HitResult) )
-	{
-		SelectedWidgetContextMenuLocation = HitResult.WidgetArranged.Geometry.AbsoluteToLocal(MouseEvent.GetScreenSpacePosition());
 
-		PendingSelectedWidget = HitResult.Widget;
+	bool bFoundWidgetUnderCursor = false;
+	{
+		// Narrow life scope of FWidgetHitResult so it doesn't keep a hard reference on any widget.
+		FWidgetHitResult HitResult;
+		bFoundWidgetUnderCursor = FindWidgetUnderCursor(MyGeometry, MouseEvent, UWidget::StaticClass(), HitResult);
+		if (bFoundWidgetUnderCursor)
+		{
+			SelectedWidgetContextMenuLocation = HitResult.WidgetArranged.Geometry.AbsoluteToLocal(MouseEvent.GetScreenSpacePosition());
+			PendingSelectedWidget = HitResult.Widget;
+		}
+	}
 
+	if (bFoundWidgetUnderCursor)
+	{
 		if ( MouseEvent.GetEffectingButton() == EKeys::LeftMouseButton )
 		{
 			const TSet<FWidgetReference>& SelectedWidgets = GetSelectedWidgets();
