Fix crash replacing a missing bound object in UMG timeline.

diff --git a/Engine/Source/Editor/UMGEditor/Private/WidgetBlueprintEditor.cpp b/Engine/Source/Editor/UMGEditor/Private/WidgetBlueprintEditor.cpp
--- a/Engine/Source/Editor/UMGEditor/Private/WidgetBlueprintEditor.cpp
+++ b/Engine/Source/Editor/UMGEditor/Private/WidgetBlueprintEditor.cpp
@@ -1702,8 +1702,11 @@ void FWidgetBlueprintEditor::ReplaceTrackWithWidgets(const TArray<FWidgetReferen
 		for (const FWidgetAnimationBinding& Binding : BindingsToRemove)
 		{
 			UWidget* BoundWidget = PreviewRoot->GetWidgetFromName(Binding.WidgetName);
-			WidgetAnimation->RemoveBinding(*BoundWidget);
-			Sequencer->PreAnimatedState.RestorePreAnimatedState(*Sequencer, *BoundWidget);
+			WidgetAnimation->RemoveBinding(Binding);
+			if (BoundWidget)
+			{
+				Sequencer->PreAnimatedState.RestorePreAnimatedState(*Sequencer, *BoundWidget);
+			}
 		}
 
 		if (WidgetsToAdd.Num() > 0)
