UMG: Fixed replacing bindings in UMG nuking bindings on objects that do not relate to the desired action

  - Also removed restriction that prevents users from rebinding a widget that doesn't contain one of the animated properties. Such properties will be highlighted red and not animate correctly, there's no reason to prevent the entire conversion in this case (plus the logic was incorrect as it only checked for a setter function name)

  - Corrected some logic that was not correctly considering the input object binding ID that resulted in all bindings being removed

 - Reworked replace logic to be implemented in terms of Remove + Add

diff --git a/Engine/Source/Editor/UMGEditor/Private/WidgetBlueprintEditor.cpp b/Engine/Source/Editor/UMGEditor/Private/WidgetBlueprintEditor.cpp
--- a/Engine/Source/Editor/UMGEditor/Private/WidgetBlueprintEditor.cpp
+++ b/Engine/Source/Editor/UMGEditor/Private/WidgetBlueprintEditor.cpp
@@ -16,6 +16,9 @@
 #if WITH_EDITOR
 	#include "EditorStyleSet.h"
 #endif // WITH_EDITOR
+
+#include "Algo/AllOf.h"
+
 #include "Components/PanelSlot.h"
 #include "Components/PanelWidget.h"
 #include "Settings/WidgetDesignerSettings.h"
@@ -1399,93 +1402,45 @@ void FWidgetBlueprintEditor::AddWidgetsToTrack(const TArray<FWidgetReference> Wi
 	UWidgetAnimation* WidgetAnimation = Cast<UWidgetAnimation>(Sequencer->GetFocusedMovieSceneSequence());
 	UMovieScene* MovieScene = WidgetAnimation->GetMovieScene();
 
-	TArray<FWidgetReference> WidgetsToAdd(Widgets);
-
-	bool bIncompatibleTracks = false;
-
-	for (const FWidgetReference Widget : Widgets)
+	TArray<FWidgetReference> WidgetsToAdd;
+	for (const FWidgetReference& Widget : Widgets)
 	{
 		UWidget* PreviewWidget = Widget.GetPreview();
 
-		// Try find if the SelectedWidget is already bound, if so return
+		// If this widget is already bound to the animation we cannot add it to 2 separate bindings
 		FGuid SelectedWidgetId = Sequencer->FindObjectId(*PreviewWidget, MovieSceneSequenceID::Root);
-		if (SelectedWidgetId.IsValid())
+		if (!SelectedWidgetId.IsValid())
 		{
-			WidgetsToAdd.Remove(Widget);
-			continue;
+			WidgetsToAdd.Add(Widget);
 		}
-
-		TArray<FMovieSceneBinding> MovieSceneBindings = MovieScene->GetBindings();
-		for (FMovieSceneBinding Binding : MovieSceneBindings)
-		{
-			if (ObjectId == Binding.GetObjectGuid())
-			{
-				TArray<UMovieSceneTrack*> MovieSceneTracks = Binding.GetTracks();
-				for (UMovieSceneTrack* Track : MovieSceneTracks)
-				{
-					UMovieScenePropertyTrack* PropertyTrack = Cast<UMovieScenePropertyTrack>(Track);
-					if (PropertyTrack)
-					{
-						FString PropertyName = PropertyTrack->GetPropertyName().ToString();
-						PropertyName.RemoveFromStart("b", ESearchCase::CaseSensitive);
-						FString NameString = "Set" + PropertyName;
-						FName FunctionName = FName(*NameString);
-						if (!Widget.GetTemplate()->FindFunction(FunctionName))
-						{
-							bIncompatibleTracks = true;
-							WidgetsToAdd.Remove(Widget);
-							goto NextWidget;
-						}
-					}
-				}
-			}
-		}
-
-	NextWidget:
-		continue;
 	}
 
 	if (WidgetsToAdd.Num() == 0)
 	{
-		if (bIncompatibleTracks)
-		{
-			// Exists a track that's not compatible 
-			FNotificationInfo Info(LOCTEXT("IncompatibleTracksToReplaceWith", "Selected Widgets do not match a Property this track is bound to"));
-			Info.FadeInDuration = 0.1f;
-			Info.FadeOutDuration = 0.5f;
-			Info.ExpireDuration = 2.5f;
-			auto NotificationItem = FSlateNotificationManager::Get().AddNotification(Info);
-
-			NotificationItem->SetCompletionState(SNotificationItem::CS_Success);
-			NotificationItem->ExpireAndFadeout();
-		}
-		else
-		{
-			FNotificationInfo Info(LOCTEXT("Widgetalreadybound", "Widget already bound"));
-			Info.FadeInDuration = 0.1f;
-			Info.FadeOutDuration = 0.5f;
-			Info.ExpireDuration = 2.5f;
-			auto NotificationItem = FSlateNotificationManager::Get().AddNotification(Info);
-
-			NotificationItem->SetCompletionState(SNotificationItem::CS_Success);
-			NotificationItem->ExpireAndFadeout();
-		}
+		FNotificationInfo Info(LOCTEXT("Widgetalreadybound", "Widget already bound"));
+		Info.FadeInDuration = 0.1f;
+		Info.FadeOutDuration = 0.5f;
+		Info.ExpireDuration = 2.5f;
+		auto NotificationItem = FSlateNotificationManager::Get().AddNotification(Info);
 
-		return;
+		NotificationItem->SetCompletionState(SNotificationItem::CS_Success);
+		NotificationItem->ExpireAndFadeout();
 	}
-
-	MovieScene->Modify();
-	WidgetAnimation->Modify();
-
-	for (const FWidgetReference Widget : WidgetsToAdd)
+	else
 	{
-		UWidget* PreviewWidget = Widget.GetPreview();
-		WidgetAnimation->BindPossessableObject(ObjectId, *PreviewWidget, GetAnimationPlaybackContext());
-	}
+		MovieScene->Modify();
+		WidgetAnimation->Modify();
 
-	UpdateTrackName(ObjectId);
+		for (const FWidgetReference Widget : WidgetsToAdd)
+		{
+			UWidget* PreviewWidget = Widget.GetPreview();
+			WidgetAnimation->BindPossessableObject(ObjectId, *PreviewWidget, GetAnimationPlaybackContext());
+		}
 
-	Sequencer->NotifyMovieSceneDataChanged(EMovieSceneDataChangeType::MovieSceneStructureItemsChanged);
+		UpdateTrackName(ObjectId);
+
+		Sequencer->NotifyMovieSceneDataChanged(EMovieSceneDataChangeType::MovieSceneStructureItemsChanged);
+	}
 }
 
 void FWidgetBlueprintEditor::RemoveWidgetsFromTrack(const TArray<FWidgetReference> Widgets, FGuid ObjectId)
@@ -1517,23 +1472,24 @@ void FWidgetBlueprintEditor::RemoveWidgetsFromTrack(const TArray<FWidgetReferenc
 
 		NotificationItem->SetCompletionState(SNotificationItem::CS_Success);
 		NotificationItem->ExpireAndFadeout();
-
-		return;
 	}
-
-	MovieScene->Modify();
-
-	for (const FWidgetReference& Widget : WidgetsToRemove)
+	else
 	{
-		UWidget* PreviewWidget = Widget.GetPreview();
-		WidgetAnimation->RemoveBinding(*PreviewWidget);
+		MovieScene->Modify();
+		WidgetAnimation->Modify();
 
-		Sequencer->PreAnimatedState.RestorePreAnimatedState(*Sequencer, *PreviewWidget);
-	}
+		for (const FWidgetReference& Widget : WidgetsToRemove)
+		{
+			UWidget* PreviewWidget = Widget.GetPreview();
+			WidgetAnimation->RemoveBinding(*PreviewWidget);
 
-	UpdateTrackName(ObjectId);
+			Sequencer->RestorePreAnimatedState(*PreviewWidget);
+		}
 
-	Sequencer->NotifyMovieSceneDataChanged(EMovieSceneDataChangeType::MovieSceneStructureItemsChanged);
+		UpdateTrackName(ObjectId);
+
+		Sequencer->NotifyMovieSceneDataChanged(EMovieSceneDataChangeType::MovieSceneStructureItemsChanged);
+	}
 }
 
 void FWidgetBlueprintEditor::RemoveAllWidgetsFromTrack(FGuid ObjectId)
@@ -1543,20 +1499,27 @@ void FWidgetBlueprintEditor::RemoveAllWidgetsFromTrack(FGuid ObjectId)
 	UWidgetAnimation* WidgetAnimation = Cast<UWidgetAnimation>(Sequencer->GetFocusedMovieSceneSequence());
 	UMovieScene* MovieScene = WidgetAnimation->GetMovieScene();
 
-	TArray<FWidgetAnimationBinding> BindingsToRemove = WidgetAnimation->GetBindings();
-
 	UUserWidget* PreviewRoot = GetPreview();
 	check(PreviewRoot);
 
+	WidgetAnimation->Modify();
 	MovieScene->Modify();
 
-	for (const FWidgetAnimationBinding& Binding : BindingsToRemove)
+	// Restore object animation state
+	for (TWeakObjectPtr<> WeakObject : Sequencer->FindBoundObjects(ObjectId, MovieSceneSequenceID::Root))
 	{
-		UWidget* BoundWidget = PreviewRoot->GetWidgetFromName(Binding.WidgetName);
-		WidgetAnimation->RemoveBinding(Binding);
-		if (BoundWidget)
+		if (UObject* Obj = WeakObject.Get())
 		{
-			Sequencer->PreAnimatedState.RestorePreAnimatedState(*Sequencer, *BoundWidget);
+			Sequencer->RestorePreAnimatedState(*Obj);
+		}
+	}
+
+	// Remove bindings
+	for (int32 Index = WidgetAnimation->AnimationBindings.Num() - 1; Index >= 0; --Index)
+	{
+		if (WidgetAnimation->AnimationBindings[Index].AnimationGuid == ObjectId)
+		{
+			WidgetAnimation->AnimationBindings.RemoveAt(Index, 1, false);
 		}
 	}
 
@@ -1570,153 +1533,65 @@ void FWidgetBlueprintEditor::RemoveMissingWidgetsFromTrack(FGuid ObjectId)
 	UWidgetAnimation* WidgetAnimation = Cast<UWidgetAnimation>(Sequencer->GetFocusedMovieSceneSequence());
 	UMovieScene* MovieScene = WidgetAnimation->GetMovieScene();
 
-	const TArray<FWidgetAnimationBinding> Bindings = WidgetAnimation->GetBindings();
-
 	UUserWidget* PreviewRoot = GetPreview();
 	check(PreviewRoot);
 
+	WidgetAnimation->Modify();
 	MovieScene->Modify();
 
-	for (const FWidgetAnimationBinding& Binding : Bindings)
+	for (int32 Index = WidgetAnimation->AnimationBindings.Num() - 1; Index >= 0; --Index)
 	{
-		if (Binding.FindRuntimeObject(*PreviewRoot->WidgetTree, *PreviewRoot) == nullptr)
+		const FWidgetAnimationBinding& Binding = WidgetAnimation->AnimationBindings[Index];
+		if (Binding.AnimationGuid == ObjectId && Binding.FindRuntimeObject(*PreviewRoot->WidgetTree, *PreviewRoot) == nullptr)
 		{
-			WidgetAnimation->RemoveBinding(Binding);
+			WidgetAnimation->AnimationBindings.RemoveAt(Index, 1, false);
 		}
 	}
 
 	UpdateTrackName(ObjectId);
 }
 
-void FWidgetBlueprintEditor::ReplaceTrackWithWidgets(const TArray<FWidgetReference> Widgets, FGuid ObjectId)
+void FWidgetBlueprintEditor::ReplaceTrackWithWidgets(TArray<FWidgetReference> Widgets, FGuid ObjectId)
 {
 	const FScopedTransaction Transaction( LOCTEXT( "ReplaceTrackWithSelectedWidgets", "Replace Track with Selected Widgets" ) );
 
 	UWidgetAnimation* WidgetAnimation = Cast<UWidgetAnimation>(Sequencer->GetFocusedMovieSceneSequence());
 	UMovieScene* MovieScene = WidgetAnimation->GetMovieScene();
 
-	TArray<FWidgetReference> WidgetsToAdd(Widgets);
-	TArray<FWidgetAnimationBinding> BindingsToRemove(WidgetAnimation->GetBindings());
-
-	TArray<FMovieSceneBinding> MovieSceneBindings = MovieScene->GetBindings();
+	WidgetAnimation->Modify();
+	MovieScene->Modify();
 
-	bool bIncompatibleTracks = false;
+	// Remove everything from the track
+	RemoveAllWidgetsFromTrack(ObjectId);
 
-	for (const FWidgetReference& Widget : Widgets)
+	// Filter out anything in the input array that is currently bound to another object in the animation
+	for (int32 Index = Widgets.Num()-1; Index >= 0; --Index)
 	{
-		UWidget* PreviewWidget = Widget.GetPreview();
-		
-		// Check if the widget is already bound
+		UWidget* PreviewWidget = Widgets[Index].GetPreview();
 		FGuid WidgetId = Sequencer->FindObjectId(*PreviewWidget, MovieSceneSequenceID::Root);
 		if (WidgetId.IsValid())
 		{
-			// If bound to the same object, we don't want to remove it
-			if (ObjectId == WidgetId)
-			{
-				FName WidgetName = PreviewWidget->GetFName();
-				FName SlotWidgetName = NAME_None;
-
-				const UPanelSlot* WidgetSlot = Cast<UPanelSlot>(PreviewWidget);
-
-				if ((WidgetSlot != nullptr) && (WidgetSlot->Content != nullptr))
-				{
-					SlotWidgetName = WidgetSlot->GetFName();
-					WidgetName = WidgetSlot->Content->GetFName();
-				}
-
-				BindingsToRemove.RemoveAll([&](const FWidgetAnimationBinding& Binding) {
-					return Binding.WidgetName.IsEqual(WidgetName) && Binding.SlotWidgetName.IsEqual(SlotWidgetName);
-				});
-			}
-			WidgetsToAdd.Remove(Widget);
-			continue;
-		}
-
-		for (FMovieSceneBinding Binding : MovieSceneBindings)
-		{
-			if (ObjectId == Binding.GetObjectGuid())
-			{
-				TArray<UMovieSceneTrack*> MovieSceneTracks = Binding.GetTracks();
-				for (UMovieSceneTrack* Track : MovieSceneTracks)
-				{
-					UMovieScenePropertyTrack* PropertyTrack = Cast<UMovieScenePropertyTrack>(Track);
-					if (PropertyTrack)
-					{
-						FString PropertyName = PropertyTrack->GetPropertyName().ToString();
-						PropertyName.RemoveFromStart("b", ESearchCase::CaseSensitive);
-						FString NameString = "Set" + PropertyName;
-						FName FunctionName = FName(*NameString);
-						if (!Widget.GetTemplate()->FindFunction(FunctionName))
-						{
-							bIncompatibleTracks = true;
-							WidgetsToAdd.Remove(Widget);
-							goto NextWidget;
-						}
-					}
-				}
-			}
+			Widgets.RemoveAt(Index, 1, false);
 		}
-
-	NextWidget:
-		continue;
 	}
 
-	if (WidgetsToAdd.Num() == 0)
+	if (Widgets.Num() > 0)
 	{
-		if (bIncompatibleTracks)
-		{
-			// Exists a track that's not compatible 
-			FNotificationInfo Info(LOCTEXT("IncompatibleTracksToReplaceWith", "Selected Widgets do not match a Property this track is bound to"));
-			Info.FadeInDuration = 0.1f;
-			Info.FadeOutDuration = 0.5f;
-			Info.ExpireDuration = 2.5f;
-			auto NotificationItem = FSlateNotificationManager::Get().AddNotification(Info);
-
-			NotificationItem->SetCompletionState(SNotificationItem::CS_Success);
-			NotificationItem->ExpireAndFadeout();
-		}
-		else
-		{
-			FNotificationInfo Info(LOCTEXT("Widgetalreadybound", "Widget already bound"));
-			Info.FadeInDuration = 0.1f;
-			Info.FadeOutDuration = 0.5f;
-			Info.ExpireDuration = 2.5f;
-			auto NotificationItem = FSlateNotificationManager::Get().AddNotification(Info);
-
-			NotificationItem->SetCompletionState(SNotificationItem::CS_Success);
-			NotificationItem->ExpireAndFadeout();
-		}
-
-		return;
+		AddWidgetsToTrack(Widgets, ObjectId);
 	}
-
-	// else it's safe to modify
-	MovieScene->Modify();
-
-	// Replace bindings in WidgetAnimation
-	WidgetAnimation->Modify();
+	else
 	{
-		UUserWidget* PreviewRoot = GetPreview();
-		check(PreviewRoot);
-
-		for (const FWidgetAnimationBinding& Binding : BindingsToRemove)
-		{
-			UWidget* BoundWidget = PreviewRoot->GetWidgetFromName(Binding.WidgetName);
-			WidgetAnimation->RemoveBinding(Binding);
-			if (BoundWidget)
-			{
-				Sequencer->PreAnimatedState.RestorePreAnimatedState(*Sequencer, *BoundWidget);
-			}
-		}
-
-		if (WidgetsToAdd.Num() > 0)
-		{
-			AddWidgetsToTrack(WidgetsToAdd, ObjectId);
-		}
+		FNotificationInfo Info(LOCTEXT("Widgetalreadybound", "Widget already bound"));
+		Info.FadeInDuration = 0.1f;
+		Info.FadeOutDuration = 0.5f;
+		Info.ExpireDuration = 2.5f;
+		auto NotificationItem = FSlateNotificationManager::Get().AddNotification(Info);
 
-		UpdateTrackName(ObjectId);
+		NotificationItem->SetCompletionState(SNotificationItem::CS_Success);
+		NotificationItem->ExpireAndFadeout();
 	}
 
+	UpdateTrackName(ObjectId);
 	Sequencer->NotifyMovieSceneDataChanged(EMovieSceneDataChangeType::MovieSceneStructureItemsChanged);
 }
 
