UMG: Add selection guards during syncing so that extra rebuilds aren't processed (same as for level sequences)

diff --git a/Engine/Source/Editor/UMGEditor/Private/WidgetBlueprintEditor.cpp b/Engine/Source/Editor/UMGEditor/Private/WidgetBlueprintEditor.cpp
--- a/Engine/Source/Editor/UMGEditor/Private/WidgetBlueprintEditor.cpp
+++ b/Engine/Source/Editor/UMGEditor/Private/WidgetBlueprintEditor.cpp
@@ -64,6 +64,8 @@ FWidgetBlueprintEditor::FWidgetBlueprintEditor()
 	, bIsSimulateEnabled(false)
 	, bIsRealTime(true)
 	, bRefreshGeneratedClassAnimations(false)
+	, bUpdatingSequencerSelection(false)
+	, bUpdatingExternalSelection(false)
 {
 	PreviewScene.GetWorld()->bBegunPlay = false;
 
@@ -1701,6 +1703,13 @@ void FWidgetBlueprintEditor::OnMovieSceneBindingsPasted(const TArray<FMovieScene
 
 void FWidgetBlueprintEditor::SyncSelectedWidgetsWithSequencerSelection(TArray<FGuid> ObjectGuids)
 {
+	if (bUpdatingSequencerSelection)
+	{
+		return;
+	}
+
+	TGuardValue<bool> Guard(bUpdatingExternalSelection, true);
+
 	UMovieSceneSequence* AnimationSequence = GetSequencer().Get()->GetFocusedMovieSceneSequence();
 	UObject* BindingContext = GetAnimationPlaybackContext();
 	TSet<FWidgetReference> SequencerSelectedWidgets;
@@ -1729,6 +1738,13 @@ void FWidgetBlueprintEditor::SyncSelectedWidgetsWithSequencerSelection(TArray<FG
 
 void FWidgetBlueprintEditor::SyncSequencerSelectionToSelectedWidgets()
 {
+	if (bUpdatingExternalSelection)
+	{
+		return;
+	}
+
+	TGuardValue<bool> Guard(bUpdatingSequencerSelection, true);
+
 	GetSequencer()->ExternalSelectionHasChanged();
 }
 
diff --git a/Engine/Source/Editor/UMGEditor/Public/WidgetBlueprintEditor.h b/Engine/Source/Editor/UMGEditor/Public/WidgetBlueprintEditor.h
--- a/Engine/Source/Editor/UMGEditor/Public/WidgetBlueprintEditor.h
+++ b/Engine/Source/Editor/UMGEditor/Public/WidgetBlueprintEditor.h
@@ -385,4 +385,10 @@ private:
 
 	/** ViewModel used by the Palette and Palette Favorite Views */
 	TSharedPtr<FPaletteViewModel> PaletteViewModel;
+
+	/** When true the sequencer selection is being updated from changes to the external selection. */
+	bool bUpdatingSequencerSelection;
+
+	/** When true the external selection is being updated from changes to the sequencer selection. */
+	bool bUpdatingExternalSelection;
 };
