Disable ensure in accessibility code

diff --git a/Engine/Source/Editor/UMGEditor/Private/Customizations/UMGDetailCustomizations.cpp b/Engine/Source/Editor/UMGEditor/Private/Customizations/UMGDetailCustomizations.cpp
--- a/Engine/Source/Editor/UMGEditor/Private/Customizations/UMGDetailCustomizations.cpp
+++ b/Engine/Source/Editor/UMGEditor/Private/Customizations/UMGDetailCustomizations.cpp
@@ -378,12 +378,17 @@ void FBlueprintWidgetCustomization::CustomizeAccessibilityProperty(IDetailLayout
 	];
 
 	TSharedRef<SWidget> ExtensionWidget = const_cast<IDetailsView*>(DetailLayout.GetDetailsView())->GetExtensionHandler()->GenerateExtensionWidget(UWidget::StaticClass(), AccessibleTextPropertyHandle);
-	ensure(ExtensionWidget != SNullWidget::NullWidget);
-	CustomTextLayout->AddSlot()
-	.AutoWidth()
-	[
-		ExtensionWidget
-	];
+	// TODO: Figure out how to handle this case when property bindings are disabled for the widget.
+	//       Accessibility isn't a "true" property binding and should always be enabled. Might have to hard-code.
+	//ensure(ExtensionWidget != SNullWidget::NullWidget);
+	if (ExtensionWidget != SNullWidget::NullWidget)
+	{
+		CustomTextLayout->AddSlot()
+			.AutoWidth()
+			[
+				ExtensionWidget
+			];
+	}
 
 	TSharedPtr<SWidget> AccessibleBehaviorNameWidget, AccessibleBehaviorValueWidget;
 	AccessibilityRow.GetDefaultWidgets(AccessibleBehaviorNameWidget, AccessibleBehaviorValueWidget);
