Make AccessibleText use SPropertyBinding instead of GenerateExtensionWidget

diff --git a/Engine/Source/Editor/UMGEditor/Private/Customizations/UMGDetailCustomizations.cpp b/Engine/Source/Editor/UMGEditor/Private/Customizations/UMGDetailCustomizations.cpp
--- a/Engine/Source/Editor/UMGEditor/Private/Customizations/UMGDetailCustomizations.cpp
+++ b/Engine/Source/Editor/UMGEditor/Private/Customizations/UMGDetailCustomizations.cpp
@@ -361,9 +361,12 @@ void FBlueprintWidgetCustomization::CustomizeAccessibilityProperty(IDetailLayout
 	IDetailPropertyRow& AccessibilityRow = DetailLayout.EditCategory("Accessibility").AddProperty(AccessibleBehaviorPropertyHandle);
 
 	TSharedRef<IPropertyHandle> AccessibleTextPropertyHandle = DetailLayout.GetProperty(TextPropertyName);
+	const FName DelegateName(*(TextPropertyName.ToString() + "Delegate"));
+	UDelegateProperty* AccessibleTextDelegateProperty = FindFieldChecked<UDelegateProperty>(CastChecked<UClass>(AccessibleTextPropertyHandle->GetProperty()->GetOuter()), DelegateName);
 	// Make sure the old AccessibleText properties are hidden so we don't get duplicate widgets
 	DetailLayout.HideProperty(AccessibleTextPropertyHandle);
 
+	TSharedRef<SWidget> BindingWidget = SNew(SPropertyBinding, Editor.Pin().ToSharedRef(), AccessibleTextDelegateProperty, AccessibleTextPropertyHandle).GeneratePureBindings(false);
 	TSharedRef<SHorizontalBox> CustomTextLayout = SNew(SHorizontalBox)
 	.Visibility(TAttribute<EVisibility>::Create([AccessibleBehaviorPropertyHandle]() -> EVisibility
 	{
@@ -375,21 +378,13 @@ void FBlueprintWidgetCustomization::CustomizeAccessibilityProperty(IDetailLayout
 	.Padding(FMargin(4.0f, 0.0f))
 	[
 		AccessibleTextPropertyHandle->CreatePropertyValueWidget()
+	]
+	+SHorizontalBox::Slot()
+	.AutoWidth()
+	[
+		BindingWidget
 	];
 
-	TSharedRef<SWidget> ExtensionWidget = const_cast<IDetailsView*>(DetailLayout.GetDetailsView())->GetExtensionHandler()->GenerateExtensionWidget(UWidget::StaticClass(), AccessibleTextPropertyHandle);
-	// TODO: Figure out how to handle this case when property bindings are disabled for the widget.
-	//       Accessibility isn't a "true" property binding and should always be enabled. Might have to hard-code.
-	//ensure(ExtensionWidget != SNullWidget::NullWidget);
-	if (ExtensionWidget != SNullWidget::NullWidget)
-	{
-		CustomTextLayout->AddSlot()
-			.AutoWidth()
-			[
-				ExtensionWidget
-			];
-	}
-
 	TSharedPtr<SWidget> AccessibleBehaviorNameWidget, AccessibleBehaviorValueWidget;
 	AccessibilityRow.GetDefaultWidgets(AccessibleBehaviorNameWidget, AccessibleBehaviorValueWidget);
 
