Changed FWidgetBlueprintCompilerContext::PostCompile to FWidgetBlueprintCompilerContext::OnPostCDOCompiled. This allows checks to be made in the widget's CDO, rather than only instances.

#rnx

diff --git a/Engine/Source/Editor/UMGEditor/Private/WidgetBlueprintCompiler.cpp b/Engine/Source/Editor/UMGEditor/Private/WidgetBlueprintCompiler.cpp
--- a/Engine/Source/Editor/UMGEditor/Private/WidgetBlueprintCompiler.cpp
+++ b/Engine/Source/Editor/UMGEditor/Private/WidgetBlueprintCompiler.cpp
@@ -837,9 +837,9 @@ private:
 };
 
 
-void FWidgetBlueprintCompilerContext::PostCompile()
+void FWidgetBlueprintCompilerContext::OnPostCDOCompiled()
 {
-	Super::PostCompile();
+	Super::OnPostCDOCompiled();
 
 	WidgetToMemberVariableMap.Empty();
 	WidgetAnimToMemberVariableMap.Empty();
diff --git a/Engine/Source/Editor/UMGEditor/Public/WidgetBlueprintCompiler.h b/Engine/Source/Editor/UMGEditor/Public/WidgetBlueprintCompiler.h
--- a/Engine/Source/Editor/UMGEditor/Public/WidgetBlueprintCompiler.h
+++ b/Engine/Source/Editor/UMGEditor/Public/WidgetBlueprintCompiler.h
@@ -72,7 +72,7 @@ protected:
 	virtual void CopyTermDefaultsToDefaultObject(UObject* DefaultObject);
 	virtual void FinishCompilingClass(UClass* Class) override;
 	virtual bool ValidateGeneratedClass(UBlueprintGeneratedClass* Class) override;
-	virtual void PostCompile() override;
+	virtual void OnPostCDOCompiled() override;
 	// End FKismetCompilerContext
 
 	void SanitizeBindings(UBlueprintGeneratedClass* Class);
