UMG Editor - Fixing the WidgetTypeCustomization so that assigning widget references within a single blueprint is now possible.  Super experimental, needs more testing.

diff --git a/Engine/Source/Editor/UMGEditor/Private/Customizations/WidgetTypeCustomization.cpp b/Engine/Source/Editor/UMGEditor/Private/Customizations/WidgetTypeCustomization.cpp
--- a/Engine/Source/Editor/UMGEditor/Private/Customizations/WidgetTypeCustomization.cpp
+++ b/Engine/Source/Editor/UMGEditor/Private/Customizations/WidgetTypeCustomization.cpp
@@ -111,7 +111,12 @@ void FWidgetTypeCustomization::OnSelectionChanged(TWeakObjectPtr<UWidget> InItem
 		TSharedPtr<IPropertyHandle> PropertyHandle = PropertyHandlePtr.Pin();
 		if (PropertyHandle.IsValid())
 		{
-			PropertyHandle->SetValue(InItem.Get());
+			TArray<FString> Refrences;
+			for (int32 Index = 0; Index < PropertyHandle->GetNumPerObjectValues(); Index++)
+			{
+				Refrences.Add(InItem.Get()->GetPathName());
+			}
+			PropertyHandle->SetPerObjectValues(Refrences);
 
 			WidgetListComboButton->SetIsOpen(false);
 		}
