Correction to the Preview widget to spawn the first WidgetTree in the hierarchy of widgets just like it's done in UUSerWidget::Initialize. UUSerWidget::GetWidgetTreeOwningClass was changed to public to be able to use the same logic between both functions.

diff --git a/Engine/Source/Editor/UMGEditor/Private/WidgetBlueprintEditor.cpp b/Engine/Source/Editor/UMGEditor/Private/WidgetBlueprintEditor.cpp
--- a/Engine/Source/Editor/UMGEditor/Private/WidgetBlueprintEditor.cpp
+++ b/Engine/Source/Editor/UMGEditor/Private/WidgetBlueprintEditor.cpp
@@ -960,17 +960,15 @@ void FWidgetBlueprintEditor::UpdatePreview(UBlueprint* InBlueprint, bool bInForc
 
 			UWidgetTree* LatestWidgetTree = PreviewBlueprint->WidgetTree;
 
-			// HACK NickD: Doing this to match the hack in UUserWidget::Initialize(), to permit some semblance of widgettree
-			// inheritance.  This will correctly show the parent widget tree provided your class does not specify a root.
-			UWidgetBlueprintGeneratedClass* SuperBGClass = Cast<UWidgetBlueprintGeneratedClass>(PreviewBlueprint->GeneratedClass->GetSuperClass());
-			if ( SuperBGClass )
+			// If there is no RootWidget, we look for a WidgetTree in the parents classes until we find one.
+			if (LatestWidgetTree->RootWidget == nullptr)
 			{
-				UWidgetBlueprint* SuperWidgetBlueprint = Cast<UWidgetBlueprint>(SuperBGClass->ClassGeneratedBy);
-				if ( SuperWidgetBlueprint && (LatestWidgetTree->RootWidget == nullptr) )
+				UWidgetBlueprintGeneratedClass* BGClass = PreviewUserWidget->GetWidgetTreeOwningClass();
+				if (BGClass)
 				{
-					LatestWidgetTree = SuperWidgetBlueprint->WidgetTree;
+					LatestWidgetTree = BGClass->WidgetTree;
 				}
-			}
+			 }
 
 			// Update the widget tree directly to match the blueprint tree.  That way the preview can update
 			// without needing to do a full recompile.
