Fix to copy paste of Widget with named slot that created a copy of the widget in the root.

diff --git a/Engine/Source/Editor/UMGEditor/Private/Hierarchy/SHierarchyViewItem.cpp b/Engine/Source/Editor/UMGEditor/Private/Hierarchy/SHierarchyViewItem.cpp
--- a/Engine/Source/Editor/UMGEditor/Private/Hierarchy/SHierarchyViewItem.cpp
+++ b/Engine/Source/Editor/UMGEditor/Private/Hierarchy/SHierarchyViewItem.cpp
@@ -1026,12 +1026,6 @@ void FNamedSlotModel::DoDrop(UWidget* NamedSlotHostWidget, UWidget* DroppingWidg
 	TSet<FWidgetReference> SelectedTemplates;
 	SelectedTemplates.Add(BlueprintEditor.Pin()->GetReferenceFromTemplate(DroppingWidget));
 
-	// Move the remaining widgets into the transient package. Otherwise, they will remain outered to the WidgetTree and end up as properties in the BP class layout as a result.
-	if (DroppingWidget->GetOutermost() != GetTransientPackage())
-	{
-		FHierarchyModel::RemovePreviewWidget(Blueprint, DroppingWidget);
-	}
-
 	BlueprintEditor.Pin()->SelectWidgets(SelectedTemplates, false);
 }
 
diff --git a/Engine/Source/Editor/UMGEditor/Private/WidgetBlueprintEditorUtils.cpp b/Engine/Source/Editor/UMGEditor/Private/WidgetBlueprintEditorUtils.cpp
--- a/Engine/Source/Editor/UMGEditor/Private/WidgetBlueprintEditorUtils.cpp
+++ b/Engine/Source/Editor/UMGEditor/Private/WidgetBlueprintEditorUtils.cpp
@@ -1130,7 +1130,25 @@ TArray<UWidget*> FWidgetBlueprintEditorUtils::PasteWidgetsInternal(TSharedRef<FW
 		// automatically.
 		if ( NewWidget->GetParent() == nullptr )
 		{
-			RootPasteWidgets.Add(NewWidget);
+			// Check to see if this widget is content of another widget holding it in a named slot.
+			bool bIsNamedSlot = false;
+			for (UWidget* ContainerWidget : PastedWidgets)
+			{
+				if (INamedSlotInterface* NamedSlotContainer = Cast<INamedSlotInterface>(ContainerWidget))
+				{
+					if (NamedSlotContainer->ContainsContent(NewWidget))
+					{
+						bIsNamedSlot = true;
+						break;
+					}
+				}
+			}
+
+			// It's a Root widget only if it's not not in a named slot.
+			if (!bIsNamedSlot)
+			{
+				RootPasteWidgets.Add(NewWidget);
+			}
 		}
 	}
 
