Fixed log spam when copying and pasting widgets by removing parent slot property from the exported text

diff --git a/Engine/Source/Editor/UMGEditor/Private/WidgetBlueprintEditorUtils.cpp b/Engine/Source/Editor/UMGEditor/Private/WidgetBlueprintEditorUtils.cpp
--- a/Engine/Source/Editor/UMGEditor/Private/WidgetBlueprintEditorUtils.cpp
+++ b/Engine/Source/Editor/UMGEditor/Private/WidgetBlueprintEditorUtils.cpp
@@ -1347,6 +1347,17 @@ void FWidgetBlueprintEditorUtils::ImportWidgetsFromText(UWidgetBlueprint* BP, co
 
 		Widget->SetFlags(RF_Transactional);
 
+		// We don't export parent slot pointers, so each panel will need to point it's children back to itself
+		UPanelWidget* PanelWidget = Cast<UPanelWidget>(Widget);
+		if (PanelWidget)
+		{
+			TArray<UPanelSlot*> PanelSlots = PanelWidget->GetSlots();
+			for (int32 i = 0; i < PanelWidget->GetChildrenCount(); i++)
+			{
+				PanelWidget->GetChildAt(i)->Slot = PanelSlots[i];
+			}
+		}
+
 		// If there is an existing widget with the same name, rename the newly placed widget.
 		FString WidgetOldName = Widget->GetName();
 		if ( FindObject<UObject>(BP->WidgetTree, *WidgetOldName) )
diff --git a/Engine/Source/Runtime/UMG/Public/Components/Widget.h b/Engine/Source/Runtime/UMG/Public/Components/Widget.h
--- a/Engine/Source/Runtime/UMG/Public/Components/Widget.h
+++ b/Engine/Source/Runtime/UMG/Public/Components/Widget.h
@@ -249,7 +249,7 @@ public:
 	/**
 	 * The parent slot of the UWidget.  Allows us to easily inline edit the layout controlling this widget.
 	 */
-	UPROPERTY(Instanced, EditAnywhere, BlueprintReadOnly, Category=Layout, meta=(ShowOnlyInnerProperties))
+	UPROPERTY(Instanced, TextExportTransient, EditAnywhere, BlueprintReadOnly, Category=Layout, meta=(ShowOnlyInnerProperties))
 	UPanelSlot* Slot;
 
 	/** A bindable delegate for bIsEnabled */
