Fix to WidgetBlueprint compilation to enable BindWidget on widget in a parent WidgetBlueprint. To fix it, we only had to disable the error message at compilation. The binding did worked.

diff --git a/Engine/Source/Editor/UMGEditor/Private/WidgetBlueprintCompiler.cpp b/Engine/Source/Editor/UMGEditor/Private/WidgetBlueprintCompiler.cpp
--- a/Engine/Source/Editor/UMGEditor/Private/WidgetBlueprintCompiler.cpp
+++ b/Engine/Source/Editor/UMGEditor/Private/WidgetBlueprintCompiler.cpp
@@ -282,14 +282,32 @@ void FWidgetBlueprintCompilerContext::CreateClassVariablesFromBlueprint()
 
 	ValidateWidgetNames();
 
-	// Build the set of variables based on the variable widgets in the widget tree.
-	TArray<UWidget*> Widgets = WidgetBP->GetAllSourceWidgets();
+	// Build the set of variables based on the variable widgets in the first Widget Tree we find:
+	// in the current blueprint, the parent blueprint, and so on, until we find one.
+	TArray<UWidget*> Widgets;
+	UWidgetBlueprint* WidgetBPToScan = WidgetBP;
+	bool bSkipVariableCreation = false;
+	while (WidgetBPToScan != nullptr)
+	{
+		Widgets = WidgetBPToScan->GetAllSourceWidgets();
+		if (Widgets.Num() != 0)
+		{
+			// We found widgets.
+			break;
+		}
+		// We don't want to create variables for widgets that are in a parent blueprint. They will be created at the Parent compilation.
+		// But we want them to be added to the Member variable map for validation of the BindWidget property
+		bSkipVariableCreation = true;
+		
+		// Get the parent WidgetBlueprint
+		WidgetBPToScan = WidgetBPToScan->ParentClass && WidgetBPToScan->ParentClass->ClassGeneratedBy ? Cast<UWidgetBlueprint>(WidgetBPToScan->ParentClass->ClassGeneratedBy):nullptr;
+	}
 
 	// Sort the widgets alphabetically
 	Widgets.Sort( []( const UWidget& Lhs, const UWidget& Rhs ) { return Rhs.GetFName() < Lhs.GetFName(); } );
 
 	// Add widget variables
-	for ( UWidget* Widget : Widgets )
+	for ( UWidget* Widget : Widgets ) 
 	{
 		bool bIsVariable = Widget->bIsVariable;
 
@@ -311,6 +329,7 @@ void FWidgetBlueprintCompilerContext::CreateClassVariablesFromBlueprint()
 			WidgetClass = BPWidgetClass->GetAuthoritativeClass();
 		}
 
+		// Look in the Parent class properties to find a property with the BindWidget meta tag of the same name and Type.
 		UObjectPropertyBase* ExistingProperty = Cast<UObjectPropertyBase>(ParentClass->FindPropertyByName(Widget->GetFName()));
 		if (ExistingProperty && 
 			FWidgetBlueprintEditorUtils::IsBindWidgetProperty(ExistingProperty) && 
@@ -326,6 +345,12 @@ void FWidgetBlueprintCompilerContext::CreateClassVariablesFromBlueprint()
 			continue;
 		}
 
+        // We skip variable creation if the Widget Tree was in the Parent Blueprint.
+		if (bSkipVariableCreation)
+		{
+			continue;
+		}
+
 		FEdGraphPinType WidgetPinType(UEdGraphSchema_K2::PC_Object, NAME_None, WidgetClass, EPinContainerType::None, false, FEdGraphTerminalType());
 		
 		// Always name the variable according to the underlying FName of the widget object
