Fix to static analysis warning on WidgetBP that could be null

diff --git a/Engine/Source/Editor/UMGEditor/Private/WidgetBlueprintCompiler.cpp b/Engine/Source/Editor/UMGEditor/Private/WidgetBlueprintCompiler.cpp
--- a/Engine/Source/Editor/UMGEditor/Private/WidgetBlueprintCompiler.cpp
+++ b/Engine/Source/Editor/UMGEditor/Private/WidgetBlueprintCompiler.cpp
@@ -378,6 +378,11 @@ void FWidgetBlueprintCompilerContext::CreateClassVariablesFromBlueprint()
 		}
 	}
 
+	if (WidgetBP == nullptr)
+	{
+		return;
+	}
+	
 	// Add movie scenes variables here
 	for (UWidgetAnimation* Animation : WidgetBP->Animations)
 	{
