PR #5604: UE-71083: Prevent clear user focus on ValueSet (Contributed by [REDACTED_PERSON])

diff --git a/Engine/Source/Editor/UMGEditor/Private/Details/SWidgetDetailsView.cpp b/Engine/Source/Editor/UMGEditor/Private/Details/SWidgetDetailsView.cpp
--- a/Engine/Source/Editor/UMGEditor/Private/Details/SWidgetDetailsView.cpp
+++ b/Engine/Source/Editor/UMGEditor/Private/Details/SWidgetDetailsView.cpp
@@ -549,7 +549,11 @@ void SWidgetDetailsView::NotifyPostChange(const FPropertyChangedEvent& PropertyC
 		// Any time we migrate a property value we need to mark the blueprint as structurally modified so users don't need 
 		// to recompile it manually before they see it play in game using the latest version.
 		FBlueprintEditorUtils::MarkBlueprintAsModified(BlueprintEditor.Pin()->GetBlueprintObj());
-		ClearFocusIfOwned();
+		
+		if (PropertyChangedEvent.ChangeType != EPropertyChangeType::ValueSet)
+		{
+			ClearFocusIfOwned();
+		}
 	}
 
 	// If the property that changed is marked as "DesignerRebuild" we invalidate
