UMG Editor - Add a way in the hierarchy view to see if this widget overrides the navigation rules, similar to how we show the localization flow direction overrides.

Dave.Belanger, Vincent.Gauthier

diff --git a/Engine/Source/Editor/UMGEditor/Private/Customizations/WidgetNavigationCustomization.cpp b/Engine/Source/Editor/UMGEditor/Private/Customizations/WidgetNavigationCustomization.cpp
--- a/Engine/Source/Editor/UMGEditor/Private/Customizations/WidgetNavigationCustomization.cpp
+++ b/Engine/Source/Editor/UMGEditor/Private/Customizations/WidgetNavigationCustomization.cpp
@@ -419,7 +419,7 @@ void FWidgetNavigationCustomization::SetNav(UWidget* Widget, EUINavigation Nav,
 		DirectionNavigation->WidgetToFocus = WidgetToFocus.GetValue();
 	}
 
-	if ( WidgetNavigation->IsDefault() )
+	if ( WidgetNavigation->IsDefaultNavigation() )
 	{
 		// If the navigation rules are all set to the defaults, remove the navigation
 		// information from the widget.
diff --git a/Engine/Source/Editor/UMGEditor/Private/Hierarchy/SHierarchyViewItem.cpp b/Engine/Source/Editor/UMGEditor/Private/Hierarchy/SHierarchyViewItem.cpp
--- a/Engine/Source/Editor/UMGEditor/Private/Hierarchy/SHierarchyViewItem.cpp
+++ b/Engine/Source/Editor/UMGEditor/Private/Hierarchy/SHierarchyViewItem.cpp
@@ -1334,7 +1334,22 @@ void SHierarchyViewItem::Construct(const FArguments& InArgs, const TSharedRef< S
 			.VAlign(VAlign_Center)
 			[
 				SNew(STextBlock)
-				.ToolTipText(LOCTEXT("FlowDirectionHierarchyToolTip", "This widget overrides the flow direction preference."))
+				// TODO Tooltip should tell you what the widget is setup to do
+				.ToolTipText(LOCTEXT("NavigationHierarchyToolTip", "This widget overrides the navigation preference."))
+				.Visibility_Lambda([InModel] { return InModel->DoesWidgetOverrideNavigation() ? EVisibility::Visible : EVisibility::Collapsed; })
+				.ColorAndOpacity(FCoreStyle::Get().GetSlateColor("Foreground"))
+				.Font(FEditorStyle::Get().GetFontStyle("FontAwesome.10"))
+				.Text(FEditorFontGlyphs::Arrows)
+			]
+
+			// Localization Flow Direction Icon
+			+ SHorizontalBox::Slot()
+			.AutoWidth()
+			.VAlign(VAlign_Center)
+			[
+				SNew(STextBlock)
+				// TODO Tooltip should tell you what the widget is setup to do
+				.ToolTipText(LOCTEXT("FlowDirectionHierarchyToolTip", "This widget overrides the culture/localization flow direction preference."))
 				.Visibility_Lambda([InModel] { return InModel->DoesWidgetOverrideFlowDirection() ? EVisibility::Visible : EVisibility::Collapsed; })
 				.ColorAndOpacity(FCoreStyle::Get().GetSlateColor("Foreground"))
 				.Font(FEditorStyle::Get().GetFontStyle("FontAwesome.10"))
diff --git a/Engine/Source/Editor/UMGEditor/Private/Hierarchy/SHierarchyViewItem.h b/Engine/Source/Editor/UMGEditor/Private/Hierarchy/SHierarchyViewItem.h
--- a/Engine/Source/Editor/UMGEditor/Private/Hierarchy/SHierarchyViewItem.h
+++ b/Engine/Source/Editor/UMGEditor/Private/Hierarchy/SHierarchyViewItem.h
@@ -85,6 +85,7 @@ public:
 		}
 	}
 
+	virtual bool DoesWidgetOverrideNavigation() const { return false; }
 	virtual bool DoesWidgetOverrideFlowDirection() const { return false; }
 
 	virtual bool IsExpanded() const { return true; }
@@ -123,7 +124,7 @@ class FHierarchyRoot : public FHierarchyModel
 {
 public:
 	FHierarchyRoot(TSharedPtr<FWidgetBlueprintEditor> InBlueprintEditor);
-	
+
 	virtual ~FHierarchyRoot() {}
 
 	virtual FName GetUniqueName() const override;
@@ -138,6 +139,7 @@ public:
 	virtual void OnSelection() override;
 
 	virtual bool DoesWidgetOverrideFlowDirection() const override;
+	virtual bool DoesWidgetOverrideNavigation() const { return false; }
 
 	virtual TOptional<EItemDropZone> HandleCanAcceptDrop(const FDragDropEvent& DragDropEvent, EItemDropZone DropZone) override;
 	virtual FReply HandleAcceptDrop(FDragDropEvent const& DragDropEvent, EItemDropZone DropZone) override;
@@ -285,6 +287,17 @@ public:
 		return false;
 	}
 
+	virtual bool DoesWidgetOverrideNavigation() const override
+	{
+		UWidget* TemplateWidget = Item.GetTemplate();
+		if (TemplateWidget)
+		{
+			return TemplateWidget->Navigation != nullptr;
+		}
+
+		return false;
+	}
+
 	virtual bool IsExpanded() const override
 	{
 		if (UWidget* Template = Item.GetTemplate())
diff --git a/Engine/Source/Runtime/UMG/Private/UWidgetNavigation.cpp b/Engine/Source/Runtime/UMG/Private/UWidgetNavigation.cpp
--- a/Engine/Source/Runtime/UMG/Private/UWidgetNavigation.cpp
+++ b/Engine/Source/Runtime/UMG/Private/UWidgetNavigation.cpp
@@ -122,7 +122,7 @@ void UWidgetNavigation::UpdateMetaData(TSharedRef<FNavigationMetaData> MetaData)
 	UpdateMetaDataEntry(MetaData, Previous, EUINavigation::Previous);
 }
 
-bool UWidgetNavigation::IsDefault() const
+bool UWidgetNavigation::IsDefaultNavigation() const
 {
 	return Up.Rule == EUINavigationRule::Escape &&
 		Down.Rule == EUINavigationRule::Escape &&
diff --git a/Engine/Source/Runtime/UMG/Public/Blueprint/WidgetNavigation.h b/Engine/Source/Runtime/UMG/Public/Blueprint/WidgetNavigation.h
--- a/Engine/Source/Runtime/UMG/Public/Blueprint/WidgetNavigation.h
+++ b/Engine/Source/Runtime/UMG/Public/Blueprint/WidgetNavigation.h
@@ -98,7 +98,7 @@ public:
 	void UpdateMetaData(TSharedRef<FNavigationMetaData> MetaData);
 
 	/** @return true if the configured navigation object is the same as an un-customized navigation rule set. */
-	bool IsDefault() const;
+	bool IsDefaultNavigation() const;
 
 private:
 
