UMG: Close curve editor tab when null animation and hook up curve editor tab icon

diff --git a/Engine/Source/Editor/Sequencer/Private/SSequencer.cpp b/Engine/Source/Editor/Sequencer/Private/SSequencer.cpp
--- a/Engine/Source/Editor/Sequencer/Private/SSequencer.cpp
+++ b/Engine/Source/Editor/Sequencer/Private/SSequencer.cpp
@@ -3093,6 +3093,12 @@ void SSequencer::OnCurveEditorVisibilityChanged(bool bShouldBeVisible)
 		// that the Toolkit Host's Tab Manager has already registered a tab with a NullWidget for content.
 		TSharedRef<SDockTab> CurveEditorTab = Sequencer->GetToolkitHost()->GetTabManager()->InvokeTab(TabId);
 		CurveEditorTab->SetContent(CurveEditorPanel.ToSharedRef());
+
+		const FSlateIcon SequencerGraphIcon = FSlateIcon(FEditorStyle::GetStyleSetName(), "GenericCurveEditor.TabIcon");
+		CurveEditorTab->SetTabIcon(SequencerGraphIcon.GetIcon());
+
+		CurveEditorTab->SetLabel(LOCTEXT("SequencerMainGraphEditorTitle", "Sequencer Curves"));
+
 		SequencerPtr.Pin()->GetCurveEditor()->ZoomToFit();
 	}
 	else
diff --git a/Engine/Source/Editor/UMGEditor/Private/WidgetBlueprintEditor.cpp b/Engine/Source/Editor/UMGEditor/Private/WidgetBlueprintEditor.cpp
--- a/Engine/Source/Editor/UMGEditor/Private/WidgetBlueprintEditor.cpp
+++ b/Engine/Source/Editor/UMGEditor/Private/WidgetBlueprintEditor.cpp
@@ -869,6 +869,13 @@ void FWidgetBlueprintEditor::ChangeViewedAnimation( UWidgetAnimation& InAnimatio
 		TSharedPtr<STextBlock> NoAnimationTextBlockPin = NoAnimationTextBlock.Pin();
 		if( &InAnimationToView == UWidgetAnimation::GetNullAnimation())
 		{
+			const FName CurveEditorTabName = FName(TEXT("SequencerGraphEditor"));
+			TSharedPtr<SDockTab> ExistingTab = GetToolkitHost()->GetTabManager()->FindExistingLiveTab(CurveEditorTabName);
+			if (ExistingTab)
+			{
+				ExistingTab->RequestCloseTab();
+			}
+
 			// Disable sequencer from interaction
 			Sequencer->GetSequencerWidget()->SetEnabled(false);
 			Sequencer->SetAutoChangeMode(EAutoChangeMode::None);
