UMG: Fix include

diff --git a/Engine/Source/Editor/UMGEditor/Private/WidgetBlueprintEditor.cpp b/Engine/Source/Editor/UMGEditor/Private/WidgetBlueprintEditor.cpp
--- a/Engine/Source/Editor/UMGEditor/Private/WidgetBlueprintEditor.cpp
+++ b/Engine/Source/Editor/UMGEditor/Private/WidgetBlueprintEditor.cpp
@@ -6,6 +6,7 @@
 #include "MovieScene.h"
 #include "Animation/WidgetAnimation.h"
 #include "Widgets/Text/STextBlock.h"
+#include "Widgets/Docking/SDockTab.h"
 #include "Framework/MultiBox/MultiBoxBuilder.h"
 #include "Engine/SimpleConstructionScript.h"
 #include "Blueprint/WidgetBlueprintGeneratedClass.h"
