Fix creation of widget template objects when they are created during creation of another 'real' widget instance

diff --git a/Engine/Source/Editor/UMGEditor/Private/WidgetBlueprintCompiler.cpp b/Engine/Source/Editor/UMGEditor/Private/WidgetBlueprintCompiler.cpp
--- a/Engine/Source/Editor/UMGEditor/Private/WidgetBlueprintCompiler.cpp
+++ b/Engine/Source/Editor/UMGEditor/Private/WidgetBlueprintCompiler.cpp
@@ -623,7 +623,10 @@ void FWidgetBlueprintCompilerContext::FinishCompilingClass(UClass* Class)
 
 		BPGClass->bCookSlowConstructionWidgetTree = GetDefault<UUMGEditorProjectSettings>()->CompilerOption_CookSlowConstructionWidgetTree(WidgetBP);
 
-		BPGClass->WidgetTree = Cast<UWidgetTree>(StaticDuplicateObject(WidgetBP->WidgetTree, BPGClass, NAME_None, RF_AllFlags & ~RF_DefaultSubObject));
+		{
+			TGuardValue<uint32> DisableInitializeFromWidgetTree(UUserWidget::bInitializingFromWidgetTree, 0);
+			BPGClass->WidgetTree = Cast<UWidgetTree>(StaticDuplicateObject(WidgetBP->WidgetTree, BPGClass, NAME_None, RF_AllFlags & ~RF_DefaultSubObject));
+		}
 
 		for ( const UWidgetAnimation* Animation : WidgetBP->Animations )
 		{
diff --git a/Engine/Source/Runtime/UMG/Public/Blueprint/UserWidget.h b/Engine/Source/Runtime/UMG/Public/Blueprint/UserWidget.h
--- a/Engine/Source/Runtime/UMG/Public/Blueprint/UserWidget.h
+++ b/Engine/Source/Runtime/UMG/Public/Blueprint/UserWidget.h
@@ -1403,6 +1403,9 @@ protected:
 	 * UserWidget of state transitions.
 	 */
 	friend UUMGSequencePlayer;
+
+	/** The compiler is a friend so that it can disable initialization from the widget tree */
+	friend class FWidgetBlueprintCompilerContext;
 };
 
 #define LOCTEXT_NAMESPACE "UMG"
