UMG Editor: A Favorite panel was added to the Widget Designer.

- A lot of the Pallette View logic was moved the the PaletteViewModel to keep only the view part of the code in the SPalletteView class. A diff between the old View class and the ViewModel will simplify the code review.
- The SPaletteFavoriteView was created and hooked to the same PaletteViewModel as the SPaletteView
- FWidgetViewModel was also refactored a bit to keep most of the ViewModel logic out of the SPaletteItemView

#rn UMG Editor: A new Palette Favorite panel was added to UMG Designer. It can be openned through the Window menu. Your favorites widgets can then be added to that panel by clicking on the star icon next to the widget in the Palette panel creating a custom list of your favorite widgets for faster access.

diff --git a/Engine/Source/Editor/EditorStyle/Private/SlateEditorStyle.cpp b/Engine/Source/Editor/EditorStyle/Private/SlateEditorStyle.cpp
--- a/Engine/Source/Editor/EditorStyle/Private/SlateEditorStyle.cpp
+++ b/Engine/Source/Editor/EditorStyle/Private/SlateEditorStyle.cpp
@@ -7763,6 +7763,17 @@ void FSlateEditorStyle::FStyle::SetupUMGEditorStyles()
 		.SetOddRowBackgroundBrush(BOX_BRUSH("PropertyView/DetailCategoryMiddle", FMargin(4 / 16.0f, 8.0f / 16.0f, 4 / 16.0f, 4 / 16.0f)))
 		);
 
+	// Style of the favorite toggle
+	const FCheckBoxStyle UMGEditorFavoriteToggleStyle = FCheckBoxStyle()
+		.SetCheckBoxType(ESlateCheckBoxType::CheckBox)
+		.SetUncheckedImage(IMAGE_BRUSH("Icons/EmptyStar_16x", Icon10x10, FLinearColor(0.8f, 0.8f, 0.8f, 1.f)))
+		.SetUncheckedHoveredImage(IMAGE_BRUSH("Icons/EmptyStar_16x", Icon10x10, FLinearColor(2.5f, 2.5f, 2.5f, 1.f)))
+		.SetUncheckedPressedImage(IMAGE_BRUSH("Icons/EmptyStar_16x", Icon10x10, FLinearColor(0.8f, 0.8f, 0.8f, 1.f)))
+		.SetCheckedImage(IMAGE_BRUSH("Icons/Star_16x", Icon10x10, FLinearColor(0.2f, 0.2f, 0.2f, 1.f)))
+		.SetCheckedHoveredImage(IMAGE_BRUSH("Icons/Star_16x", Icon10x10, FLinearColor(0.4f, 0.4f, 0.4f, 1.f)))
+		.SetCheckedPressedImage(IMAGE_BRUSH("Icons/Star_16x", Icon10x10, FLinearColor(0.2f, 0.2f, 0.2f, 1.f)));
+	Set("UMGEditor.Palette.FavoriteToggleStyle", UMGEditorFavoriteToggleStyle);
+
 	Set("HorizontalAlignment_Left", new IMAGE_BRUSH("Icons/UMG/Alignment/Horizontal_Left", Icon20x20));
 	Set("HorizontalAlignment_Center", new IMAGE_BRUSH("Icons/UMG/Alignment/Horizontal_Center", Icon20x20));
 	Set("HorizontalAlignment_Right", new IMAGE_BRUSH("Icons/UMG/Alignment/Horizontal_Right", Icon20x20));
diff --git a/Engine/Source/Editor/UMGEditor/Classes/WidgetPaletteFavorites.cpp b/Engine/Source/Editor/UMGEditor/Classes/WidgetPaletteFavorites.cpp
new file mode 100644
--- /dev/null
+++ b/Engine/Source/Editor/UMGEditor/Classes/WidgetPaletteFavorites.cpp
@@ -0,0 +1,24 @@
+// Copyright 1998-2019 Epic Games, Inc. All Rights Reserved.
+
+#include "WidgetPaletteFavorites.h"
+#include "UObject/UnrealType.h"
+
+UWidgetPaletteFavorites::UWidgetPaletteFavorites(FObjectInitializer const& ObjectInitializer)
+	: Super(ObjectInitializer)
+{
+}
+
+void UWidgetPaletteFavorites::Add(const FString& InWidgetTemplateName)
+{
+	Favorites.AddUnique(InWidgetTemplateName);
+
+	SaveConfig();
+
+}
+
+void UWidgetPaletteFavorites::Remove(const FString& InWidgetTemplateName)
+{
+	Favorites.Remove(InWidgetTemplateName);
+
+	SaveConfig();
+}
\ No newline at end of file
diff --git a/Engine/Source/Editor/UMGEditor/Classes/WidgetPaletteFavorites.h b/Engine/Source/Editor/UMGEditor/Classes/WidgetPaletteFavorites.h
new file mode 100644
--- /dev/null
+++ b/Engine/Source/Editor/UMGEditor/Classes/WidgetPaletteFavorites.h
@@ -0,0 +1,24 @@
+// Copyright 1998-2019 Epic Games, Inc. All Rights Reserved.
+
+#pragma once
+
+#include "CoreMinimal.h"
+#include "UObject/ObjectMacros.h"
+#include "UObject/Object.h"
+#include "WidgetPaletteFavorites.generated.h"
+
+UCLASS(config = EditorPerProjectUserSettings)
+class UMGEDITOR_API UWidgetPaletteFavorites : public UObject
+{
+	GENERATED_UCLASS_BODY()
+public:
+
+	void Add(const FString& InWidgetTemplateName);
+	void Remove(const FString& InWidgetTemplateName);
+
+	TArray<FString> GetFavorites() const { return Favorites; }
+
+private:
+	UPROPERTY(config)
+	TArray<FString> Favorites;
+};
diff --git a/Engine/Source/Editor/UMGEditor/Private/BlueprintModes/WidgetDesignerApplicationMode.cpp b/Engine/Source/Editor/UMGEditor/Private/BlueprintModes/WidgetDesignerApplicationMode.cpp
--- a/Engine/Source/Editor/UMGEditor/Private/BlueprintModes/WidgetDesignerApplicationMode.cpp
+++ b/Engine/Source/Editor/UMGEditor/Private/BlueprintModes/WidgetDesignerApplicationMode.cpp
@@ -11,6 +11,7 @@
 
 
 #include "TabFactory/PaletteTabSummoner.h"
+#include "TabFactory/PaletteFavoritesTabSummoner.h"
 #include "TabFactory/HierarchyTabSummoner.h"
 #include "TabFactory/DesignerTabSummoner.h"
 #include "TabFactory/SequencerTabSummoner.h"
@@ -58,6 +59,12 @@ FWidgetDesignerApplicationMode::FWidgetDesignerApplicationMode(TSharedPtr<FWidge
 					->AddTab( FPaletteTabSummoner::TabID, ETabState::OpenedTab )
 				)
 				->Split
+				(
+					FTabManager::NewStack()
+					->SetSizeCoefficient( 0.5f )
+					->AddTab( FPaletteFavoritesTabSummoner::TabID, ETabState::OpenedTab )
+				)
+				->Split
 				(
 					FTabManager::NewStack()
 					->SetSizeCoefficient( 0.5f )
@@ -111,6 +118,7 @@ FWidgetDesignerApplicationMode::FWidgetDesignerApplicationMode(TSharedPtr<FWidge
 	TabFactories.RegisterFactory(MakeShareable(new FDesignerTabSummoner(InWidgetEditor)));
 	TabFactories.RegisterFactory(MakeShareable(new FHierarchyTabSummoner(InWidgetEditor)));
 	TabFactories.RegisterFactory(MakeShareable(new FPaletteTabSummoner(InWidgetEditor)));
+	TabFactories.RegisterFactory(MakeShareable(new FPaletteFavoritesTabSummoner(InWidgetEditor)));
 	TabFactories.RegisterFactory(MakeShareable(new FSequencerTabSummoner(InWidgetEditor)));
 	TabFactories.RegisterFactory(MakeShareable(new FAnimationTabSummoner(InWidgetEditor)));
 	TabFactories.RegisterFactory(MakeShareable(new FCompilerResultsSummoner(InWidgetEditor)));
diff --git a/Engine/Source/Editor/UMGEditor/Private/Palette/SPaletteFavoritesView.cpp b/Engine/Source/Editor/UMGEditor/Private/Palette/SPaletteFavoritesView.cpp
new file mode 100644
--- /dev/null
+++ b/Engine/Source/Editor/UMGEditor/Private/Palette/SPaletteFavoritesView.cpp
@@ -0,0 +1,42 @@
+// Copyright 1998-2019 Epic Games, Inc. All Rights Reserved.
+
+#include "Palette/SPaletteFavoritesView.h"
+#include "Palette/SPaletteView.h"
+#include "Palette/SPaletteViewModel.h"
+
+#include "UMGEditorProjectSettings.h"
+
+SPaletteFavoritesView::~SPaletteFavoritesView()
+{
+	PaletteViewModel->OnFavoritesUpdated.RemoveAll(this);
+}
+
+void SPaletteFavoritesView::Construct(const FArguments& InArgs, TSharedPtr<FWidgetBlueprintEditor> InBlueprintEditor)
+{
+	// Get the list of favorites from the Palette view model.
+	PaletteViewModel = InBlueprintEditor->GetPaletteViewModel();	
+
+	SAssignNew(WidgetTemplatesView, SListView<TSharedPtr<FWidgetViewModel>>)
+		.ListItemsSource(&(PaletteViewModel->GetFavoritesViewModels()))
+		.ItemHeight(1.0f)
+		.OnGenerateRow(this, &SPaletteFavoritesView::OnGenerateWidgetTemplateItem)
+		.SelectionMode(ESelectionMode::Single);
+
+	ChildSlot
+		[
+			WidgetTemplatesView.ToSharedRef()
+		];
+
+	PaletteViewModel->OnFavoritesUpdated.AddRaw(this, &SPaletteFavoritesView::OnFavoriteListUpdated);
+}
+
+void SPaletteFavoritesView::OnFavoriteListUpdated()
+{
+	WidgetTemplatesView->RequestListRefresh();
+}
+
+
+TSharedRef<ITableRow> SPaletteFavoritesView::OnGenerateWidgetTemplateItem(TSharedPtr<FWidgetViewModel> Item, const TSharedRef<STableViewBase>& OwnerTable)
+{
+	return Item->BuildRow(OwnerTable);
+}
diff --git a/Engine/Source/Editor/UMGEditor/Private/Palette/SPaletteFavoritesView.h b/Engine/Source/Editor/UMGEditor/Private/Palette/SPaletteFavoritesView.h
new file mode 100644
--- /dev/null
+++ b/Engine/Source/Editor/UMGEditor/Private/Palette/SPaletteFavoritesView.h
@@ -0,0 +1,31 @@
+// Copyright 1998-2019 Epic Games, Inc. All Rights Reserved.
+
+#pragma once
+
+#include "CoreMinimal.h"
+#include "Widgets/DeclarativeSyntaxSupport.h"
+#include "Widgets/SCompoundWidget.h"
+#include "Widgets/Views/STableViewBase.h"
+#include "Widgets/Views/SListView.h"
+
+class FWidgetBlueprintEditor;
+class FWidgetViewModel;
+class FPaletteViewModel;
+
+class SPaletteFavoritesView : public SCompoundWidget
+{
+public:
+	~SPaletteFavoritesView();
+
+	SLATE_BEGIN_ARGS(SPaletteFavoritesView){}
+	SLATE_END_ARGS()
+	
+	void Construct(const FArguments& InArgs, TSharedPtr<FWidgetBlueprintEditor> InBlueprintEditor);
+	TSharedRef<ITableRow> OnGenerateWidgetTemplateItem(TSharedPtr<FWidgetViewModel> Item, const TSharedRef<STableViewBase>& OwnerTable);
+	void OnFavoriteListUpdated();
+
+private:
+
+	TSharedPtr<SListView<TSharedPtr<FWidgetViewModel>>> WidgetTemplatesView;
+	TSharedPtr<FPaletteViewModel> PaletteViewModel;
+};
diff --git a/Engine/Source/Editor/UMGEditor/Private/Palette/SPaletteView.cpp b/Engine/Source/Editor/UMGEditor/Private/Palette/SPaletteView.cpp
--- a/Engine/Source/Editor/UMGEditor/Private/Palette/SPaletteView.cpp
+++ b/Engine/Source/Editor/UMGEditor/Private/Palette/SPaletteView.cpp
@@ -1,6 +1,7 @@
 // Copyright 1998-2019 Epic Games, Inc. All Rights Reserved.
 
 #include "Palette/SPaletteView.h"
+#include "Palette/SPaletteViewModel.h"
 #include "Misc/ConfigCacheIni.h"
 #include "UObject/UObjectHash.h"
 #include "UObject/UObjectIterator.h"
@@ -24,6 +25,7 @@
 
 #include "AssetRegistryModule.h"
 #include "Widgets/Input/SSearchBox.h"
+#include "Widgets/Input/SCheckBox.h"
 
 #include "Settings/ContentBrowserSettings.h"
 #include "WidgetBlueprintEditorUtils.h"
@@ -34,39 +36,66 @@
 
 #define LOCTEXT_NAMESPACE "UMG"
 
-class SPaletteViewItem : public SCompoundWidget
+FText SPaletteViewItem::GetFavoriteToggleToolTipText() const
 {
-public:
-
-	SLATE_BEGIN_ARGS(SPaletteViewItem) {}
-
-		/** The current text to highlight */
-		SLATE_ATTRIBUTE(FText, HighlightText)
+	if (GetFavoritedState() == ECheckBoxState::Checked)
+	{
+		return LOCTEXT("Unfavorite", "Click to remove this widget from your favorites.");
+	}
+	return LOCTEXT("Favorite", "Click to add this widget to your favorites.");
+}
 
-	SLATE_END_ARGS()
+ECheckBoxState SPaletteViewItem::GetFavoritedState() const
+{
+	if (WidgetViewModel->IsFavorite())
+	{
+		return ECheckBoxState::Checked;
+	}
+	else
+	{
+		return ECheckBoxState::Unchecked;
+	}
+}
 
-	/**
-	* Constructs this widget
-	*
-	* @param InArgs    Declaration from which to construct the widget
-	*/
-	void Construct(const FArguments& InArgs, TSharedPtr<FWidgetTemplate> InTemplate)
+void SPaletteViewItem::OnFavoriteToggled(ECheckBoxState InNewState)
+{
+	if (InNewState == ECheckBoxState::Checked)
 	{
-		Template = InTemplate;
+		//Add to favorites
+		WidgetViewModel->AddToFavorites();
+	}
+	else
+	{
+		//Remove from favorites
+		WidgetViewModel->RemoveFromFavorites();
+	}
+}
+
+void SPaletteViewItem::Construct(const FArguments& InArgs, TSharedPtr<FWidgetTemplateViewModel> InWidgetViewModel)
+{
+	WidgetViewModel = InWidgetViewModel;
 
-		ChildSlot
+	ChildSlot
 		[
 			SNew(SHorizontalBox)
-			.Visibility(EVisibility::Visible)
-			.ToolTip(Template->GetToolTip())
-
+			.ToolTip(WidgetViewModel->Template->GetToolTip())
+			+ SHorizontalBox::Slot()
+			.AutoWidth()
+			.VAlign(VAlign_Center)
+			[	
+				SNew(SCheckBox)
+				.ToolTipText(this, &SPaletteViewItem::GetFavoriteToggleToolTipText)
+				.IsChecked(this, &SPaletteViewItem::GetFavoritedState)
+				.OnCheckStateChanged(this, &SPaletteViewItem::OnFavoriteToggled)
+				.Style(FEditorStyle::Get(), "UMGEditor.Palette.FavoriteToggleStyle")
+			]
 			+ SHorizontalBox::Slot()
 			.AutoWidth()
 			.VAlign(VAlign_Center)
 			[
 				SNew(SImage)
 				.ColorAndOpacity(FLinearColor(1, 1, 1, 0.5))
-				.Image(Template->GetIcon())
+				.Image(WidgetViewModel->Template->GetIcon())
 			]
 
 			+ SHorizontalBox::Slot()
@@ -75,134 +104,34 @@ public:
 			.VAlign(VAlign_Center)
 			[
 				SNew(STextBlock)
-				.Text(Template->Name)
+				.Text(InWidgetViewModel->GetName())
 				.HighlightText(InArgs._HighlightText)
 			]
 		];
-	}
-
-	virtual FReply OnMouseButtonDoubleClick( const FGeometry& InMyGeometry, const FPointerEvent& InMouseEvent ) override
-	{
-		return Template->OnDoubleClicked();
-	}
-
-private:
-	TSharedPtr<FWidgetTemplate> Template;
-};
+}
 
-class FWidgetTemplateViewModel : public FWidgetViewModel
+FReply SPaletteViewItem::OnMouseButtonDoubleClick(const FGeometry& InMyGeometry, const FPointerEvent& InMouseEvent)
 {
-public:
-	virtual ~FWidgetTemplateViewModel()
-	{
-	}
-
-	virtual FText GetName() const override
-	{
-		return Template->Name;
-	}
-
-	virtual bool IsTemplate() const override
-	{
-		return true;
-	}
-
-	virtual FString GetFilterString() const override
-	{
-		return Template->Name.ToString();
-	}
-
-	virtual TSharedRef<ITableRow> BuildRow(const TSharedRef<STableViewBase>& OwnerTable) override
-	{
-		return SNew(STableRow< TSharedPtr<FWidgetViewModel> >, OwnerTable)
-			.Padding(2.0f)
-			.Style(FEditorStyle::Get(), "UMGEditor.PaletteItem")
-			.OnDragDetected(this, &FWidgetTemplateViewModel::OnDraggingWidgetTemplateItem)
-			[
-				SNew(SPaletteViewItem, Template)
-				.HighlightText(OwnerView, &SPaletteView::GetSearchText)
-			];
-	}
-
-	FReply OnDraggingWidgetTemplateItem(const FGeometry& MyGeometry, const FPointerEvent& MouseEvent)
-	{
-		return FReply::Handled().BeginDragDrop(FWidgetTemplateDragDropOp::New(Template));
-	}
-
-	SPaletteView* OwnerView;
-	TSharedPtr<FWidgetTemplate> Template;
-};
-
-class FWidgetHeaderViewModel : public FWidgetViewModel
-{
-public:
-	virtual ~FWidgetHeaderViewModel()
-	{
-	}
-
-	virtual FText GetName() const override
-	{
-		return GroupName;
-	}
-
-	virtual bool IsTemplate() const override
-	{
-		return false;
-	}
-
-	virtual FString GetFilterString() const override
-	{
-		// Headers should never be included in filtering to avoid showing a header with all of
-		// it's widgets filtered out, so return an empty filter string.
-		return TEXT("");
-	}
-	
-	virtual TSharedRef<ITableRow> BuildRow(const TSharedRef<STableViewBase>& OwnerTable) override
-	{
-		return SNew(STableRow< TSharedPtr<FWidgetViewModel> >, OwnerTable)
-			.Style( FEditorStyle::Get(), "UMGEditor.PaletteHeader" )
-			.Padding(2.0f)
-			.ShowSelection(false)
-			[
-				SNew(STextBlock)
-				.Text(GroupName)
-				.Font(FEditorStyle::GetFontStyle("DetailsView.CategoryFontStyle"))
-				.ShadowOffset(FVector2D(1.0f, 1.0f))
-			];
-	}
-
-	virtual void GetChildren(TArray< TSharedPtr<FWidgetViewModel> >& OutChildren) override
-	{
-		for ( TSharedPtr<FWidgetViewModel>& Child : Children )
-		{
-			OutChildren.Add(Child);
-		}
-	}
-
-	FText GroupName;
-	TArray< TSharedPtr<FWidgetViewModel> > Children;
+	return WidgetViewModel->Template->OnDoubleClicked();
 };
 
 void SPaletteView::Construct(const FArguments& InArgs, TSharedPtr<FWidgetBlueprintEditor> InBlueprintEditor)
 {
-	// Register for events that can trigger a palette rebuild
-	GEditor->OnBlueprintReinstanced().AddRaw(this, &SPaletteView::OnBlueprintReinstanced);
-	FEditorDelegates::OnAssetsDeleted.AddSP(this, &SPaletteView::HandleOnAssetsDeleted);
-	IHotReloadModule::Get().OnHotReload().AddSP(this, &SPaletteView::HandleOnHotReload);
-	
-	// register for any objects replaced
-	GEditor->OnObjectsReplaced().AddRaw(this, &SPaletteView::OnObjectsReplaced);
-
 	BlueprintEditor = InBlueprintEditor;
 
 	UBlueprint* BP = InBlueprintEditor->GetBlueprintObj();
+	PaletteViewModel = InBlueprintEditor->GetPaletteViewModel();
+
+	// Register to the update of the viewmodel.
+	PaletteViewModel->OnUpdating.AddRaw(this, &SPaletteView::OnViewModelUpdating);
+	PaletteViewModel->OnUpdated.AddRaw(this, &SPaletteView::OnViewModelUpdated);
 
 	WidgetFilter = MakeShareable(new WidgetViewModelTextFilter(
 		WidgetViewModelTextFilter::FItemToStringArray::CreateSP(this, &SPaletteView::TransformWidgetViewModelToString)));
 
 	FilterHandler = MakeShareable(new PaletteFilterHandler());
 	FilterHandler->SetFilter(WidgetFilter.Get());
-	FilterHandler->SetRootItems(&WidgetViewModels, &TreeWidgetViewModels);
+	FilterHandler->SetRootItems(&(PaletteViewModel->GetWidgetViewModels()), &TreeWidgetViewModels);
 	FilterHandler->SetGetChildrenDelegate(PaletteFilterHandler::FOnGetChildren::CreateRaw(this, &SPaletteView::OnGetChildren));
 
 	SAssignNew(WidgetTemplatesView, STreeView< TSharedPtr<FWidgetViewModel> >)
@@ -241,14 +170,16 @@ void SPaletteView::Construct(const FArguments& InArgs, TSharedPtr<FWidgetBluepri
 
 	bRefreshRequested = true;
 
-	BuildWidgetList();
+	PaletteViewModel->Update();
 	LoadItemExpansion();
-
-	bRebuildRequested = false;
 }
 
 SPaletteView::~SPaletteView()
 {
+	// Unregister to the update of the viewmodel.
+	PaletteViewModel->OnUpdating.RemoveAll(this);
+	PaletteViewModel->OnUpdated.RemoveAll(this);
+
 	// If the filter is enabled, disable it before saving the expanded items since
 	// filtering expands all items by default.
 	if (FilterHandler->GetIsEnabled())
@@ -257,12 +188,6 @@ SPaletteView::~SPaletteView()
 		FilterHandler->RefreshAndFilterTree();
 	}
 
-	GEditor->OnBlueprintReinstanced().RemoveAll(this);
-	FEditorDelegates::OnAssetsDeleted.RemoveAll(this);
-	IHotReloadModule::Get().OnHotReload().RemoveAll(this);
-	GEditor->OnObjectsReplaced().RemoveAll( this );
-	
-
 	SaveItemExpansion();
 }
 
@@ -272,12 +197,7 @@ void SPaletteView::OnSearchChanged(const FText& InFilterText)
 	FilterHandler->SetIsEnabled(!InFilterText.IsEmpty());
 	WidgetFilter->SetRawFilterText(InFilterText);
 	SearchBoxPtr->SetError(WidgetFilter->GetFilterErrorText());
-	SearchText = InFilterText;
-}
-
-FText SPaletteView::GetSearchText() const
-{
-	return SearchText;
+	PaletteViewModel->SetSearchText(InFilterText);
 }
 
 void SPaletteView::WidgetPalette_OnSelectionChanged(TSharedPtr<FWidgetViewModel> SelectedItem, ESelectInfo::Type SelectInfo)
@@ -338,7 +258,7 @@ TSharedPtr<FWidgetTemplate> SPaletteView::GetSelectedTemplateWidget() const
 void SPaletteView::LoadItemExpansion()
 {
 	// Restore the expansion state of the widget groups.
-	for ( TSharedPtr<FWidgetViewModel>& ViewModel : WidgetViewModels )
+	for ( TSharedPtr<FWidgetViewModel>& ViewModel : PaletteViewModel->GetWidgetViewModels())
 	{
 		bool IsExpanded;
 		if ( GConfig->GetBool(TEXT("WidgetTemplatesExpanded"), *ViewModel->GetName().ToString(), IsExpanded, GEditorPerProjectIni) && IsExpanded )
@@ -351,277 +271,13 @@ void SPaletteView::LoadItemExpansion()
 void SPaletteView::SaveItemExpansion()
 {
 	// Restore the expansion state of the widget groups.
-	for ( TSharedPtr<FWidgetViewModel>& ViewModel : WidgetViewModels )
+	for ( TSharedPtr<FWidgetViewModel>& ViewModel : PaletteViewModel->GetWidgetViewModels() )
 	{
 		const bool IsExpanded = WidgetTemplatesView->IsItemExpanded(ViewModel);
 		GConfig->SetBool(TEXT("WidgetTemplatesExpanded"), *ViewModel->GetName().ToString(), IsExpanded, GEditorPerProjectIni);
 	}
 }
 
-UWidgetBlueprint* SPaletteView::GetBlueprint() const
-{
-	if ( BlueprintEditor.IsValid() )
-	{
-		UBlueprint* BP = BlueprintEditor.Pin()->GetBlueprintObj();
-		return Cast<UWidgetBlueprint>(BP);
-	}
-
-	return NULL;
-}
-
-void SPaletteView::BuildWidgetList()
-{
-	// Clear the current list of view models and categories
-	WidgetViewModels.Reset();
-	WidgetTemplateCategories.Reset();
-
-	// Generate a list of templates
-	BuildClassWidgetList();
-	BuildSpecialWidgetList();
-
-	// For each entry in the category create a view model for the widget template
-	for ( auto& Entry : WidgetTemplateCategories )
-	{
-		TSharedPtr<FWidgetHeaderViewModel> Header = MakeShareable(new FWidgetHeaderViewModel());
-		Header->GroupName = FText::FromString(Entry.Key);
-
-		for ( auto& Template : Entry.Value )
-		{
-			TSharedPtr<FWidgetTemplateViewModel> TemplateViewModel = MakeShareable(new FWidgetTemplateViewModel());
-			TemplateViewModel->Template = Template;
-			TemplateViewModel->OwnerView = this;
-			Header->Children.Add(TemplateViewModel);
-		}
-
-		Header->Children.Sort([] (TSharedPtr<FWidgetViewModel> L, TSharedPtr<FWidgetViewModel> R) { return R->GetName().CompareTo(L->GetName()) > 0; });
-
-		WidgetViewModels.Add(Header);
-	}
-
-	// Sort the view models by name
-	WidgetViewModels.Sort([] (TSharedPtr<FWidgetViewModel> L, TSharedPtr<FWidgetViewModel> R) { return R->GetName().CompareTo(L->GetName()) > 0; });
-
-	// Take the Advanced Section, and put it at the end.
-	TSharedPtr<FWidgetViewModel>* advancedSectionPtr = WidgetViewModels.FindByPredicate([](TSharedPtr<FWidgetViewModel> widget) {return widget->GetName().CompareTo(LOCTEXT("Advanced", "Advanced")) == 0; });
-	if (advancedSectionPtr)
-	{
-		TSharedPtr<FWidgetViewModel> advancedSection = *advancedSectionPtr;
-		WidgetViewModels.Remove(advancedSection);
-		WidgetViewModels.Push(advancedSection);
-	}
-}
-
-void SPaletteView::BuildClassWidgetList()
-{
-	static const FName DevelopmentStatusKey(TEXT("DevelopmentStatus"));
-
-	TMap<FName, TSubclassOf<UUserWidget>> LoadedWidgetBlueprintClassesByName;
-
-	auto ActiveWidgetBlueprintClass = GetBlueprint()->GeneratedClass;
-	FName ActiveWidgetBlueprintClassName = ActiveWidgetBlueprintClass->GetFName();
-
-	TArray<FSoftClassPath> WidgetClassesToHide = GetDefault<UUMGEditorProjectSettings>()->WidgetClassesToHide;
-
-	// Locate all UWidget classes from code and loaded widget BPs
-	for (TObjectIterator<UClass> ClassIt; ClassIt; ++ClassIt)
-	{
-		UClass* WidgetClass = *ClassIt;
-
-		if (!FWidgetBlueprintEditorUtils::IsUsableWidgetClass(WidgetClass))
-		{
-			continue;
-		}
-
-		// Initialize AssetData for checking PackagePath
-		FAssetData WidgetAssetData = FAssetData(WidgetClass);
-
-		// Excludes engine content if user sets it to false
-		if (!GetDefault<UContentBrowserSettings>()->GetDisplayEngineFolder() || !GetDefault<UUMGEditorProjectSettings>()->bShowWidgetsFromEngineContent)
-		{
-			if (WidgetAssetData.PackagePath.ToString().Find(TEXT("/Engine")) == 0)
-			{
-				continue;
-			}
-		}
-
-		// Excludes developer content if user sets it to false
-		if (!GetDefault<UContentBrowserSettings>()->GetDisplayDevelopersFolder() || !GetDefault<UUMGEditorProjectSettings>()->bShowWidgetsFromDeveloperContent)
-		{
-			if (WidgetAssetData.PackagePath.ToString().Find(TEXT("/Game/Developers")) == 0)
-			{
-				continue;
-			}
-		}
-
-		// Excludes this widget if it is on the hide list
-		bool bIsOnList = false;
-		for (FSoftClassPath Widget : WidgetClassesToHide)
-		{
-			if (WidgetAssetData.ObjectPath.ToString().Find(Widget.ToString()) == 0)
-			{
-				bIsOnList = true;
-				break;
-			}
-		}
-		if (bIsOnList)
-		{
-			continue;
-		}
-
-		const bool bIsSameClass = WidgetClass->GetFName() == ActiveWidgetBlueprintClassName;
-
-		// Check that the asset that generated this class is valid (necessary b/c of a larger issue wherein force delete does not wipe the generated class object)
-		if ( bIsSameClass )
-		{
-			continue;
-		}
-
-		if (WidgetClass->IsChildOf(UUserWidget::StaticClass()))
-		{
-			if ( WidgetClass->ClassGeneratedBy )
-			{
-				// Track the widget blueprint classes that are already loaded
-				LoadedWidgetBlueprintClassesByName.Add(WidgetClass->ClassGeneratedBy->GetFName()) = WidgetClass;
-			}
-		}
-		else
-		{
-			TSharedPtr<FWidgetTemplateClass> Template = MakeShareable(new FWidgetTemplateClass(WidgetClass));
-
-			AddWidgetTemplate(Template);
-		}
-
-		//TODO UMG does not prevent deep nested circular references
-	}
-
-	// Locate all widget BP assets (include unloaded)
-	const FAssetRegistryModule& AssetRegistryModule = FModuleManager::LoadModuleChecked<FAssetRegistryModule>(TEXT("AssetRegistry"));
-	TArray<FAssetData> AllBPsAssetData;
-	AssetRegistryModule.Get().GetAssetsByClass(UBlueprint::StaticClass()->GetFName(), AllBPsAssetData, true);
-
-	for (FAssetData& BPAssetData : AllBPsAssetData)
-	{
-		// Blueprints get the class type actions for their parent native class - this avoids us having to load the blueprint
-		UClass* ParentClass = nullptr;
-		FString ParentClassName;
-		if (!BPAssetData.GetTagValue(FBlueprintTags::NativeParentClassPath, ParentClassName))
-		{
-			BPAssetData.GetTagValue(FBlueprintTags::ParentClassPath, ParentClassName);
-		}
-		if (!ParentClassName.IsEmpty())
-		{
-			UObject* Outer = nullptr;
-			ResolveName(Outer, ParentClassName, false, false);
-			ParentClass = FindObject<UClass>(ANY_PACKAGE, *ParentClassName);
-			// UUserWidgets have their own loading section, and we don't want to process any blueprints that don't have UWidget parents
-			if (!ParentClass->IsChildOf(UWidget::StaticClass()) || ParentClass->IsChildOf(UUserWidget::StaticClass()))
-			{
-				continue;
-			}
-		}
-
-		if (!FilterAssetData(BPAssetData))
-		{
-			// If this object isn't currently loaded, add it to the palette view
-			if (BPAssetData.ToSoftObjectPath().ResolveObject() == nullptr)
-			{
-				auto Template = MakeShareable(new FWidgetTemplateClass(BPAssetData, nullptr));
-				AddWidgetTemplate(Template);
-			}
-		}
-	}
-
-	TArray<FAssetData> AllWidgetBPsAssetData;
-	AssetRegistryModule.Get().GetAssetsByClass(UWidgetBlueprint::StaticClass()->GetFName(), AllWidgetBPsAssetData, true);
-
-	FName ActiveWidgetBlueprintName = ActiveWidgetBlueprintClass->ClassGeneratedBy->GetFName();
-	for (FAssetData& WidgetBPAssetData : AllWidgetBPsAssetData)
-	{
-		// Excludes the blueprint you're currently in
-		if (WidgetBPAssetData.AssetName == ActiveWidgetBlueprintName)
-		{
-			continue;
-		}
-
-		if (!FilterAssetData(WidgetBPAssetData))
-		{
-			// Excludes this widget if it is on the hide list
-			bool bIsOnList = false;
-			for (FSoftClassPath Widget : WidgetClassesToHide)
-			{
-				if (Widget.ToString().Find(WidgetBPAssetData.ObjectPath.ToString()) == 0)
-				{
-					bIsOnList = true;
-					break;
-				}
-			}
-			if (bIsOnList)
-			{
-				continue;
-			}
-
-			// If the blueprint generated class was found earlier, pass it to the template
-			TSubclassOf<UUserWidget> WidgetBPClass = nullptr;
-			auto LoadedWidgetBPClass = LoadedWidgetBlueprintClassesByName.Find(WidgetBPAssetData.AssetName);
-			if (LoadedWidgetBPClass)
-			{
-				WidgetBPClass = *LoadedWidgetBPClass;
-			}
-
-			auto Template = MakeShareable(new FWidgetTemplateBlueprintClass(WidgetBPAssetData, WidgetBPClass));
-
-			AddWidgetTemplate(Template);
-		}
-	}
-}
-
-bool SPaletteView::FilterAssetData(FAssetData &InAssetData)
-{
-	// Excludes engine content if user sets it to false
-	if (!GetDefault<UContentBrowserSettings>()->GetDisplayEngineFolder() || !GetDefault<UUMGEditorProjectSettings>()->bShowWidgetsFromEngineContent)
-	{
-		if (InAssetData.PackagePath.ToString().Find(TEXT("/Engine")) == 0)
-		{
-			return true;
-		}
-	}
-
-	// Excludes developer content if user sets it to false
-	if (!GetDefault<UContentBrowserSettings>()->GetDisplayDevelopersFolder() || !GetDefault<UUMGEditorProjectSettings>()->bShowWidgetsFromDeveloperContent)
-	{
-		if (InAssetData.PackagePath.ToString().Find(TEXT("/Game/Developers")) == 0)
-		{
-			return true;
-		}
-	}
-	return false;
-}
-
-void SPaletteView::BuildSpecialWidgetList()
-{
-	//AddWidgetTemplate(MakeShareable(new FWidgetTemplateButton()));
-	//AddWidgetTemplate(MakeShareable(new FWidgetTemplateCheckBox()));
-
-	//TODO UMG Make this pluggable.
-}
-
-void SPaletteView::AddWidgetTemplate(TSharedPtr<FWidgetTemplate> Template)
-{
-	FString Category = Template->GetCategory().ToString();
-
-	// Hide user specific categories
-	TArray<FString> CategoriesToHide = GetDefault<UUMGEditorProjectSettings>()->CategoriesToHide;
-	for (FString CategoryName : CategoriesToHide)
-	{
-		if (Category == CategoryName)
-		{
-			return;
-		}
-	}
-	WidgetTemplateArray& Group = WidgetTemplateCategories.FindOrAdd(Category);
-	Group.Add(Template);
-}
-
 void SPaletteView::OnGetChildren(TSharedPtr<FWidgetViewModel> Item, TArray< TSharedPtr<FWidgetViewModel> >& Children)
 {
 	return Item->GetChildren(Children);
@@ -632,31 +288,32 @@ TSharedRef<ITableRow> SPaletteView::OnGenerateWidgetTemplateItem(TSharedPtr<FWid
 	return Item->BuildRow(OwnerTable);
 }
 
-void SPaletteView::Tick( const FGeometry& AllottedGeometry, const double InCurrentTime, const float InDeltaTime )
+void SPaletteView::OnViewModelUpdating()
 {
-	if ( bRebuildRequested )
-	{
-		bRebuildRequested = false;
-
-		// Save the old expanded items temporarily
-		TSet<TSharedPtr<FWidgetViewModel>> ExpandedItems;
-		WidgetTemplatesView->GetExpandedItems(ExpandedItems);
+	// Save the old expanded items temporarily
+	WidgetTemplatesView->GetExpandedItems(ExpandedItems);
+}
 
-		BuildWidgetList();
+void SPaletteView::OnViewModelUpdated()
+{
+	bRefreshRequested = true;
 
-		// Restore the expansion state
-		for ( TSharedPtr<FWidgetViewModel>& ExpandedItem : ExpandedItems )
+	// Restore the expansion state
+	for (TSharedPtr<FWidgetViewModel>& ExpandedItem : ExpandedItems)
+	{
+		for (TSharedPtr<FWidgetViewModel>& ViewModel : PaletteViewModel->GetWidgetViewModels())
 		{
-			for ( TSharedPtr<FWidgetViewModel>& ViewModel : WidgetViewModels )
+			if (ViewModel->GetName().EqualTo(ExpandedItem->GetName()))
 			{
-				if ( ViewModel->GetName().EqualTo(ExpandedItem->GetName()) )
-				{
-					WidgetTemplatesView->SetItemExpansion(ViewModel, true);
-				}
+				WidgetTemplatesView->SetItemExpansion(ViewModel, true);
 			}
 		}
 	}
+	ExpandedItems.Reset();
+}
 
+void SPaletteView::Tick(const FGeometry& AllottedGeometry, const double InCurrentTime, const float InDeltaTime)
+{
 	if (bRefreshRequested)
 	{
 		bRefreshRequested = false;
@@ -669,35 +326,4 @@ void SPaletteView::TransformWidgetViewModelToString(TSharedPtr<FWidgetViewModel>
 	Array.Add(WidgetViewModel->GetFilterString());
 }
 
-void SPaletteView::OnObjectsReplaced(const TMap<UObject*, UObject*>& ReplacementMap)
-{
-	//bRefreshRequested = true;
-	//bRebuildRequested = true;
-}
-
-void SPaletteView::OnBlueprintReinstanced()
-{
-	bRebuildRequested = true;
-	bRefreshRequested = true;
-}
-
-void SPaletteView::HandleOnHotReload(bool bWasTriggeredAutomatically)
-{
-	bRebuildRequested = true;
-	bRefreshRequested = true;
-}
-
-void SPaletteView::HandleOnAssetsDeleted(const TArray<UClass*>& DeletedAssetClasses)
-{
-	for (auto DeletedAssetClass : DeletedAssetClasses)
-	{
-		if (DeletedAssetClass->IsChildOf(UWidgetBlueprint::StaticClass()))
-		{
-			bRebuildRequested = true;
-			bRefreshRequested = true;
-		}
-	}
-	
-}
-
 #undef LOCTEXT_NAMESPACE
diff --git a/Engine/Source/Editor/UMGEditor/Private/Palette/SPaletteView.h b/Engine/Source/Editor/UMGEditor/Private/Palette/SPaletteView.h
--- a/Engine/Source/Editor/UMGEditor/Private/Palette/SPaletteView.h
+++ b/Engine/Source/Editor/UMGEditor/Private/Palette/SPaletteView.h
@@ -5,8 +5,6 @@
 #include "CoreMinimal.h"
 #include "Widgets/DeclarativeSyntaxSupport.h"
 #include "Widgets/SCompoundWidget.h"
-#include "Widgets/Views/STableViewBase.h"
-#include "Widgets/Views/STableRow.h"
 #include "WidgetBlueprintEditor.h"
 #include "Misc/TextFilter.h"
 #include "Widgets/Views/STreeView.h"
@@ -14,28 +12,61 @@
 
 class FWidgetTemplate;
 class UWidgetBlueprint;
+class SPaletteView;
+class FPaletteViewModel;
+class FWidgetViewModel;
+class FWidgetTemplateViewModel;
 
-/** View model for the items in the widget template list */
-class FWidgetViewModel : public TSharedFromThis<FWidgetViewModel>
+/** Widget used to show a single row of the Palette and Palette favorite panel. */
+class SPaletteViewItem : public SCompoundWidget
 {
 public:
-	virtual ~FWidgetViewModel() { }
 
-	virtual FText GetName() const = 0;
+	SLATE_BEGIN_ARGS(SPaletteViewItem) {}
 
-	virtual bool IsTemplate() const = 0;
+	/** The current text to highlight */
+	SLATE_ATTRIBUTE(FText, HighlightText)
 
-	/** Get the string which should be used for filtering the item. */
-	virtual FString GetFilterString() const = 0;
+	SLATE_END_ARGS()
+
+	/**
+	* Constructs this widget
+	*
+	* @param InArgs    Declaration from which to construct the widget
+	*/
+	void Construct(const FArguments& InArgs, TSharedPtr<FWidgetTemplateViewModel> InWidgetViewModel);
 
-	virtual TSharedRef<ITableRow> BuildRow(const TSharedRef<STableViewBase>& OwnerTable) = 0;
+private:
+	/**
+	 * Retrieves tooltip that describes the current favorited state 
+	 *
+	 * @return Text describing what this toggle will do when you click on it.
+	 */
+	FText GetFavoriteToggleToolTipText() const;
+
+	/**
+	 * Checks on the associated action's favorite state, and returns a
+	 * corresponding checkbox state to match.
+	 *
+	 * @return ECheckBoxState::Checked if the associated action is already favorited, ECheckBoxState::Unchecked if not.
+	 */
+	ECheckBoxState GetFavoritedState() const;
+
+	/**
+	 * Triggers when the user clicks this toggle, adds or removes the associated
+	 * action to the user's favorites.
+	 *
+	 * @param  InNewState	The new state that the user set the checkbox to.
+	 */
+	void OnFavoriteToggled(ECheckBoxState InNewState);
+
+	virtual FReply OnMouseButtonDoubleClick(const FGeometry& InMyGeometry, const FPointerEvent& InMouseEvent) override;
 
-	virtual void GetChildren(TArray< TSharedPtr<FWidgetViewModel> >& OutChildren)
-	{
-	}
+private:
+	TSharedPtr<FWidgetTemplateViewModel> WidgetViewModel;
 };
 
-/**  */
+/** */
 class SPaletteView : public SCompoundWidget
 {
 public:
@@ -60,12 +91,6 @@ public:
 	TSharedPtr<FWidgetTemplate> GetSelectedTemplateWidget() const;
 
 private:
-	UWidgetBlueprint* GetBlueprint() const;
-
-	void BuildWidgetList();
-	void BuildClassWidgetList();
-	bool FilterAssetData(FAssetData &BPAssetData);
-	void BuildSpecialWidgetList();
 
 	void OnGetChildren(TSharedPtr<FWidgetViewModel> Item, TArray< TSharedPtr<FWidgetViewModel> >& Children);
 	TSharedRef<ITableRow> OnGenerateWidgetTemplateItem(TSharedPtr<FWidgetViewModel> Item, const TSharedRef<STableViewBase>& OwnerTable);
@@ -73,41 +98,28 @@ private:
 	/** Called when the filter text is changed. */
 	void OnSearchChanged(const FText& InFilterText);
 
+	void OnViewModelUpdating();
+	void OnViewModelUpdated();
+
 private:
+	/** Load the expansion state for the TreeView */
 	void LoadItemExpansion();
-	void SaveItemExpansion();
 
-	/** Called when a Blueprint is recompiled and live objects are swapped out for replacements */
-	void OnObjectsReplaced(const TMap<UObject*, UObject*>& ReplacementMap);
-
-	void AddWidgetTemplate(TSharedPtr<FWidgetTemplate> Template);
+	/** Save the expansion state for the TreeView */
+	void SaveItemExpansion();
 
+	// Search functions
 	/** Transforms the widget view model into a searchable string. */
 	void TransformWidgetViewModelToString(TSharedPtr<FWidgetViewModel> WidgetViewModel, OUT TArray< FString >& Array);
 
-	/** Requests a rebuild of the widget list if a widget blueprint was compiled */
-	void OnBlueprintReinstanced();
-
-	/** Requests a rebuild of the widget list */
-	void HandleOnHotReload(bool bWasTriggeredAutomatically);
-
-	/** Requests a rebuild of the widget list if a widget blueprint was deleted */
-	void HandleOnAssetsDeleted(const TArray<UClass*>& DeletedAssetClasses);
-
 	TWeakPtr<class FWidgetBlueprintEditor> BlueprintEditor;
+	TSharedPtr<FPaletteViewModel> PaletteViewModel;
 
 	/** Handles filtering the palette based on an IFilter. */
 	typedef TreeFilterHandler<TSharedPtr<FWidgetViewModel>> PaletteFilterHandler;
 	TSharedPtr<PaletteFilterHandler> FilterHandler;
 
-	typedef TArray< TSharedPtr<FWidgetTemplate> > WidgetTemplateArray;
-	TMap< FString, WidgetTemplateArray > WidgetTemplateCategories;
-
-	typedef TArray< TSharedPtr<FWidgetViewModel> > ViewModelsArray;
-	
-	/** The source root view models for the tree. */
-	ViewModelsArray WidgetViewModels;
-
+	typedef TArray<TSharedPtr<FWidgetViewModel>> ViewModelsArray;
 	/** The root view models which are actually displayed by the TreeView which will be managed by the TreeFilterHandler. */
 	ViewModelsArray TreeWidgetViewModels;
 
@@ -119,9 +131,9 @@ private:
 	/** The filter instance which is used by the TreeFilterHandler to filter the TreeView. */
 	TSharedPtr<WidgetViewModelTextFilter> WidgetFilter;
 
-	bool bRefreshRequested;
-	FText SearchText;
+	/** Expended Items in the Tree view */
+	TSet<TSharedPtr<FWidgetViewModel>> ExpandedItems;
 
-	/** Controls rebuilding the list of spawnable widgets */
-	bool bRebuildRequested;
+	/** Set to true to force a refresh of the treeview */
+	bool bRefreshRequested;
 };
diff --git a/Engine/Source/Editor/UMGEditor/Private/Palette/SPaletteViewModel.cpp b/Engine/Source/Editor/UMGEditor/Private/Palette/SPaletteViewModel.cpp
new file mode 100644
--- /dev/null
+++ b/Engine/Source/Editor/UMGEditor/Private/Palette/SPaletteViewModel.cpp
@@ -0,0 +1,486 @@
+// Copyright 1998-2019 Epic Games, Inc. All Rights Reserved.
+
+#include "Palette/SPaletteViewModel.h"
+#include "Palette/SPaletteView.h"
+#include "Widgets/Views/STableViewBase.h"
+#include "Widgets/Views/STableRow.h"
+#include "WidgetBlueprint.h"
+#include "Editor.h"
+
+#if WITH_EDITOR
+	#include "EditorStyleSet.h"
+#endif // WITH_EDITOR
+
+#include "DragDrop/WidgetTemplateDragDropOp.h"
+
+#include "Templates/WidgetTemplateClass.h"
+#include "Templates/WidgetTemplateBlueprintClass.h"
+
+#include "Developer/HotReload/Public/IHotReload.h"
+
+#include "AssetRegistryModule.h"
+#include "WidgetBlueprintEditorUtils.h"
+
+#include "Settings/ContentBrowserSettings.h"
+#include "Settings/WidgetDesignerSettings.h"
+#include "UMGEditorProjectSettings.h"
+#include "WidgetPaletteFavorites.h"
+
+#define LOCTEXT_NAMESPACE "UMG"
+
+FWidgetTemplateViewModel::FWidgetTemplateViewModel()
+	: PaletteViewModel(nullptr),
+	bIsFavorite(false)
+{
+}
+
+FText FWidgetTemplateViewModel::GetName() const
+{
+	return Template->Name;
+}
+
+bool FWidgetTemplateViewModel::IsTemplate() const
+{
+	return true;
+}
+
+FString FWidgetTemplateViewModel::GetFilterString() const
+{
+	return Template->Name.ToString();
+}
+
+TSharedRef<ITableRow> FWidgetTemplateViewModel::BuildRow(const TSharedRef<STableViewBase>& OwnerTable)
+{
+	return SNew(STableRow<TSharedPtr<FWidgetViewModel>>, OwnerTable)
+		.Padding(2.0f)
+		.Style(FEditorStyle::Get(), "UMGEditor.PaletteItem")
+		.OnDragDetected(this, &FWidgetTemplateViewModel::OnDraggingWidgetTemplateItem)
+		[
+			SNew(SPaletteViewItem, SharedThis(this))
+			.HighlightText(PaletteViewModel, &FPaletteViewModel::GetSearchText)
+		];
+}
+
+FReply FWidgetTemplateViewModel::OnDraggingWidgetTemplateItem(const FGeometry& MyGeometry, const FPointerEvent& MouseEvent)
+{
+	return FReply::Handled().BeginDragDrop(FWidgetTemplateDragDropOp::New(Template));
+}
+
+void FWidgetTemplateViewModel::AddToFavorites()
+{
+	bIsFavorite = true;
+	PaletteViewModel->AddToFavorites(this);
+}
+
+void FWidgetTemplateViewModel::RemoveFromFavorites()
+{
+	bIsFavorite = false;
+	PaletteViewModel->RemoveFromFavorites(this);
+}
+
+TSharedRef<ITableRow> FWidgetHeaderViewModel::BuildRow(const TSharedRef<STableViewBase>& OwnerTable)
+{
+	return SNew(STableRow<TSharedPtr<FWidgetViewModel>>, OwnerTable)
+		.Style(FEditorStyle::Get(), "UMGEditor.PaletteHeader")
+		.Padding(2.0f)
+		.ShowSelection(false)
+		[
+			SNew(STextBlock)
+			.Text(GroupName)
+			.Font(FEditorStyle::GetFontStyle("DetailsView.CategoryFontStyle"))
+			.ShadowOffset(FVector2D(1.0f, 1.0f))
+		];
+}
+
+void FWidgetHeaderViewModel::GetChildren(TArray< TSharedPtr<FWidgetViewModel> >& OutChildren)
+{
+	for (TSharedPtr<FWidgetViewModel>& Child : Children)
+	{
+		OutChildren.Add(Child);
+	}
+}
+
+FPaletteViewModel::FPaletteViewModel(TSharedPtr<FWidgetBlueprintEditor> InBlueprintEditor)
+	: bRebuildRequested(true),
+	  bRebuildFavorites(true)
+{
+	BlueprintEditor = InBlueprintEditor;
+}
+
+void FPaletteViewModel::RegisterToEvents()
+{
+	// Register for events that can trigger a palette rebuild
+	GEditor->OnBlueprintReinstanced().AddRaw(this, &FPaletteViewModel::OnBlueprintReinstanced);
+	FEditorDelegates::OnAssetsDeleted.AddSP(this, &FPaletteViewModel::HandleOnAssetsDeleted);
+	IHotReloadModule::Get().OnHotReload().AddSP(this, &FPaletteViewModel::HandleOnHotReload);
+
+	// register for any objects replaced
+	GEditor->OnObjectsReplaced().AddRaw(this, &FPaletteViewModel::OnObjectsReplaced);
+}
+
+FPaletteViewModel::~FPaletteViewModel()
+{
+	GEditor->OnBlueprintReinstanced().RemoveAll(this);
+	FEditorDelegates::OnAssetsDeleted.RemoveAll(this);
+	IHotReloadModule::Get().OnHotReload().RemoveAll(this);
+	GEditor->OnObjectsReplaced().RemoveAll(this);
+}
+
+void FPaletteViewModel::AddToFavorites(const FWidgetTemplateViewModel* WidgetTemplateViewModel)
+{
+	UWidgetPaletteFavorites* Favorites = GetDefault<UWidgetDesignerSettings>()->Favorites;
+	Favorites->Add(WidgetTemplateViewModel->GetName().ToString());
+
+	// Force a rebuild of the favorites list.
+	bRebuildFavorites = true;
+}
+
+void FPaletteViewModel::RemoveFromFavorites(const FWidgetTemplateViewModel* WidgetTemplateViewModel)
+{
+	UWidgetPaletteFavorites* Favorites = GetDefault<UWidgetDesignerSettings>()->Favorites;
+	Favorites->Remove(WidgetTemplateViewModel->GetName().ToString());
+	
+	// Force a rebuild of the favorites list.
+	bRebuildFavorites = true;
+}
+
+void FPaletteViewModel::Update()
+{
+	if (bRebuildRequested)
+	{
+		bRebuildRequested = false;
+
+		OnUpdating.Broadcast();
+		BuildWidgetList();
+		OnUpdated.Broadcast();
+
+		// We need to rebuild the list of favorites to also clear widgets that where potentially removed.
+		bRebuildFavorites = true;
+	}
+
+	if (bRebuildFavorites)
+	{
+		bRebuildFavorites = false;
+		BuildWidgetFavoriteList();
+		OnFavoritesUpdated.Broadcast();
+	}
+}
+
+
+UWidgetBlueprint* FPaletteViewModel::GetBlueprint() const
+{
+	if (BlueprintEditor.IsValid())
+	{
+		UBlueprint* BP = BlueprintEditor.Pin()->GetBlueprintObj();
+		return Cast<UWidgetBlueprint>(BP);
+	}
+
+	return NULL;
+}
+
+void FPaletteViewModel::BuildWidgetList()
+{
+	// Clear the current list of view models and categories
+	WidgetViewModels.Reset();
+	WidgetTemplateCategories.Reset();
+
+	// Generate a list of templates
+	BuildClassWidgetList();
+
+	// For each entry in the category create a view model for the widget template
+	for ( auto& Entry : WidgetTemplateCategories )
+	{
+		TSharedPtr<FWidgetHeaderViewModel> Header = MakeShareable(new FWidgetHeaderViewModel());
+		Header->GroupName = FText::FromString(Entry.Key);
+
+		for ( auto& Template : Entry.Value )
+		{
+			TSharedPtr<FWidgetTemplateViewModel> TemplateViewModel = MakeShareable(new FWidgetTemplateViewModel());
+			TemplateViewModel->Template = Template;
+			TemplateViewModel->PaletteViewModel = this;
+			Header->Children.Add(TemplateViewModel);
+		}
+
+		Header->Children.Sort([] (TSharedPtr<FWidgetViewModel> L, TSharedPtr<FWidgetViewModel> R) { return R->GetName().CompareTo(L->GetName()) > 0; });
+
+		WidgetViewModels.Add(Header);
+	}
+
+	// Sort the view models by name
+	WidgetViewModels.Sort([] (TSharedPtr<FWidgetViewModel> L, TSharedPtr<FWidgetViewModel> R) { return R->GetName().CompareTo(L->GetName()) > 0; });
+
+	// Take the Advanced Section, and put it at the end.
+	TSharedPtr<FWidgetViewModel>* advancedSectionPtr = WidgetViewModels.FindByPredicate([](TSharedPtr<FWidgetViewModel> widget) {return widget->GetName().CompareTo(LOCTEXT("Advanced", "Advanced")) == 0; });
+	if (advancedSectionPtr)
+	{
+		TSharedPtr<FWidgetViewModel> advancedSection = *advancedSectionPtr;
+		WidgetViewModels.Remove(advancedSection);
+		WidgetViewModels.Push(advancedSection);
+	}
+}
+
+void FPaletteViewModel::BuildWidgetFavoriteList()
+{
+	WidgetFavoritesViewModels.Reset();
+	UWidgetPaletteFavorites* FavoritesPalette = GetDefault<UWidgetDesignerSettings>()->Favorites;
+
+	// We Iterate on a copy to be able to clean widgets that may no longer exist
+	TArray<FString> FavoritesList = FavoritesPalette->GetFavorites();
+	for (const FString& favoriteName : FavoritesList)
+	{
+		bool bWidgetFound = false;
+		for (TSharedPtr < FWidgetViewModel>& widgetSection : WidgetViewModels)
+		{
+			TArray < TSharedPtr<FWidgetViewModel> > Children;
+			widgetSection->GetChildren(Children);
+
+			TSharedPtr<FWidgetViewModel>* widgetViewModel = Children.FindByPredicate([&](const TSharedPtr<FWidgetViewModel>& widget) { return widget->GetName().ToString().Compare(favoriteName) == 0; });
+			if (widgetViewModel != nullptr)
+			{
+				(*widgetViewModel)->SetFavorite();
+
+				WidgetFavoritesViewModels.Add(*widgetViewModel);
+				bWidgetFound = true;
+				break;
+			}
+		}
+
+		// If we did not found the widget, it no longer exists. We remove it from the original list of favorites.
+		if (!bWidgetFound)
+		{
+			FavoritesPalette->Remove(favoriteName);
+		}
+	}
+
+	WidgetFavoritesViewModels.Sort([](TSharedPtr<FWidgetViewModel> L, TSharedPtr<FWidgetViewModel> R) { return R->GetName().CompareTo(L->GetName()) > 0; });
+}
+
+void FPaletteViewModel::BuildClassWidgetList()
+{
+	static const FName DevelopmentStatusKey(TEXT("DevelopmentStatus"));
+
+	TMap<FName, TSubclassOf<UUserWidget>> LoadedWidgetBlueprintClassesByName;
+
+	auto ActiveWidgetBlueprintClass = GetBlueprint()->GeneratedClass;
+	FName ActiveWidgetBlueprintClassName = ActiveWidgetBlueprintClass->GetFName();
+
+	TArray<FSoftClassPath> WidgetClassesToHide = GetDefault<UUMGEditorProjectSettings>()->WidgetClassesToHide;
+
+	// Locate all UWidget classes from code and loaded widget BPs
+	for (TObjectIterator<UClass> ClassIt; ClassIt; ++ClassIt)
+	{
+		UClass* WidgetClass = *ClassIt;
+
+		if (!FWidgetBlueprintEditorUtils::IsUsableWidgetClass(WidgetClass))
+		{
+			continue;
+		}
+
+		// Initialize AssetData for checking PackagePath
+		FAssetData WidgetAssetData = FAssetData(WidgetClass);
+
+		// Excludes engine content if user sets it to false
+		if (!GetDefault<UContentBrowserSettings>()->GetDisplayEngineFolder() || !GetDefault<UUMGEditorProjectSettings>()->bShowWidgetsFromEngineContent)
+		{
+			if (WidgetAssetData.PackagePath.ToString().Find(TEXT("/Engine")) == 0)
+			{
+				continue;
+			}
+		}
+
+		// Excludes developer content if user sets it to false
+		if (!GetDefault<UContentBrowserSettings>()->GetDisplayDevelopersFolder() || !GetDefault<UUMGEditorProjectSettings>()->bShowWidgetsFromDeveloperContent)
+		{
+			if (WidgetAssetData.PackagePath.ToString().Find(TEXT("/Game/Developers")) == 0)
+			{
+				continue;
+			}
+		}
+
+		// Excludes this widget if it is on the hide list
+		bool bIsOnList = false;
+		for (FSoftClassPath Widget : WidgetClassesToHide)
+		{
+			if (WidgetAssetData.ObjectPath.ToString().Find(Widget.ToString()) == 0)
+			{
+				bIsOnList = true;
+				break;
+			}
+		}
+		if (bIsOnList)
+		{
+			continue;
+		}
+
+		const bool bIsSameClass = WidgetClass->GetFName() == ActiveWidgetBlueprintClassName;
+
+		// Check that the asset that generated this class is valid (necessary b/c of a larger issue wherein force delete does not wipe the generated class object)
+		if ( bIsSameClass )
+		{
+			continue;
+		}
+
+		if (WidgetClass->IsChildOf(UUserWidget::StaticClass()))
+		{
+			if ( WidgetClass->ClassGeneratedBy )
+			{
+				// Track the widget blueprint classes that are already loaded
+				LoadedWidgetBlueprintClassesByName.Add(WidgetClass->ClassGeneratedBy->GetFName()) = WidgetClass;
+			}
+		}
+		else
+		{
+			TSharedPtr<FWidgetTemplateClass> Template = MakeShareable(new FWidgetTemplateClass(WidgetClass));
+
+			AddWidgetTemplate(Template);
+		}
+
+		//TODO UMG does not prevent deep nested circular references
+	}
+
+	// Locate all widget BP assets (include unloaded)
+	const FAssetRegistryModule& AssetRegistryModule = FModuleManager::LoadModuleChecked<FAssetRegistryModule>(TEXT("AssetRegistry"));
+	TArray<FAssetData> AllBPsAssetData;
+	AssetRegistryModule.Get().GetAssetsByClass(UBlueprint::StaticClass()->GetFName(), AllBPsAssetData, true);
+
+	for (FAssetData& BPAssetData : AllBPsAssetData)
+	{
+		// Blueprints get the class type actions for their parent native class - this avoids us having to load the blueprint
+		UClass* ParentClass = nullptr;
+		FString ParentClassName;
+		if (!BPAssetData.GetTagValue(FBlueprintTags::NativeParentClassPath, ParentClassName))
+		{
+			BPAssetData.GetTagValue(FBlueprintTags::ParentClassPath, ParentClassName);
+		}
+		if (!ParentClassName.IsEmpty())
+		{
+			UObject* Outer = nullptr;
+			ResolveName(Outer, ParentClassName, false, false);
+			ParentClass = FindObject<UClass>(ANY_PACKAGE, *ParentClassName);
+			// UUserWidgets have their own loading section, and we don't want to process any blueprints that don't have UWidget parents
+			if (!ParentClass->IsChildOf(UWidget::StaticClass()) || ParentClass->IsChildOf(UUserWidget::StaticClass()))
+			{
+				continue;
+			}
+		}
+
+		if (!FilterAssetData(BPAssetData))
+		{
+			// If this object isn't currently loaded, add it to the palette view
+			if (BPAssetData.ToSoftObjectPath().ResolveObject() == nullptr)
+			{
+				auto Template = MakeShareable(new FWidgetTemplateClass(BPAssetData, nullptr));
+				AddWidgetTemplate(Template);
+			}
+		}
+	}
+
+	TArray<FAssetData> AllWidgetBPsAssetData;
+	AssetRegistryModule.Get().GetAssetsByClass(UWidgetBlueprint::StaticClass()->GetFName(), AllWidgetBPsAssetData, true);
+
+	FName ActiveWidgetBlueprintName = ActiveWidgetBlueprintClass->ClassGeneratedBy->GetFName();
+	for (FAssetData& WidgetBPAssetData : AllWidgetBPsAssetData)
+	{
+		// Excludes the blueprint you're currently in
+		if (WidgetBPAssetData.AssetName == ActiveWidgetBlueprintName)
+		{
+			continue;
+		}
+
+		if (!FilterAssetData(WidgetBPAssetData))
+		{
+			// Excludes this widget if it is on the hide list
+			bool bIsOnList = false;
+			for (FSoftClassPath Widget : WidgetClassesToHide)
+			{
+				if (Widget.ToString().Find(WidgetBPAssetData.ObjectPath.ToString()) == 0)
+				{
+					bIsOnList = true;
+					break;
+				}
+			}
+			if (bIsOnList)
+			{
+				continue;
+			}
+
+			// If the blueprint generated class was found earlier, pass it to the template
+			TSubclassOf<UUserWidget> WidgetBPClass = nullptr;
+			auto LoadedWidgetBPClass = LoadedWidgetBlueprintClassesByName.Find(WidgetBPAssetData.AssetName);
+			if (LoadedWidgetBPClass)
+			{
+				WidgetBPClass = *LoadedWidgetBPClass;
+			}
+
+			auto Template = MakeShareable(new FWidgetTemplateBlueprintClass(WidgetBPAssetData, WidgetBPClass));
+
+			AddWidgetTemplate(Template);
+		}
+	}
+}
+
+bool FPaletteViewModel::FilterAssetData(FAssetData &InAssetData)
+{
+	// Excludes engine content if user sets it to false
+	if (!GetDefault<UContentBrowserSettings>()->GetDisplayEngineFolder() || !GetDefault<UUMGEditorProjectSettings>()->bShowWidgetsFromEngineContent)
+	{
+		if (InAssetData.PackagePath.ToString().Find(TEXT("/Engine")) == 0)
+		{
+			return true;
+		}
+	}
+
+	// Excludes developer content if user sets it to false
+	if (!GetDefault<UContentBrowserSettings>()->GetDisplayDevelopersFolder() || !GetDefault<UUMGEditorProjectSettings>()->bShowWidgetsFromDeveloperContent)
+	{
+		if (InAssetData.PackagePath.ToString().Find(TEXT("/Game/Developers")) == 0)
+		{
+			return true;
+		}
+	}
+	return false;
+}
+
+void FPaletteViewModel::AddWidgetTemplate(TSharedPtr<FWidgetTemplate> Template)
+{
+	FString Category = Template->GetCategory().ToString();
+
+	// Hide user specific categories
+	TArray<FString> CategoriesToHide = GetDefault<UUMGEditorProjectSettings>()->CategoriesToHide;
+	for (FString CategoryName : CategoriesToHide)
+	{
+		if (Category == CategoryName)
+		{
+			return;
+		}
+	}
+	WidgetTemplateArray& Group = WidgetTemplateCategories.FindOrAdd(Category);
+	Group.Add(Template);
+}
+
+void FPaletteViewModel::OnObjectsReplaced(const TMap<UObject*, UObject*>& ReplacementMap)
+{
+}
+
+void FPaletteViewModel::OnBlueprintReinstanced()
+{
+	bRebuildRequested = true;
+}
+
+void FPaletteViewModel::HandleOnHotReload(bool bWasTriggeredAutomatically)
+{
+	bRebuildRequested = true;
+}
+
+void FPaletteViewModel::HandleOnAssetsDeleted(const TArray<UClass*>& DeletedAssetClasses)
+{
+	for (auto DeletedAssetClass : DeletedAssetClasses)
+	{
+		if (DeletedAssetClass->IsChildOf(UWidgetBlueprint::StaticClass()))
+		{
+			bRebuildRequested = true;
+		}
+	}
+}
+
+#undef LOCTEXT_NAMESPACE
diff --git a/Engine/Source/Editor/UMGEditor/Private/Palette/SPaletteViewModel.h b/Engine/Source/Editor/UMGEditor/Private/Palette/SPaletteViewModel.h
new file mode 100644
--- /dev/null
+++ b/Engine/Source/Editor/UMGEditor/Private/Palette/SPaletteViewModel.h
@@ -0,0 +1,198 @@
+// Copyright 1998-2019 Epic Games, Inc. All Rights Reserved.
+
+#pragma once
+
+#include "CoreMinimal.h"
+#include "WidgetBlueprintEditor.h"
+#include "AssetData.h"
+
+#include "Widgets/Views/STableViewBase.h"
+#include "Widgets/Views/STableRow.h"
+
+class FWidgetTemplate;
+class FWidgetBlueprintEditor;
+class UWidgetBlueprint;
+class SPaletteView;
+
+/** View model for the items in the widget template list */
+class FWidgetViewModel : public TSharedFromThis<FWidgetViewModel>
+{
+public:
+	virtual ~FWidgetViewModel() { }
+
+	virtual FText GetName() const = 0;
+
+	virtual bool IsTemplate() const = 0;
+
+	/** Get the string which should be used for filtering the item. */
+	virtual FString GetFilterString() const = 0;
+
+	virtual TSharedRef<ITableRow> BuildRow(const TSharedRef<STableViewBase>& OwnerTable) = 0;
+
+	virtual void GetChildren(TArray< TSharedPtr<FWidgetViewModel> >& OutChildren)
+	{
+	}
+
+	/** Return true if the widget is a favorite */
+	virtual bool IsFavorite() const { return false; }
+
+	/** Set the favorite flag */
+	virtual void SetFavorite()
+	{
+	}
+};
+
+class FWidgetTemplateViewModel : public FWidgetViewModel
+{
+public:
+	FWidgetTemplateViewModel();
+
+	virtual FText GetName() const override;
+
+	virtual bool IsTemplate() const override;
+
+	virtual FString GetFilterString() const override;
+
+	virtual TSharedRef<ITableRow> BuildRow(const TSharedRef<STableViewBase>& OwnerTable) override;
+
+	FReply OnDraggingWidgetTemplateItem(const FGeometry& MyGeometry, const FPointerEvent& MouseEvent);
+
+	/** Add the widget template to the list of favorites */
+	void AddToFavorites();
+
+	/** Remove the widget template from the list of favorites */
+	void RemoveFromFavorites();
+
+	/** Return true if the widget is a favorite */
+	virtual bool IsFavorite() const override { return bIsFavorite; }
+
+	/** Set the favorite flag */
+	virtual void SetFavorite() override { bIsFavorite = true; }
+
+	TSharedPtr<FWidgetTemplate> Template;
+	FPaletteViewModel* PaletteViewModel;
+private:
+	/** True is the widget is a favorite. It's keep as a state to prevent a search in the favorite list. */
+	bool bIsFavorite;
+};
+
+class FWidgetHeaderViewModel : public FWidgetViewModel
+{
+public:
+	virtual ~FWidgetHeaderViewModel()
+	{
+	}
+
+	virtual FText GetName() const override
+	{
+		return GroupName;
+	}
+
+	virtual bool IsTemplate() const override
+	{
+		return false;
+	}
+
+	virtual FString GetFilterString() const override
+	{
+		// Headers should never be included in filtering to avoid showing a header with all of
+		// it's widgets filtered out, so return an empty filter string.
+		return TEXT("");
+	}
+
+	virtual TSharedRef<ITableRow> BuildRow(const TSharedRef<STableViewBase>& OwnerTable) override;
+
+	virtual void GetChildren(TArray< TSharedPtr<FWidgetViewModel> >& OutChildren) override;
+
+	FText GroupName;
+	TArray< TSharedPtr<FWidgetViewModel> > Children;
+};
+
+class FPaletteViewModel: public TSharedFromThis<FPaletteViewModel>
+{
+public:
+
+	DECLARE_MULTICAST_DELEGATE(FOnUpdating)
+	DECLARE_MULTICAST_DELEGATE(FOnUpdated)
+	DECLARE_MULTICAST_DELEGATE(FOnFavoritesUpdated)
+
+public:
+	FPaletteViewModel(TSharedPtr<FWidgetBlueprintEditor> InBlueprintEditor);
+	~FPaletteViewModel();
+
+	/** Register the View Model to events that should trigger a update of the Palette*/
+	void RegisterToEvents();
+
+	/** Update the view model if needed and returns true if it did. */
+	void Update();
+
+	/** Returns true if the view model needs to be updated */
+	bool NeedUpdate() const { return bRebuildRequested || bRebuildFavorites; }
+	   
+	/** Add the widget template to the list of favorites */
+	void AddToFavorites(const FWidgetTemplateViewModel* WidgetTemplateViewModel);
+
+	/** Remove the widget template to the list of favorites */
+	void RemoveFromFavorites(const FWidgetTemplateViewModel* WidgetTemplateViewModel);
+
+	typedef TArray< TSharedPtr<FWidgetViewModel> > ViewModelsArray;
+	ViewModelsArray& GetFavoritesViewModels() { return WidgetFavoritesViewModels; }
+	ViewModelsArray& GetWidgetViewModels() { return WidgetViewModels; }
+
+	void SetSearchText(const FText& inSearchText) { SearchText = inSearchText; }
+	FText GetSearchText() const { return SearchText; }
+
+	/** Fires before the view model is updated */
+	FOnUpdating OnUpdating;
+
+	/** Fires after the view model is updated */
+	FOnUpdated OnUpdated;
+
+	/** Fires after the list of favorites is updated */
+	FOnFavoritesUpdated OnFavoritesUpdated;
+
+private:
+	FPaletteViewModel() {};
+
+	UWidgetBlueprint* GetBlueprint() const;
+
+	void BuildWidgetList();
+	void BuildWidgetFavoriteList();
+	void BuildClassWidgetList();
+
+	static bool FilterAssetData(FAssetData &BPAssetData);
+
+	void AddWidgetTemplate(TSharedPtr<FWidgetTemplate> Template);
+
+	/** Called when a Blueprint is recompiled and live objects are swapped out for replacements */
+	void OnObjectsReplaced(const TMap<UObject*, UObject*>& ReplacementMap);
+
+	/** Requests a rebuild of the widget list if a widget blueprint was compiled */
+	void OnBlueprintReinstanced();
+
+	/** Requests a rebuild of the widget list */
+	void HandleOnHotReload(bool bWasTriggeredAutomatically);
+
+	/** Requests a rebuild of the widget list if a widget blueprint was deleted */
+	void HandleOnAssetsDeleted(const TArray<UClass*>& DeletedAssetClasses);
+
+	TWeakPtr<class FWidgetBlueprintEditor> BlueprintEditor;
+
+	typedef TArray<TSharedPtr<FWidgetTemplate>> WidgetTemplateArray;
+	TMap<FString, WidgetTemplateArray> WidgetTemplateCategories;
+
+	/** The source root view models for the tree. */
+	ViewModelsArray WidgetViewModels;
+
+	/** The source root view models for the favorites. */
+	ViewModelsArray WidgetFavoritesViewModels;
+	   
+	/** Controls rebuilding the list of spawnable widgets */
+	bool bRebuildRequested;
+
+	/** Set to true when the Palette favorites list needs to be updated*/
+	bool bRebuildFavorites;
+
+	FText SearchText;
+};
+
diff --git a/Engine/Source/Editor/UMGEditor/Private/Settings/WidgetDesignerSettings.cpp b/Engine/Source/Editor/UMGEditor/Private/Settings/WidgetDesignerSettings.cpp
--- a/Engine/Source/Editor/UMGEditor/Private/Settings/WidgetDesignerSettings.cpp
+++ b/Engine/Source/Editor/UMGEditor/Private/Settings/WidgetDesignerSettings.cpp
@@ -4,6 +4,9 @@
 
 UWidgetDesignerSettings::UWidgetDesignerSettings()
 {
+
+	Favorites = CreateDefaultSubobject<UWidgetPaletteFavorites>(TEXT("WidgetPaletteFavorites"));
+
 	CategoryName = TEXT("ContentEditors");
 
 	GridSnapEnabled = true;
diff --git a/Engine/Source/Editor/UMGEditor/Private/TabFactory/PaletteFavoritesTabSummoner.cpp b/Engine/Source/Editor/UMGEditor/Private/TabFactory/PaletteFavoritesTabSummoner.cpp
new file mode 100644
--- /dev/null
+++ b/Engine/Source/Editor/UMGEditor/Private/TabFactory/PaletteFavoritesTabSummoner.cpp
@@ -0,0 +1,32 @@
+// Copyright 1998-2019 Epic Games, Inc. All Rights Reserved.
+
+#include "TabFactory/PaletteFavoritesTabSummoner.h"
+#include "Palette/SPaletteFavoritesView.h"
+#include "UMGStyle.h"
+
+#define LOCTEXT_NAMESPACE "UMG"
+
+const FName FPaletteFavoritesTabSummoner::TabID(TEXT("WidgetTemplatesFavorites"));
+
+FPaletteFavoritesTabSummoner::FPaletteFavoritesTabSummoner(TSharedPtr<class FWidgetBlueprintEditor> InBlueprintEditor)
+		: FWorkflowTabFactory(TabID, InBlueprintEditor)
+		, BlueprintEditor(InBlueprintEditor)
+{
+	TabLabel = LOCTEXT("WidgetTemplatesFavoritesTabLabel", "Palette Favorites");
+	TabIcon = FSlateIcon(FUMGStyle::GetStyleSetName(), "Palette.TabIcon");
+
+	bIsSingleton = true;
+
+	ViewMenuDescription = LOCTEXT("WidgetTemplatesFavorites_ViewMenu_Desc", "Palette Favorites");
+	ViewMenuTooltip = LOCTEXT("WidgetTemplatesFavorites_ViewMenu_ToolTip", "Show the Palette Favorites");
+}
+
+TSharedRef<SWidget> FPaletteFavoritesTabSummoner::CreateTabBody(const FWorkflowTabSpawnInfo& Info) const
+{
+	TSharedPtr<FWidgetBlueprintEditor> BlueprintEditorPtr = StaticCastSharedPtr<FWidgetBlueprintEditor>(BlueprintEditor.Pin());
+
+	return SNew(SPaletteFavoritesView, BlueprintEditorPtr);
+		
+}
+
+#undef LOCTEXT_NAMESPACE 
diff --git a/Engine/Source/Editor/UMGEditor/Private/TabFactory/PaletteFavoritesTabSummoner.h b/Engine/Source/Editor/UMGEditor/Private/TabFactory/PaletteFavoritesTabSummoner.h
new file mode 100644
--- /dev/null
+++ b/Engine/Source/Editor/UMGEditor/Private/TabFactory/PaletteFavoritesTabSummoner.h
@@ -0,0 +1,22 @@
+// Copyright 1998-2019 Epic Games, Inc. All Rights Reserved.
+
+#pragma once
+
+#include "CoreMinimal.h"
+#include "Widgets/SWidget.h"
+#include "WorkflowOrientedApp/WorkflowTabFactory.h"
+#include "WidgetBlueprintEditor.h"
+
+struct FPaletteFavoritesTabSummoner : public FWorkflowTabFactory
+{
+public:
+	static const FName TabID;
+	
+public:
+	FPaletteFavoritesTabSummoner(TSharedPtr<class FWidgetBlueprintEditor> InBlueprintEditor);
+	
+	virtual TSharedRef<SWidget> CreateTabBody(const FWorkflowTabSpawnInfo& Info) const override;
+	
+protected:
+	TWeakPtr<class FWidgetBlueprintEditor> BlueprintEditor;
+};
diff --git a/Engine/Source/Editor/UMGEditor/Private/WidgetBlueprintEditor.cpp b/Engine/Source/Editor/UMGEditor/Private/WidgetBlueprintEditor.cpp
--- a/Engine/Source/Editor/UMGEditor/Private/WidgetBlueprintEditor.cpp
+++ b/Engine/Source/Editor/UMGEditor/Private/WidgetBlueprintEditor.cpp
@@ -50,6 +50,8 @@
 #include "UMGEditorActions.h"
 #include "GameProjectGenerationModule.h"
 
+#include "SPaletteViewModel.h"
+
 #define LOCTEXT_NAMESPACE "UMG"
 
 FWidgetBlueprintEditor::FWidgetBlueprintEditor()
@@ -102,6 +104,10 @@ void FWidgetBlueprintEditor::InitWidgetBlueprintEditor(const EToolkitMode::Type
 	bRespectLocks = GetDefault<UWidgetDesignerSettings>()->bRespectLocks;
 
 	TSharedPtr<FWidgetBlueprintEditor> ThisPtr(SharedThis(this));
+
+	PaletteViewModel = MakeShareable(new FPaletteViewModel(ThisPtr));
+	PaletteViewModel->RegisterToEvents();
+
 	WidgetToolbar = MakeShareable(new FWidgetBlueprintEditorToolbar(ThisPtr));
 
 	BindToolkitCommands();
@@ -552,6 +558,12 @@ void FWidgetBlueprintEditor::Tick(float DeltaTime)
 		bPreviewInvalidated = false;
 		RefreshPreview();
 	}
+
+	// Updat the palette view model.
+	if (PaletteViewModel->NeedUpdate())
+	{
+		PaletteViewModel->Update();
+	}
 }
 
 static bool MigratePropertyValue(UObject* SourceObject, UObject* DestinationObject, FEditPropertyChain::TDoubleLinkedListNode* PropertyChainNode, UProperty* MemberProperty, bool bIsModify)
diff --git a/Engine/Source/Editor/UMGEditor/Public/Settings/WidgetDesignerSettings.h b/Engine/Source/Editor/UMGEditor/Public/Settings/WidgetDesignerSettings.h
--- a/Engine/Source/Editor/UMGEditor/Public/Settings/WidgetDesignerSettings.h
+++ b/Engine/Source/Editor/UMGEditor/Public/Settings/WidgetDesignerSettings.h
@@ -6,6 +6,7 @@
 #include "UObject/ObjectMacros.h"
 #include "UObject/Object.h"
 #include "Engine/DeveloperSettings.h"
+#include "WidgetPaletteFavorites.h"
 #include "WidgetDesignerSettings.generated.h"
 
 /**
@@ -61,4 +62,11 @@ public:
 	 */
 	UPROPERTY(EditAnywhere, config, Category = Interaction)
 	bool bRespectLocks;
+	
+	/**
+	 * List of Favorites widgets used to populate the Favorites Palette
+	 */
+	UPROPERTY()
+	UWidgetPaletteFavorites* Favorites;
+
 };
diff --git a/Engine/Source/Editor/UMGEditor/Public/WidgetBlueprintEditor.h b/Engine/Source/Editor/UMGEditor/Public/WidgetBlueprintEditor.h
--- a/Engine/Source/Editor/UMGEditor/Public/WidgetBlueprintEditor.h
+++ b/Engine/Source/Editor/UMGEditor/Public/WidgetBlueprintEditor.h
@@ -22,6 +22,7 @@ class STextBlock;
 class UPanelSlot;
 class UWidgetAnimation;
 class UWidgetBlueprint;
+class FPaletteViewModel;
 
 struct FNamedSlotSelection
 {
@@ -172,6 +173,8 @@ public:
 	bool GetIsRespectingLocks() const;
 	void SetIsRespectingLocks(bool Value);
 
+	TSharedPtr<FPaletteViewModel> GetPaletteViewModel() { return PaletteViewModel; };
+
 public:
 	/** Fires whenever a new widget is being hovered over */
 	FOnHoveredWidgetSet OnHoveredWidgetSet;
@@ -361,4 +364,7 @@ private:
 
 	/** When true the animation data in the generated class should be replaced with the current animation data. */
 	bool bRefreshGeneratedClassAnimations;
+
+	/** ViewModel used by the Palette and Palette Favorite Views */
+	TSharedPtr<FPaletteViewModel> PaletteViewModel;
 };
