PR #5525: Exposing FDesignerExtension in UMGEditor to public interface (Contributed by [REDACTED_PERSON])

diff --git a/Engine/Source/Editor/UMGEditor/Private/Designer/SDesignerView.cpp b/Engine/Source/Editor/UMGEditor/Private/Designer/SDesignerView.cpp
--- a/Engine/Source/Editor/UMGEditor/Private/Designer/SDesignerView.cpp
+++ b/Engine/Source/Editor/UMGEditor/Private/Designer/SDesignerView.cpp
@@ -70,6 +70,7 @@
 #include "DeviceProfiles/DeviceProfile.h"
 #include "DeviceProfiles/DeviceProfileManager.h"
 #include "Engine/DPICustomScalingRule.h"
+#include "UMGEditorModule.h"
 
 #define LOCTEXT_NAMESPACE "UMG"
 
@@ -355,6 +356,15 @@ void SDesignerView::Construct(const FArguments& InArgs, TSharedPtr<FWidgetBluepr
 	Register(MakeShareable(new FUniformGridSlotExtension()));
 	Register(MakeShareable(new FGridSlotExtension()));
 
+	//Register External Extensions
+	IUMGEditorModule& UMGEditorInterface = FModuleManager::GetModuleChecked<IUMGEditorModule>("UMGEditor");
+
+	TSharedPtr<FDesignerExtensibilityManager>DesignerExtensibilityManager = UMGEditorInterface.GetDesignerExtensibilityManager();
+	for (const auto& Extension : DesignerExtensibilityManager->GetExternalDesignerExtensions())
+	{
+		Register(Extension);
+	}
+
 	GEditor->OnBlueprintReinstanced().AddRaw(this, &SDesignerView::OnPreviewNeedsRecreation);
 
 	BindCommands();
diff --git a/Engine/Source/Editor/UMGEditor/Private/UMGEditorModule.cpp b/Engine/Source/Editor/UMGEditor/Private/UMGEditorModule.cpp
--- a/Engine/Source/Editor/UMGEditor/Private/UMGEditorModule.cpp
+++ b/Engine/Source/Editor/UMGEditor/Private/UMGEditorModule.cpp
@@ -61,6 +61,7 @@ public:
 
 		MenuExtensibilityManager = MakeShareable(new FExtensibilityManager());
 		ToolBarExtensibilityManager = MakeShareable(new FExtensibilityManager());
+		DesignerExtensibilityManager = MakeShareable(new FDesignerExtensibilityManager());
 
 		// Register widget blueprint compiler we do this no matter what.
 		IKismetCompilerInterface& KismetCompilerModule = FModuleManager::LoadModuleChecked<IKismetCompilerInterface>("KismetCompiler");
@@ -133,6 +134,7 @@ public:
 	/** Gets the extensibility managers for outside entities to extend gui page editor's menus and toolbars */
 	virtual TSharedPtr<FExtensibilityManager> GetMenuExtensibilityManager() override { return MenuExtensibilityManager; }
 	virtual TSharedPtr<FExtensibilityManager> GetToolBarExtensibilityManager() override { return ToolBarExtensibilityManager; }
+	virtual TSharedPtr<FDesignerExtensibilityManager> GetDesignerExtensibilityManager() override { return DesignerExtensibilityManager; }
 
 	/** Register settings objects. */
 	void RegisterSettings()
@@ -185,6 +187,7 @@ private:
 private:
 	TSharedPtr<FExtensibilityManager> MenuExtensibilityManager;
 	TSharedPtr<FExtensibilityManager> ToolBarExtensibilityManager;
+	TSharedPtr<FDesignerExtensibilityManager> DesignerExtensibilityManager;
 
 	FDelegateHandle SequenceEditorHandle;
 	FDelegateHandle MarginTrackEditorCreateTrackEditorHandle;
diff --git a/Engine/Source/Editor/UMGEditor/Public/DesignerExtension.h b/Engine/Source/Editor/UMGEditor/Public/DesignerExtension.h
--- a/Engine/Source/Editor/UMGEditor/Public/DesignerExtension.h
+++ b/Engine/Source/Editor/UMGEditor/Public/DesignerExtension.h
@@ -109,20 +109,23 @@ public:
 	/** Initializes the designer extension, this is called the first time a designer extension is registered */
 	virtual void Initialize(IUMGDesigner* InDesigner, UWidgetBlueprint* InBlueprint);
 
+	/** Returns true if the designer extension can extend the current selection. */
 	virtual bool CanExtendSelection(const TArray< FWidgetReference >& Selection) const
 	{
 		return false;
 	}
 
-	/** Called every time the selection in the designer changes. */
+	/** Called every time a element the designer can extend is selected. */
 	virtual void ExtendSelection(const TArray< FWidgetReference >& Selection, TArray< TSharedRef<FDesignerSurfaceElement> >& SurfaceElements)
 	{
 	}
 
+	/** Called each frames to tick the extension. */
 	virtual void Tick( const FGeometry& AllottedGeometry, const double InCurrentTime, const float InDeltaTime )
 	{
 	}
 
+	/** Called to paint the extension. */
 	virtual void Paint(const TSet< FWidgetReference >& Selection, const FGeometry& AllottedGeometry, const FSlateRect& MyCullingRect, FSlateWindowElementList& OutDrawElements, int32 LayerId) const
 	{
 	}
@@ -130,6 +133,11 @@ public:
 	/** Gets the ID identifying this extension. */
 	FName GetExtensionId() const;
 
+	bool operator==(const FDesignerExtension& Other) const
+	{
+		return ExtensionId.IsEqual(Other.GetExtensionId());
+	}
+
 protected:
 	void BeginTransaction(const FText& SessionName);
 
diff --git a/Engine/Source/Editor/UMGEditor/Public/IHasDesignerExtensibility.h b/Engine/Source/Editor/UMGEditor/Public/IHasDesignerExtensibility.h
new file mode 100644
--- /dev/null
+++ b/Engine/Source/Editor/UMGEditor/Public/IHasDesignerExtensibility.h
@@ -0,0 +1,39 @@
+// Copyright 1998-2019 Epic Games, Inc. All Rights Reserved.
+
+#pragma once
+
+#include "CoreMinimal.h"
+#include "DesignerExtension.h"
+
+/**
+ * Designer Extensibility Manager keep a series of Designer Extensions. See FDesignerExtension class for more information.
+ */
+class UMGEDITOR_API FDesignerExtensibilityManager
+{
+public:
+	void AddDesignerExtension(const TSharedRef<FDesignerExtension>& Extension)
+	{
+		ExternalExtensions.AddUnique(Extension);
+	}
+
+	void RemoveDesignerExtension(const TSharedRef<FDesignerExtension>& Extension)
+	{
+		ExternalExtensions.Remove(Extension);
+	}
+
+	const TArray<TSharedRef<FDesignerExtension>>& GetExternalDesignerExtensions() const
+	{
+		return ExternalExtensions;
+	}
+
+private:
+	TArray<TSharedRef<FDesignerExtension>> ExternalExtensions;
+};
+
+/** Indicates that a class has a designer that is extensible */
+class IHasDesignerExtensibility
+{
+public:
+	virtual TSharedPtr<FDesignerExtensibilityManager> GetDesignerExtensibilityManager() = 0;
+};
+
diff --git a/Engine/Source/Editor/UMGEditor/Public/UMGEditorModule.h b/Engine/Source/Editor/UMGEditor/Public/UMGEditorModule.h
--- a/Engine/Source/Editor/UMGEditor/Public/UMGEditorModule.h
+++ b/Engine/Source/Editor/UMGEditor/Public/UMGEditorModule.h
@@ -5,6 +5,7 @@
 #include "CoreMinimal.h"
 #include "Modules/ModuleInterface.h"
 #include "Toolkits/AssetEditorToolkit.h"
+#include "IHasDesignerExtensibility.h"
 
 extern const FName UMGEditorAppIdentifier;
 
@@ -12,7 +13,7 @@ class FUMGEditor;
 class FWidgetBlueprintCompiler;
 
 /** The public interface of the UMG editor module. */
-class IUMGEditorModule : public IModuleInterface, public IHasMenuExtensibility, public IHasToolBarExtensibility
+class IUMGEditorModule : public IModuleInterface, public IHasMenuExtensibility, public IHasToolBarExtensibility, public IHasDesignerExtensibility
 {
 public:
 	virtual FWidgetBlueprintCompiler* GetRegisteredCompiler() = 0;
