UMG hierarchy search includes widget class name (and display name) on top of widget name.
TODO: There is no option to include/exclude the class name in the search (it's always included). The JIRA mentioned an option, but Nick Darnell suggested to not have one.
TODO: When the widget name matches we highlight the substring in green, but there is no visual feedback when the class name is the reason for the match. But then again we don't highlight the widget name when the search string is tokenized by whitespaces (for example "Editable Text Box" rather than "EditableTextBox").

diff --git a/Engine/Source/Editor/UMGEditor/Private/Hierarchy/SHierarchyView.cpp b/Engine/Source/Editor/UMGEditor/Private/Hierarchy/SHierarchyView.cpp
--- a/Engine/Source/Editor/UMGEditor/Private/Hierarchy/SHierarchyView.cpp
+++ b/Engine/Source/Editor/UMGEditor/Private/Hierarchy/SHierarchyView.cpp
@@ -32,7 +32,7 @@ void SHierarchyView::Construct(const FArguments& InArgs, TSharedPtr<FWidgetBluep
 	GEditor->OnObjectsReplaced().AddRaw(this, &SHierarchyView::OnObjectsReplaced);
 
 	// Create the filter for searching in the tree
-	SearchBoxWidgetFilter = MakeShareable(new WidgetTextFilter(WidgetTextFilter::FItemToStringArray::CreateSP(this, &SHierarchyView::TransformWidgetToString)));
+	SearchBoxWidgetFilter = MakeShareable(new WidgetTextFilter(WidgetTextFilter::FItemToStringArray::CreateSP(this, &SHierarchyView::TransformWidgetToSearchableStrings)));
 
 	UWidgetBlueprint* Blueprint = GetBlueprint();
 	Blueprint->OnChanged().AddRaw(this, &SHierarchyView::OnBlueprintChanged);
@@ -176,9 +176,9 @@ bool SHierarchyView::CanRename() const
 	return SelectedItems.Num() == 1 && SelectedItems[0]->CanRename();
 }
 
-void SHierarchyView::TransformWidgetToString(TSharedPtr<FHierarchyModel> Item, OUT TArray< FString >& Array)
+void SHierarchyView::TransformWidgetToSearchableStrings(TSharedPtr<FHierarchyModel> Item, TArray<FString>& OutStrings)
 {
-	Array.Add( Item->GetText().ToString() );
+	Item->GetSearchableStrings(OutStrings);
 }
 
 void SHierarchyView::OnSearchChanged(const FText& InFilterText)
diff --git a/Engine/Source/Editor/UMGEditor/Private/Hierarchy/SHierarchyView.h b/Engine/Source/Editor/UMGEditor/Private/Hierarchy/SHierarchyView.h
--- a/Engine/Source/Editor/UMGEditor/Private/Hierarchy/SHierarchyView.h
+++ b/Engine/Source/Editor/UMGEditor/Private/Hierarchy/SHierarchyView.h
@@ -89,7 +89,7 @@ private:
 	FText GetSearchText() const;
 
 	/** Transforms the widget into a searchable string */
-	void TransformWidgetToString(TSharedPtr<FHierarchyModel> Widget, OUT TArray< FString >& Array);
+	void TransformWidgetToSearchableStrings(TSharedPtr<FHierarchyModel> Widget, TArray<FString>& OutStrings);
 
 	/** Called when a Blueprint is recompiled and live objects are swapped out for replacements */
 	void OnObjectsReplaced(const TMap<UObject*, UObject*>& ReplacementMap);
diff --git a/Engine/Source/Editor/UMGEditor/Private/Hierarchy/SHierarchyViewItem.cpp b/Engine/Source/Editor/UMGEditor/Private/Hierarchy/SHierarchyViewItem.cpp
--- a/Engine/Source/Editor/UMGEditor/Private/Hierarchy/SHierarchyViewItem.cpp
+++ b/Engine/Source/Editor/UMGEditor/Private/Hierarchy/SHierarchyViewItem.cpp
@@ -1093,6 +1093,18 @@ FText FHierarchyWidget::GetLabelToolTipText() const
 	return FText::GetEmpty();
 }
 
+void FHierarchyWidget::GetSearchableStrings(TArray<FString>& OutStrings) const
+{
+	FHierarchyModel::GetSearchableStrings(OutStrings);
+
+	UWidget* WidgetTemplate = Item.GetTemplate();
+	if (WidgetTemplate && !WidgetTemplate->IsGeneratedName())
+	{
+		OutStrings.Add(WidgetTemplate->GetClass()->GetName());
+		OutStrings.Add(WidgetTemplate->GetClass()->GetDisplayNameText().ToString());
+	}
+}
+
 const FSlateBrush* FHierarchyWidget::GetImage() const
 {
 	if (Item.GetTemplate())
diff --git a/Engine/Source/Editor/UMGEditor/Private/Hierarchy/SHierarchyViewItem.h b/Engine/Source/Editor/UMGEditor/Private/Hierarchy/SHierarchyViewItem.h
--- a/Engine/Source/Editor/UMGEditor/Private/Hierarchy/SHierarchyViewItem.h
+++ b/Engine/Source/Editor/UMGEditor/Private/Hierarchy/SHierarchyViewItem.h
@@ -32,6 +32,9 @@ public:
 
 	/** @return The tooltip for the tree item label */
 	virtual FText GetLabelToolTipText() const { return FText::GetEmpty(); }
+	
+	/** @param OutStrings - Returns an array of strings used to search for the tree item. */
+	virtual void GetSearchableStrings(TArray<FString>& OutStrings) const { OutStrings.Add(GetText().ToString()); }
 
 	virtual const FSlateBrush* GetImage() const = 0;
 
@@ -196,6 +199,8 @@ public:
 	virtual FText GetText() const override;
 	virtual FText GetImageToolTipText() const override;
 	virtual FText GetLabelToolTipText() const override;
+	
+	virtual void GetSearchableStrings(TArray<FString>& OutStrings) const override;
 
 	virtual const FSlateBrush* GetImage() const override;
 
