Search by class internal and display names in UMG editor Palette view.

diff --git a/Engine/Source/Editor/UMGEditor/Private/Hierarchy/SHierarchyView.cpp b/Engine/Source/Editor/UMGEditor/Private/Hierarchy/SHierarchyView.cpp
--- a/Engine/Source/Editor/UMGEditor/Private/Hierarchy/SHierarchyView.cpp
+++ b/Engine/Source/Editor/UMGEditor/Private/Hierarchy/SHierarchyView.cpp
@@ -32,7 +32,7 @@ void SHierarchyView::Construct(const FArguments& InArgs, TSharedPtr<FWidgetBluep
 	GEditor->OnObjectsReplaced().AddRaw(this, &SHierarchyView::OnObjectsReplaced);
 
 	// Create the filter for searching in the tree
-	SearchBoxWidgetFilter = MakeShareable(new WidgetTextFilter(WidgetTextFilter::FItemToStringArray::CreateSP(this, &SHierarchyView::TransformWidgetToSearchableStrings)));
+	SearchBoxWidgetFilter = MakeShareable(new WidgetTextFilter(WidgetTextFilter::FItemToStringArray::CreateSP(this, &SHierarchyView::GetWidgetFilterStrings)));
 
 	UWidgetBlueprint* Blueprint = GetBlueprint();
 	Blueprint->OnChanged().AddRaw(this, &SHierarchyView::OnBlueprintChanged);
@@ -176,9 +176,9 @@ bool SHierarchyView::CanRename() const
 	return SelectedItems.Num() == 1 && SelectedItems[0]->CanRename();
 }
 
-void SHierarchyView::TransformWidgetToSearchableStrings(TSharedPtr<FHierarchyModel> Item, TArray<FString>& OutStrings)
+void SHierarchyView::GetWidgetFilterStrings(TSharedPtr<FHierarchyModel> Item, TArray<FString>& OutStrings)
 {
-	Item->GetSearchableStrings(OutStrings);
+	Item->GetFilterStrings(OutStrings);
 }
 
 void SHierarchyView::OnSearchChanged(const FText& InFilterText)
diff --git a/Engine/Source/Editor/UMGEditor/Private/Hierarchy/SHierarchyView.h b/Engine/Source/Editor/UMGEditor/Private/Hierarchy/SHierarchyView.h
--- a/Engine/Source/Editor/UMGEditor/Private/Hierarchy/SHierarchyView.h
+++ b/Engine/Source/Editor/UMGEditor/Private/Hierarchy/SHierarchyView.h
@@ -88,8 +88,8 @@ private:
 	/** Gets the search text currently being used to filter the list, also used to highlight text */
 	FText GetSearchText() const;
 
-	/** Transforms the widget into a searchable string */
-	void TransformWidgetToSearchableStrings(TSharedPtr<FHierarchyModel> Widget, TArray<FString>& OutStrings);
+	/**  Gets an array of strings used for filtering/searching the specified widget. */
+	void GetWidgetFilterStrings(TSharedPtr<FHierarchyModel> Widget, TArray<FString>& OutStrings);
 
 	/** Called when a Blueprint is recompiled and live objects are swapped out for replacements */
 	void OnObjectsReplaced(const TMap<UObject*, UObject*>& ReplacementMap);
diff --git a/Engine/Source/Editor/UMGEditor/Private/Hierarchy/SHierarchyViewItem.cpp b/Engine/Source/Editor/UMGEditor/Private/Hierarchy/SHierarchyViewItem.cpp
--- a/Engine/Source/Editor/UMGEditor/Private/Hierarchy/SHierarchyViewItem.cpp
+++ b/Engine/Source/Editor/UMGEditor/Private/Hierarchy/SHierarchyViewItem.cpp
@@ -1093,9 +1093,9 @@ FText FHierarchyWidget::GetLabelToolTipText() const
 	return FText::GetEmpty();
 }
 
-void FHierarchyWidget::GetSearchableStrings(TArray<FString>& OutStrings) const
+void FHierarchyWidget::GetFilterStrings(TArray<FString>& OutStrings) const
 {
-	FHierarchyModel::GetSearchableStrings(OutStrings);
+	FHierarchyModel::GetFilterStrings(OutStrings);
 
 	UWidget* WidgetTemplate = Item.GetTemplate();
 	if (WidgetTemplate && !WidgetTemplate->IsGeneratedName())
diff --git a/Engine/Source/Editor/UMGEditor/Private/Hierarchy/SHierarchyViewItem.h b/Engine/Source/Editor/UMGEditor/Private/Hierarchy/SHierarchyViewItem.h
--- a/Engine/Source/Editor/UMGEditor/Private/Hierarchy/SHierarchyViewItem.h
+++ b/Engine/Source/Editor/UMGEditor/Private/Hierarchy/SHierarchyViewItem.h
@@ -33,8 +33,8 @@ public:
 	/** @return The tooltip for the tree item label */
 	virtual FText GetLabelToolTipText() const { return FText::GetEmpty(); }
 	
-	/** @param OutStrings - Returns an array of strings used to search for the tree item. */
-	virtual void GetSearchableStrings(TArray<FString>& OutStrings) const { OutStrings.Add(GetText().ToString()); }
+	/** @param OutStrings - Returns an array of strings used for filtering/searching this item. */
+	virtual void GetFilterStrings(TArray<FString>& OutStrings) const { OutStrings.Add(GetText().ToString()); }
 
 	virtual const FSlateBrush* GetImage() const = 0;
 
@@ -200,7 +200,7 @@ public:
 	virtual FText GetImageToolTipText() const override;
 	virtual FText GetLabelToolTipText() const override;
 	
-	virtual void GetSearchableStrings(TArray<FString>& OutStrings) const override;
+	virtual void GetFilterStrings(TArray<FString>& OutStrings) const override;
 
 	virtual const FSlateBrush* GetImage() const override;
 
diff --git a/Engine/Source/Editor/UMGEditor/Private/Palette/SPaletteView.cpp b/Engine/Source/Editor/UMGEditor/Private/Palette/SPaletteView.cpp
--- a/Engine/Source/Editor/UMGEditor/Private/Palette/SPaletteView.cpp
+++ b/Engine/Source/Editor/UMGEditor/Private/Palette/SPaletteView.cpp
@@ -127,7 +127,7 @@ void SPaletteView::Construct(const FArguments& InArgs, TSharedPtr<FWidgetBluepri
 	PaletteViewModel->OnUpdated.AddRaw(this, &SPaletteView::OnViewModelUpdated);
 
 	WidgetFilter = MakeShareable(new WidgetViewModelTextFilter(
-		WidgetViewModelTextFilter::FItemToStringArray::CreateSP(this, &SPaletteView::TransformWidgetViewModelToString)));
+		WidgetViewModelTextFilter::FItemToStringArray::CreateSP(this, &SPaletteView::GetWidgetFilterStrings)));
 
 	FilterHandler = MakeShareable(new PaletteFilterHandler());
 	FilterHandler->SetFilter(WidgetFilter.Get());
@@ -321,9 +321,9 @@ void SPaletteView::Tick(const FGeometry& AllottedGeometry, const double InCurren
 	}
 }
 
-void SPaletteView::TransformWidgetViewModelToString(TSharedPtr<FWidgetViewModel> WidgetViewModel, OUT TArray< FString >& Array)
+void SPaletteView::GetWidgetFilterStrings(TSharedPtr<FWidgetViewModel> WidgetViewModel, TArray<FString>& OutStrings)
 {
-	Array.Add(WidgetViewModel->GetFilterString());
+	WidgetViewModel->GetFilterStrings(OutStrings);
 }
 
 #undef LOCTEXT_NAMESPACE
diff --git a/Engine/Source/Editor/UMGEditor/Private/Palette/SPaletteView.h b/Engine/Source/Editor/UMGEditor/Private/Palette/SPaletteView.h
--- a/Engine/Source/Editor/UMGEditor/Private/Palette/SPaletteView.h
+++ b/Engine/Source/Editor/UMGEditor/Private/Palette/SPaletteView.h
@@ -108,9 +108,8 @@ private:
 	/** Save the expansion state for the TreeView */
 	void SaveItemExpansion();
 
-	// Search functions
-	/** Transforms the widget view model into a searchable string. */
-	void TransformWidgetViewModelToString(TSharedPtr<FWidgetViewModel> WidgetViewModel, OUT TArray< FString >& Array);
+	/** Gets an array of strings used for filtering/searching the specified widget. */
+	void GetWidgetFilterStrings(TSharedPtr<FWidgetViewModel> WidgetViewModel, TArray<FString>& OutStrings);
 
 	TWeakPtr<class FWidgetBlueprintEditor> BlueprintEditor;
 	TSharedPtr<FPaletteViewModel> PaletteViewModel;
diff --git a/Engine/Source/Editor/UMGEditor/Private/Palette/SPaletteViewModel.cpp b/Engine/Source/Editor/UMGEditor/Private/Palette/SPaletteViewModel.cpp
--- a/Engine/Source/Editor/UMGEditor/Private/Palette/SPaletteViewModel.cpp
+++ b/Engine/Source/Editor/UMGEditor/Private/Palette/SPaletteViewModel.cpp
@@ -44,9 +44,9 @@ bool FWidgetTemplateViewModel::IsTemplate() const
 	return true;
 }
 
-FString FWidgetTemplateViewModel::GetFilterString() const
+void FWidgetTemplateViewModel::GetFilterStrings(TArray<FString>& OutStrings) const
 {
-	return Template->Name.ToString();
+	Template->GetFilterStrings(OutStrings);
 }
 
 TSharedRef<ITableRow> FWidgetTemplateViewModel::BuildRow(const TSharedRef<STableViewBase>& OwnerTable)
diff --git a/Engine/Source/Editor/UMGEditor/Private/Palette/SPaletteViewModel.h b/Engine/Source/Editor/UMGEditor/Private/Palette/SPaletteViewModel.h
--- a/Engine/Source/Editor/UMGEditor/Private/Palette/SPaletteViewModel.h
+++ b/Engine/Source/Editor/UMGEditor/Private/Palette/SPaletteViewModel.h
@@ -24,8 +24,8 @@ public:
 
 	virtual bool IsTemplate() const = 0;
 
-	/** Get the string which should be used for filtering the item. */
-	virtual FString GetFilterString() const = 0;
+	/** @param OutStrings - Returns an array of strings used for filtering/searching this item. */
+	virtual void GetFilterStrings(TArray<FString>& OutStrings) const = 0;
 
 	virtual TSharedRef<ITableRow> BuildRow(const TSharedRef<STableViewBase>& OwnerTable) = 0;
 
@@ -51,7 +51,7 @@ public:
 
 	virtual bool IsTemplate() const override;
 
-	virtual FString GetFilterString() const override;
+	virtual void GetFilterStrings(TArray<FString>& OutStrings) const override;
 
 	virtual TSharedRef<ITableRow> BuildRow(const TSharedRef<STableViewBase>& OwnerTable) override;
 
@@ -93,11 +93,10 @@ public:
 		return false;
 	}
 
-	virtual FString GetFilterString() const override
+	virtual void GetFilterStrings(TArray<FString>& OutStrings) const override
 	{
 		// Headers should never be included in filtering to avoid showing a header with all of
 		// it's widgets filtered out, so return an empty filter string.
-		return TEXT("");
 	}
 
 	virtual TSharedRef<ITableRow> BuildRow(const TSharedRef<STableViewBase>& OwnerTable) override;
diff --git a/Engine/Source/Editor/UMGEditor/Private/Templates/WidgetTemplateBlueprintClass.cpp b/Engine/Source/Editor/UMGEditor/Private/Templates/WidgetTemplateBlueprintClass.cpp
--- a/Engine/Source/Editor/UMGEditor/Private/Templates/WidgetTemplateBlueprintClass.cpp
+++ b/Engine/Source/Editor/UMGEditor/Private/Templates/WidgetTemplateBlueprintClass.cpp
@@ -87,11 +87,6 @@ FReply FWidgetTemplateBlueprintClass::OnDoubleClicked()
 	return FReply::Handled();
 }
 
-FAssetData FWidgetTemplateBlueprintClass::GetWidgetAssetData()
-{
-	return WidgetAssetData;
-}
-
 bool FWidgetTemplateBlueprintClass::Supports(UClass* InClass)
 {
 	return InClass != nullptr && InClass->IsChildOf(UWidgetBlueprint::StaticClass());
diff --git a/Engine/Source/Editor/UMGEditor/Private/Templates/WidgetTemplateBlueprintClass.h b/Engine/Source/Editor/UMGEditor/Private/Templates/WidgetTemplateBlueprintClass.h
--- a/Engine/Source/Editor/UMGEditor/Private/Templates/WidgetTemplateBlueprintClass.h
+++ b/Engine/Source/Editor/UMGEditor/Private/Templates/WidgetTemplateBlueprintClass.h
@@ -44,9 +44,6 @@ public:
 	/** Opens the widget blueprint for edit */
 	virtual FReply OnDoubleClicked() override;
 
-	/** Gets the asset data for this widget blueprint */
-	FAssetData GetWidgetAssetData();
-
 	/** Returns true if the supplied class is supported by this template */
 	static bool Supports(UClass* InClass);
 
diff --git a/Engine/Source/Editor/UMGEditor/Private/Templates/WidgetTemplateClass.cpp b/Engine/Source/Editor/UMGEditor/Private/Templates/WidgetTemplateClass.cpp
--- a/Engine/Source/Editor/UMGEditor/Private/Templates/WidgetTemplateClass.cpp
+++ b/Engine/Source/Editor/UMGEditor/Private/Templates/WidgetTemplateClass.cpp
@@ -113,6 +113,19 @@ TSharedRef<IToolTip> FWidgetTemplateClass::GetToolTip() const
 	}
 }
 
+void FWidgetTemplateClass::GetFilterStrings(TArray<FString>& OutStrings) const
+{
+	FWidgetTemplate::GetFilterStrings(OutStrings);
+	if (WidgetClass.IsValid())
+	{
+		OutStrings.Add(WidgetClass->GetName());
+	}
+	if (WidgetAssetData.IsValid())
+	{
+		OutStrings.Add(WidgetAssetData.AssetName.ToString());
+	}
+}
+
 void FWidgetTemplateClass::OnObjectsReplaced(const TMap<UObject*, UObject*>& ReplacementMap)
 {
 	UObject* const* NewObject = ReplacementMap.Find(WidgetClass.Get());
diff --git a/Engine/Source/Editor/UMGEditor/Private/Templates/WidgetTemplateClass.h b/Engine/Source/Editor/UMGEditor/Private/Templates/WidgetTemplateClass.h
--- a/Engine/Source/Editor/UMGEditor/Private/Templates/WidgetTemplateClass.h
+++ b/Engine/Source/Editor/UMGEditor/Private/Templates/WidgetTemplateClass.h
@@ -37,9 +37,15 @@ public:
 	/** Gets the tooltip widget for this palette item. */
 	virtual TSharedRef<IToolTip> GetToolTip() const override;
 
-	/** Gets the WidgetClass */
+	/** @param OutStrings - Returns an array of strings used for filtering/searching this widget template. */
+	virtual void GetFilterStrings(TArray<FString>& OutStrings) const override;
+
+	/** Gets the WidgetClass which might be null. */
 	TWeakObjectPtr<UClass> GetWidgetClass() const { return WidgetClass; }
 
+	/** Returns the asset data for this widget which might be invalid. */
+	FAssetData GetWidgetAssetData() { return WidgetAssetData; }
+
 protected:
 	/** Creates a widget template class without any class reference */
 	FWidgetTemplateClass();
diff --git a/Engine/Source/Editor/UMGEditor/Private/Templates/WidgetTemplateImageClass.cpp b/Engine/Source/Editor/UMGEditor/Private/Templates/WidgetTemplateImageClass.cpp
--- a/Engine/Source/Editor/UMGEditor/Private/Templates/WidgetTemplateImageClass.cpp
+++ b/Engine/Source/Editor/UMGEditor/Private/Templates/WidgetTemplateImageClass.cpp
@@ -6,8 +6,7 @@
 #include "Engine/Texture.h"
 
 FWidgetTemplateImageClass::FWidgetTemplateImageClass(const FAssetData& InAssetData)
-	: FWidgetTemplateClass(UImage::StaticClass())
-	, WidgetAssetData(InAssetData)
+	: FWidgetTemplateClass(InAssetData, UImage::StaticClass())
 {
 }
 
@@ -30,11 +29,6 @@ UWidget* FWidgetTemplateImageClass::Create(UWidgetTree* WidgetTree)
 	return Widget;
 }
 
-FAssetData FWidgetTemplateImageClass::GetWidgetAssetData()
-{
-	return WidgetAssetData;
-}
-
 bool FWidgetTemplateImageClass::Supports(UClass* InClass)
 {
 	static const UClass* SlateTextureAtlasInterface = FindObject<UClass>(ANY_PACKAGE, TEXT("SlateTextureAtlasInterface"));
diff --git a/Engine/Source/Editor/UMGEditor/Private/Templates/WidgetTemplateImageClass.h b/Engine/Source/Editor/UMGEditor/Private/Templates/WidgetTemplateImageClass.h
--- a/Engine/Source/Editor/UMGEditor/Private/Templates/WidgetTemplateImageClass.h
+++ b/Engine/Source/Editor/UMGEditor/Private/Templates/WidgetTemplateImageClass.h
@@ -27,12 +27,6 @@ public:
 	/** Creates an instance of the widget for the tree */
 	virtual UWidget* Create(UWidgetTree* WidgetTree) override;
 
-	/** Returns the asset data for this widget */
-	FAssetData GetWidgetAssetData();
-
 	/** Returns true if the supplied class is supported by this template */
 	static bool Supports(UClass* InClass);
-
-private:
-	FAssetData WidgetAssetData;
 };
diff --git a/Engine/Source/Editor/UMGEditor/Public/WidgetTemplate.h b/Engine/Source/Editor/UMGEditor/Public/WidgetTemplate.h
--- a/Engine/Source/Editor/UMGEditor/Public/WidgetTemplate.h
+++ b/Engine/Source/Editor/UMGEditor/Public/WidgetTemplate.h
@@ -26,16 +26,19 @@ public:
 	/** Constructs the widget template. */
 	virtual UWidget* Create(class UWidgetTree* Tree) = 0;
 
-	/** Gets the icon to display in the template palate for this template. */
+	/** Gets the icon to display for this widget template. */
 	virtual const FSlateBrush* GetIcon() const
 	{
 		static FSlateNoResource NullBrush;
 		return &NullBrush;
 	}
 
-	/** Gets tooltip widget for this palette item. */
+	/** Gets tooltip widget for this widget template. */
 	virtual TSharedRef<IToolTip> GetToolTip() const = 0;
 
+	/** @param OutStrings - Returns an array of strings used for filtering/searching this widget template. */
+	virtual void GetFilterStrings(TArray<FString>& OutStrings) const { OutStrings.Add(Name.ToString()); }
+
 	/** The the action to perform when the template item is double clicked */
 	virtual FReply OnDoubleClicked() { return FReply::Unhandled(); }
 
