Collapse/expansion state of UMG hierarchy items is saved when filtering starts and restored when filtering ends so that mutations during filtering don't affect the state the user has explicitly set.

diff --git a/Engine/Source/Editor/UMGEditor/Private/Hierarchy/SHierarchyView.cpp b/Engine/Source/Editor/UMGEditor/Private/Hierarchy/SHierarchyView.cpp
--- a/Engine/Source/Editor/UMGEditor/Private/Hierarchy/SHierarchyView.cpp
+++ b/Engine/Source/Editor/UMGEditor/Private/Hierarchy/SHierarchyView.cpp
@@ -105,11 +105,6 @@ void SHierarchyView::Tick( const FGeometry& AllottedGeometry, const double InCur
 {
 	if ( bRebuildTreeRequested || bRefreshRequested )
 	{
-		if (!bExpandAllNodes)
-		{
-			FindExpandedItemNames();
-		}
-
 		if ( bRebuildTreeRequested )
 		{
 			RebuildTreeView();
@@ -117,15 +112,12 @@ void SHierarchyView::Tick( const FGeometry& AllottedGeometry, const double InCur
 
 		RefreshTree();
 
-		RestoreExpandedItems();
+		UpdateItemsExpansionFromModel();
 
 		OnEditorSelectionChanged();
 
 		bRefreshRequested = false;
 		bRebuildTreeRequested = false;
-		bExpandAllNodes = false;
-
-		ExpandedItemNames.Empty();
 	}
 }
 
@@ -184,8 +176,19 @@ void SHierarchyView::GetWidgetFilterStrings(TSharedPtr<FHierarchyModel> Item, TA
 void SHierarchyView::OnSearchChanged(const FText& InFilterText)
 {
 	bRefreshRequested = true;
-	bExpandAllNodes = InFilterText.IsEmpty();
-	FilterHandler->SetIsEnabled(!InFilterText.IsEmpty());
+	const bool bFilteringEnabled = !InFilterText.IsEmpty();
+	if (bFilteringEnabled != FilterHandler->GetIsEnabled())
+	{
+		FilterHandler->SetIsEnabled(bFilteringEnabled);
+		if (bFilteringEnabled)
+		{
+			SaveItemsExpansion();
+		}
+		else
+		{
+			RestoreItemsExpansion();
+		}
+	}
 	SearchBoxWidgetFilter->SetRawFilterText(InFilterText);
 	SearchBoxPtr->SetError(SearchBoxWidgetFilter->GetFilterErrorText());
 }
@@ -346,17 +349,23 @@ void SHierarchyView::OnObjectsReplaced(const TMap<UObject*, UObject*>& Replaceme
 	}
 }
 
-void SHierarchyView::RestoreExpandedItems()
+void SHierarchyView::UpdateItemsExpansionFromModel()
 {
-	EExpandBehavior ExpandBehavior = bExpandAllNodes ? EExpandBehavior::AlwaysExpand : EExpandBehavior::RestoreFromPrevious;
+	for (TSharedPtr<FHierarchyModel>& Model : RootWidgets)
+	{
+		RecursiveExpand(Model, EExpandBehavior::FromModel);
+	}
+}
 
-	for ( TSharedPtr<FHierarchyModel>& Model : RootWidgets )
+void SHierarchyView::RestoreItemsExpansion()
+{
+	for (TSharedPtr<FHierarchyModel>& Model : RootWidgets)
 	{
-		RecursiveExpand(Model, ExpandBehavior);
+		RecursiveExpand(Model, EExpandBehavior::RestoreFromPrevious);
 	}
 }
 
-void SHierarchyView::FindExpandedItemNames()
+void SHierarchyView::SaveItemsExpansion()
 {
 	ExpandedItemNames.Empty();
 
@@ -394,11 +403,17 @@ void SHierarchyView::RecursiveExpand(TSharedPtr<FHierarchyModel>& Model, EExpand
 		break;
 
 		case EExpandBehavior::AlwaysExpand:
-		default:
 		{
 			bShouldExpandItem = true;
 		}
 		break;
+
+		case EExpandBehavior::FromModel:
+		default:
+		{
+			bShouldExpandItem = Model->IsExpanded();
+		}
+		break;
 	}
 
 	WidgetTreeView->SetItemExpansion(Model, bShouldExpandItem);
@@ -410,11 +425,6 @@ void SHierarchyView::RecursiveExpand(TSharedPtr<FHierarchyModel>& Model, EExpand
 	{
 		RecursiveExpand(ChildModel, ExpandBehavior);
 	}
-
-	if (Model->IsExpanded())
-	{
-		WidgetTreeView->SetItemExpansion(Model, true);
-	}
 }
 
 void SHierarchyView::RestoreSelectedItems()
diff --git a/Engine/Source/Editor/UMGEditor/Private/Hierarchy/SHierarchyView.h b/Engine/Source/Editor/UMGEditor/Private/Hierarchy/SHierarchyView.h
--- a/Engine/Source/Editor/UMGEditor/Private/Hierarchy/SHierarchyView.h
+++ b/Engine/Source/Editor/UMGEditor/Private/Hierarchy/SHierarchyView.h
@@ -94,18 +94,25 @@ private:
 	/** Called when a Blueprint is recompiled and live objects are swapped out for replacements */
 	void OnObjectsReplaced(const TMap<UObject*, UObject*>& ReplacementMap);
 
-	/** Restores the state of expanded items based on the saved expanded item state, then clears the expanded state cache. */
-	void RestoreExpandedItems();
+	/** Sets the expansion state of hierarchy view items based on their model. */
+	void UpdateItemsExpansionFromModel();
+
+	/** Stores the names of all currently expanded nodes in the hierarchy view. */
+	void SaveItemsExpansion();
+
+	/** Sets the expansion state of hierarchy view items based on the state saved by SaveItemsExpansion. */
+	void RestoreItemsExpansion();
 
 	enum class EExpandBehavior : uint8
 	{
 		NeverExpand,
 		AlwaysExpand,
 		RestoreFromPrevious,
+		FromModel
 	};
 
 	/** Recursively expands the models based on the expansion set. */
-	void RecursiveExpand(TSharedPtr<FHierarchyModel>& Model, EExpandBehavior ExpandBehavior = EExpandBehavior::AlwaysExpand);
+	void RecursiveExpand(TSharedPtr<FHierarchyModel>& Model, EExpandBehavior ExpandBehavior);
 
 	/**  */
 	void RestoreSelectedItems();
@@ -116,9 +123,6 @@ private:
 	/** Handler for recursively expanding/collapsing items */
 	void SetItemExpansionRecursive(TSharedPtr<FHierarchyModel> Model, bool bInExpansionState);
 
-	/** Find and store the names of all currently expanded nodes in the hierarchy view. Should only be called when rebuilding the tree */
-	void FindExpandedItemNames();
-
 private:
 
 	/** Cached pointer to the blueprint editor that owns this tree. */
@@ -159,7 +163,4 @@ private:
 
 	/** Flag to ignore selections while the hierarchy view is updating the selection. */
 	bool bIsUpdatingSelection;
-
-	/** Should all nodes in the tree be expanded? */
-	bool bExpandAllNodes;
 };
