[NO NON-PII COMMIT MESSAGE]

diff --git a/Engine/Source/Editor/UMGEditor/Private/Designer/SDesignerView.cpp b/Engine/Source/Editor/UMGEditor/Private/Designer/SDesignerView.cpp
--- a/Engine/Source/Editor/UMGEditor/Private/Designer/SDesignerView.cpp
+++ b/Engine/Source/Editor/UMGEditor/Private/Designer/SDesignerView.cpp
@@ -216,6 +216,7 @@ public:
 		/** The original parent of the widget. */
 		FWidgetReference ParentWidget;
 
+		/** The offset of the original click location, as a percentage of the widget's size. */
 		FVector2D DraggedOffset;
 	};
 
@@ -242,7 +243,7 @@ TSharedRef<FSelectedWidgetDragDropOp> FSelectedWidgetDragDropOp::New(TSharedPtr<
 	Operation->bShowingMessage = false;
 	Operation->Designer = InDesigner;
 
-	for (const auto& InDraggedWidget : InWidgets)
+	for (const FDraggingWidgetReference& InDraggedWidget : InWidgets)
 	{
 		FItem DraggedWidget;
 		DraggedWidget.bStayingInParent = false;
@@ -2528,7 +2529,7 @@ FReply SDesignerView::OnDragDetected(const FGeometry& MyGeometry, const FPointer
 		// Clear any pending selected widgets, the user has already decided what widget they want.
 		PendingSelectedWidget = FWidgetReference();
 
-		for (const auto& SelectedWidget : SelectedWidgets)
+		for (const FWidgetReference& SelectedWidget : SelectedWidgets)
 		{
 			// Determine The offset to keep the widget from the mouse while dragging
 			FArrangedWidget ArrangedWidget(SNullWidget::NullWidget, FGeometry());
@@ -2537,14 +2538,13 @@ FReply SDesignerView::OnDragDetected(const FGeometry& MyGeometry, const FPointer
 
 			FDragWidget DraggingWidget;
 			DraggingWidget.Widget = SelectedWidget;
-			DraggingWidget.DraggedOffset = SelectedWidgetContextMenuLocation;
-
+			DraggingWidget.DraggedOffset = SelectedWidgetContextMenuLocation / ArrangedWidget.Geometry.GetLocalSize();
 			DraggingWidgetCandidates.Add(DraggingWidget);
 		}
 
 		TArray<FDragWidget> DraggingWidgets;
 
-		for (const auto& Candidate : DraggingWidgetCandidates)
+		for (const FDragWidget& Candidate : DraggingWidgetCandidates)
 		{
 			// check the parent chain of each dragged widget and ignore those that are children of other dragged widgets
 			bool bIsChild = false;
@@ -2980,8 +2980,6 @@ void SDesignerView::ProcessDropAndAddWidget(const FGeometry& MyGeometry, const F
 
 					if (Slot != nullptr)
 					{
-						FVector2D NewPosition = LocalPosition - DraggedWidget.DraggedOffset;
-
 						FWidgetBlueprintEditorUtils::ImportPropertiesFromText(Slot, DraggedWidget.ExportedSlotProperties);
 
 						bool HasChangedLayout = false;
@@ -2996,6 +2994,12 @@ void SDesignerView::ProcessDropAndAddWidget(const FGeometry& MyGeometry, const F
 							{
 								CanvasSlot->SaveBaseLayout();
 
+								FArrangedWidget ArrangedWidget(SNullWidget::NullWidget, FGeometry());
+								FDesignTimeUtils::GetArrangedWidget(Widget->GetCachedWidget().ToSharedRef(), ArrangedWidget);
+
+								FVector2D Offset = DraggedWidget.DraggedOffset * ArrangedWidget.Geometry.GetLocalSize();
+								FVector2D NewPosition = LocalPosition - Offset;
+
 								// Perform grid snapping on X and Y if we need to.
 								if (bGridSnapX)
 								{
