Fix glitchy first frame text layout inside a widget in the SDesignerView; perform SlatePrepass on the preview widget with the correct layout scale.

diff --git a/Engine/Source/Editor/UMGEditor/Private/Designer/SDesignerView.cpp b/Engine/Source/Editor/UMGEditor/Private/Designer/SDesignerView.cpp
--- a/Engine/Source/Editor/UMGEditor/Private/Designer/SDesignerView.cpp
+++ b/Engine/Source/Editor/UMGEditor/Private/Designer/SDesignerView.cpp
@@ -442,6 +442,11 @@ void SDesignerView::Construct(const FArguments& InArgs, TSharedPtr<FWidgetBluepr
 								[
 									SAssignNew(PreviewSurface, SDPIScaler)
 									.DPIScale(this, &SDesignerView::GetPreviewDPIScale)
+									[
+										SAssignNew(PreviewSizeConstraint, SBox)
+										.WidthOverride(this, &SDesignerView::GetPreviewSizeWidth)
+										.HeightOverride(this, &SDesignerView::GetPreviewSizeHeight)
+									]
 								]
 							]
 						]
@@ -1623,7 +1628,7 @@ void SDesignerView::OnPreviewNeedsRecreation()
 	CachedWidgetGeometry.Reset();
 
 	PreviewWidget = nullptr;
-	PreviewSurface->SetContent(SNullWidget::NullWidget);
+	PreviewSizeConstraint->SetContent(SNullWidget::NullWidget);
 }
 
 SDesignerView::FWidgetHitResult::FWidgetHitResult()
@@ -2299,20 +2304,11 @@ void SDesignerView::UpdatePreviewWidget(bool bForceUpdate)
 		if ( PreviewWidget )
 		{
 			TSharedRef<SWidget> NewPreviewSlateWidget = PreviewWidget->TakeWidget();
-			NewPreviewSlateWidget->SlatePrepass();
+			NewPreviewSlateWidget->SlatePrepass(PreviewSizeConstraint->GetCachedGeometry().Scale);
 
 			PreviewSlateWidget = NewPreviewSlateWidget;
 
-			// The constraint box for the widget size needs to inside the DPI scaler in order to make sure it too
-			// is sized accurately for the size screen it's on.
-			TSharedRef<SBox> NewPreviewSizeConstraintBox = SNew(SBox)
-				.WidthOverride(this, &SDesignerView::GetPreviewSizeWidth)
-				.HeightOverride(this, &SDesignerView::GetPreviewSizeHeight)
-				[
-					NewPreviewSlateWidget
-				];
-
-			PreviewSurface->SetContent(NewPreviewSizeConstraintBox);
+			PreviewSizeConstraint->SetContent(NewPreviewSlateWidget);
 
 			// Notify all selected widgets that they are selected, because there are new preview objects
 			// state may have been lost so this will recreate it if the widget does something special when
diff --git a/Engine/Source/Editor/UMGEditor/Private/Designer/SDesignerView.h b/Engine/Source/Editor/UMGEditor/Private/Designer/SDesignerView.h
--- a/Engine/Source/Editor/UMGEditor/Private/Designer/SDesignerView.h
+++ b/Engine/Source/Editor/UMGEditor/Private/Designer/SDesignerView.h
@@ -318,6 +318,7 @@ private:
 	TSharedPtr<SWidget> PreviewHitTestRoot;
 	TSharedPtr<SBox> PreviewAreaConstraint;
 	TSharedPtr<SDPIScaler> PreviewSurface;
+	TSharedPtr<SBox> PreviewSizeConstraint;
 
 	TSharedPtr<SCanvas> DesignerControls;
 	TSharedPtr<SCanvas> DesignerWidgetCanvas;
