Porting over fix for UE-53728 to also support CreateWidget nodes, fixing an issue where properties aren't assigned during construction

diff --git a/Engine/Source/Editor/UMGEditor/Private/Nodes/K2Node_CreateWidget.cpp b/Engine/Source/Editor/UMGEditor/Private/Nodes/K2Node_CreateWidget.cpp
--- a/Engine/Source/Editor/UMGEditor/Private/Nodes/K2Node_CreateWidget.cpp
+++ b/Engine/Source/Editor/UMGEditor/Private/Nodes/K2Node_CreateWidget.cpp
@@ -110,6 +110,9 @@ void UK2Node_CreateWidget::ExpandNode(class FKismetCompilerContext& CompilerCont
 	CallCreateNode->FunctionReference.SetExternalMember(Create_FunctionName, UWidgetBlueprintLibrary::StaticClass());
 	CallCreateNode->AllocateDefaultPins();
 
+	// store off the class to spawn before we mutate pin connections:
+	UClass* ClassToSpawn = GetClassToSpawn();
+
 	UEdGraphPin* CallCreateExec = CallCreateNode->GetExecPin();
 	UEdGraphPin* CallCreateWorldContextPin = CallCreateNode->FindPinChecked(WorldContextObject_ParamName);
 	UEdGraphPin* CallCreateWidgetTypePin = CallCreateNode->FindPinChecked(WidgetType_ParamName);
@@ -147,7 +150,7 @@ void UK2Node_CreateWidget::ExpandNode(class FKismetCompilerContext& CompilerCont
 	// create 'set var' nodes
 
 	// Get 'result' pin from 'begin spawn', this is the actual actor we want to set properties on
-	UEdGraphPin* LastThen = FKismetCompilerUtilities::GenerateAssignmentNodes(CompilerContext, SourceGraph, CallCreateNode, CreateWidgetNode, CallCreateResult, GetClassToSpawn());
+	UEdGraphPin* LastThen = FKismetCompilerUtilities::GenerateAssignmentNodes(CompilerContext, SourceGraph, CallCreateNode, CreateWidgetNode, CallCreateResult, ClassToSpawn);
 
 	// Move 'then' connection from create widget node to the last 'then'
 	CompilerContext.MovePinLinksToIntermediate(*SpawnNodeThen, *LastThen);
