UMG Editor: Change to the Widget Favorite Feature to include it in the Palette Window instead of having a separate window.

diff --git a/Engine/Source/Editor/UMGEditor/Classes/WidgetPaletteFavorites.cpp b/Engine/Source/Editor/UMGEditor/Classes/WidgetPaletteFavorites.cpp
--- a/Engine/Source/Editor/UMGEditor/Classes/WidgetPaletteFavorites.cpp
+++ b/Engine/Source/Editor/UMGEditor/Classes/WidgetPaletteFavorites.cpp
@@ -14,6 +14,7 @@ void UWidgetPaletteFavorites::Add(const FString& InWidgetTemplateName)
 
 	SaveConfig();
 
+	OnFavoritesUpdated.Broadcast();
 }
 
 void UWidgetPaletteFavorites::Remove(const FString& InWidgetTemplateName)
@@ -21,4 +22,6 @@ void UWidgetPaletteFavorites::Remove(const FString& InWidgetTemplateName)
 	Favorites.Remove(InWidgetTemplateName);
 
 	SaveConfig();
+
+	OnFavoritesUpdated.Broadcast();
 }
\ No newline at end of file
diff --git a/Engine/Source/Editor/UMGEditor/Classes/WidgetPaletteFavorites.h b/Engine/Source/Editor/UMGEditor/Classes/WidgetPaletteFavorites.h
--- a/Engine/Source/Editor/UMGEditor/Classes/WidgetPaletteFavorites.h
+++ b/Engine/Source/Editor/UMGEditor/Classes/WidgetPaletteFavorites.h
@@ -17,6 +17,11 @@ public:
 	void Remove(const FString& InWidgetTemplateName);
 
 	TArray<FString> GetFavorites() const { return Favorites; }
+	
+	DECLARE_MULTICAST_DELEGATE(FOnFavoritesUpdated)
+
+	/** Fires after the list of favorites is updated */
+	FOnFavoritesUpdated OnFavoritesUpdated;
 
 private:
 	UPROPERTY(config)
diff --git a/Engine/Source/Editor/UMGEditor/Private/BlueprintModes/WidgetDesignerApplicationMode.cpp b/Engine/Source/Editor/UMGEditor/Private/BlueprintModes/WidgetDesignerApplicationMode.cpp
--- a/Engine/Source/Editor/UMGEditor/Private/BlueprintModes/WidgetDesignerApplicationMode.cpp
+++ b/Engine/Source/Editor/UMGEditor/Private/BlueprintModes/WidgetDesignerApplicationMode.cpp
@@ -11,7 +11,6 @@
 
 
 #include "TabFactory/PaletteTabSummoner.h"
-#include "TabFactory/PaletteFavoritesTabSummoner.h"
 #include "TabFactory/HierarchyTabSummoner.h"
 #include "TabFactory/DesignerTabSummoner.h"
 #include "TabFactory/SequencerTabSummoner.h"
@@ -59,12 +58,6 @@ FWidgetDesignerApplicationMode::FWidgetDesignerApplicationMode(TSharedPtr<FWidge
 					->AddTab( FPaletteTabSummoner::TabID, ETabState::OpenedTab )
 				)
 				->Split
-				(
-					FTabManager::NewStack()
-					->SetSizeCoefficient( 0.5f )
-					->AddTab( FPaletteFavoritesTabSummoner::TabID, ETabState::OpenedTab )
-				)
-				->Split
 				(
 					FTabManager::NewStack()
 					->SetSizeCoefficient( 0.5f )
@@ -118,7 +111,6 @@ FWidgetDesignerApplicationMode::FWidgetDesignerApplicationMode(TSharedPtr<FWidge
 	TabFactories.RegisterFactory(MakeShareable(new FDesignerTabSummoner(InWidgetEditor)));
 	TabFactories.RegisterFactory(MakeShareable(new FHierarchyTabSummoner(InWidgetEditor)));
 	TabFactories.RegisterFactory(MakeShareable(new FPaletteTabSummoner(InWidgetEditor)));
-	TabFactories.RegisterFactory(MakeShareable(new FPaletteFavoritesTabSummoner(InWidgetEditor)));
 	TabFactories.RegisterFactory(MakeShareable(new FSequencerTabSummoner(InWidgetEditor)));
 	TabFactories.RegisterFactory(MakeShareable(new FAnimationTabSummoner(InWidgetEditor)));
 	TabFactories.RegisterFactory(MakeShareable(new FCompilerResultsSummoner(InWidgetEditor)));
diff --git a/Engine/Source/Editor/UMGEditor/Private/Palette/SPaletteFavoritesView.cpp b/Engine/Source/Editor/UMGEditor/Private/Palette/SPaletteFavoritesView.cpp
deleted file mode 100644
diff --git a/Engine/Source/Editor/UMGEditor/Private/Palette/SPaletteFavoritesView.h b/Engine/Source/Editor/UMGEditor/Private/Palette/SPaletteFavoritesView.h
deleted file mode 100644
diff --git a/Engine/Source/Editor/UMGEditor/Private/Palette/SPaletteViewModel.cpp b/Engine/Source/Editor/UMGEditor/Private/Palette/SPaletteViewModel.cpp
--- a/Engine/Source/Editor/UMGEditor/Private/Palette/SPaletteViewModel.cpp
+++ b/Engine/Source/Editor/UMGEditor/Private/Palette/SPaletteViewModel.cpp
@@ -101,8 +101,7 @@ void FWidgetHeaderViewModel::GetChildren(TArray< TSharedPtr<FWidgetViewModel> >&
 }
 
 FPaletteViewModel::FPaletteViewModel(TSharedPtr<FWidgetBlueprintEditor> InBlueprintEditor)
-	: bRebuildRequested(true),
-	  bRebuildFavorites(true)
+	: bRebuildRequested(true)
 {
 	BlueprintEditor = InBlueprintEditor;
 }
@@ -116,6 +115,10 @@ void FPaletteViewModel::RegisterToEvents()
 
 	// register for any objects replaced
 	GEditor->OnObjectsReplaced().AddRaw(this, &FPaletteViewModel::OnObjectsReplaced);
+
+	// Register for favorite list update to handle the case where a favorite is added in another window of the UMG Designer
+	UWidgetPaletteFavorites* Favorites = GetDefault<UWidgetDesignerSettings>()->Favorites;
+	Favorites->OnFavoritesUpdated.AddSP(this, &FPaletteViewModel::OnFavoritesUpdated);
 }
 
 FPaletteViewModel::~FPaletteViewModel()
@@ -124,45 +127,31 @@ FPaletteViewModel::~FPaletteViewModel()
 	FEditorDelegates::OnAssetsDeleted.RemoveAll(this);
 	IHotReloadModule::Get().OnHotReload().RemoveAll(this);
 	GEditor->OnObjectsReplaced().RemoveAll(this);
+
+	UWidgetPaletteFavorites* Favorites = GetDefault<UWidgetDesignerSettings>()->Favorites;
+	Favorites->OnFavoritesUpdated.RemoveAll(this);
 }
 
 void FPaletteViewModel::AddToFavorites(const FWidgetTemplateViewModel* WidgetTemplateViewModel)
 {
 	UWidgetPaletteFavorites* Favorites = GetDefault<UWidgetDesignerSettings>()->Favorites;
 	Favorites->Add(WidgetTemplateViewModel->GetName().ToString());
-
-	// Force a rebuild of the favorites list.
-	bRebuildFavorites = true;
 }
 
 void FPaletteViewModel::RemoveFromFavorites(const FWidgetTemplateViewModel* WidgetTemplateViewModel)
 {
 	UWidgetPaletteFavorites* Favorites = GetDefault<UWidgetDesignerSettings>()->Favorites;
 	Favorites->Remove(WidgetTemplateViewModel->GetName().ToString());
-	
-	// Force a rebuild of the favorites list.
-	bRebuildFavorites = true;
 }
 
 void FPaletteViewModel::Update()
 {
 	if (bRebuildRequested)
 	{
-		bRebuildRequested = false;
-
 		OnUpdating.Broadcast();
 		BuildWidgetList();
+		bRebuildRequested = false;
 		OnUpdated.Broadcast();
-
-		// We need to rebuild the list of favorites to also clear widgets that where potentially removed.
-		bRebuildFavorites = true;
-	}
-
-	if (bRebuildFavorites)
-	{
-		bRebuildFavorites = false;
-		BuildWidgetFavoriteList();
-		OnFavoritesUpdated.Broadcast();
 	}
 }
 
@@ -187,6 +176,14 @@ void FPaletteViewModel::BuildWidgetList()
 	// Generate a list of templates
 	BuildClassWidgetList();
 
+	// Build the Favorite section
+	TSharedPtr<FWidgetHeaderViewModel> FavoriteHeader = MakeShareable(new FWidgetHeaderViewModel());
+	FavoriteHeader->GroupName = LOCTEXT("Favorites", "Favorites");
+	
+	// Copy of the list of favorites to be able to do some cleanup in the real list
+	UWidgetPaletteFavorites* FavoritesPalette = GetDefault<UWidgetDesignerSettings>()->Favorites;
+	TArray<FString> FavoritesList = FavoritesPalette->GetFavorites();
+
 	// For each entry in the category create a view model for the widget template
 	for ( auto& Entry : WidgetTemplateCategories )
 	{
@@ -199,16 +196,48 @@ void FPaletteViewModel::BuildWidgetList()
 			TemplateViewModel->Template = Template;
 			TemplateViewModel->PaletteViewModel = this;
 			Header->Children.Add(TemplateViewModel);
+
+			// If it's a favorite, we also add it to the Favorite section
+			int32 index = FavoritesList.Find(Template->Name.ToString());
+			if (index != INDEX_NONE)
+			{
+				TemplateViewModel->SetFavorite();
+
+				// We have to create a second copy of the ViewModel for the treeview has it doesn't support to have the same element twice.
+				TSharedPtr<FWidgetTemplateViewModel> FavoriteTemplateViewModel = MakeShareable(new FWidgetTemplateViewModel());
+				FavoriteTemplateViewModel->Template = Template;
+				FavoriteTemplateViewModel->PaletteViewModel = this;
+				FavoriteTemplateViewModel->SetFavorite();
+
+				FavoriteHeader->Children.Add(FavoriteTemplateViewModel);
+
+				// Remove the favorite from the temporary list
+				FavoritesList.RemoveAt(index);
+			}
+
 		}
 
 		Header->Children.Sort([] (TSharedPtr<FWidgetViewModel> L, TSharedPtr<FWidgetViewModel> R) { return R->GetName().CompareTo(L->GetName()) > 0; });
 
 		WidgetViewModels.Add(Header);
+	}	
+
+	// Remove all Favorites that may be left in the list.Typically happening when the list of favorite contains widget that were deleted since the last opening.
+	for (const FString& favoriteName : FavoritesList)
+	{
+		FavoritesPalette->Remove(favoriteName);
 	}
 
 	// Sort the view models by name
 	WidgetViewModels.Sort([] (TSharedPtr<FWidgetViewModel> L, TSharedPtr<FWidgetViewModel> R) { return R->GetName().CompareTo(L->GetName()) > 0; });
 
+	// Add the Favorite section at the top
+	if (FavoriteHeader->Children.Num() != 0)
+	{
+		FavoriteHeader->Children.Sort([](TSharedPtr<FWidgetViewModel> L, TSharedPtr<FWidgetViewModel> R) { return R->GetName().CompareTo(L->GetName()) > 0; });
+		WidgetViewModels.Insert(FavoriteHeader, 0);
+	}
+	
 	// Take the Advanced Section, and put it at the end.
 	TSharedPtr<FWidgetViewModel>* advancedSectionPtr = WidgetViewModels.FindByPredicate([](TSharedPtr<FWidgetViewModel> widget) {return widget->GetName().CompareTo(LOCTEXT("Advanced", "Advanced")) == 0; });
 	if (advancedSectionPtr)
@@ -219,42 +248,6 @@ void FPaletteViewModel::BuildWidgetList()
 	}
 }
 
-void FPaletteViewModel::BuildWidgetFavoriteList()
-{
-	WidgetFavoritesViewModels.Reset();
-	UWidgetPaletteFavorites* FavoritesPalette = GetDefault<UWidgetDesignerSettings>()->Favorites;
-
-	// We Iterate on a copy to be able to clean widgets that may no longer exist
-	TArray<FString> FavoritesList = FavoritesPalette->GetFavorites();
-	for (const FString& favoriteName : FavoritesList)
-	{
-		bool bWidgetFound = false;
-		for (TSharedPtr < FWidgetViewModel>& widgetSection : WidgetViewModels)
-		{
-			TArray < TSharedPtr<FWidgetViewModel> > Children;
-			widgetSection->GetChildren(Children);
-
-			TSharedPtr<FWidgetViewModel>* widgetViewModel = Children.FindByPredicate([&](const TSharedPtr<FWidgetViewModel>& widget) { return widget->GetName().ToString().Compare(favoriteName) == 0; });
-			if (widgetViewModel != nullptr)
-			{
-				(*widgetViewModel)->SetFavorite();
-
-				WidgetFavoritesViewModels.Add(*widgetViewModel);
-				bWidgetFound = true;
-				break;
-			}
-		}
-
-		// If we did not found the widget, it no longer exists. We remove it from the original list of favorites.
-		if (!bWidgetFound)
-		{
-			FavoritesPalette->Remove(favoriteName);
-		}
-	}
-
-	WidgetFavoritesViewModels.Sort([](TSharedPtr<FWidgetViewModel> L, TSharedPtr<FWidgetViewModel> R) { return R->GetName().CompareTo(L->GetName()) > 0; });
-}
-
 void FPaletteViewModel::BuildClassWidgetList()
 {
 	static const FName DevelopmentStatusKey(TEXT("DevelopmentStatus"));
@@ -467,6 +460,11 @@ void FPaletteViewModel::OnBlueprintReinstanced()
 	bRebuildRequested = true;
 }
 
+void FPaletteViewModel::OnFavoritesUpdated()
+{
+	bRebuildRequested = true;
+}
+
 void FPaletteViewModel::HandleOnHotReload(bool bWasTriggeredAutomatically)
 {
 	bRebuildRequested = true;
diff --git a/Engine/Source/Editor/UMGEditor/Private/Palette/SPaletteViewModel.h b/Engine/Source/Editor/UMGEditor/Private/Palette/SPaletteViewModel.h
--- a/Engine/Source/Editor/UMGEditor/Private/Palette/SPaletteViewModel.h
+++ b/Engine/Source/Editor/UMGEditor/Private/Palette/SPaletteViewModel.h
@@ -113,7 +113,6 @@ public:
 
 	DECLARE_MULTICAST_DELEGATE(FOnUpdating)
 	DECLARE_MULTICAST_DELEGATE(FOnUpdated)
-	DECLARE_MULTICAST_DELEGATE(FOnFavoritesUpdated)
 
 public:
 	FPaletteViewModel(TSharedPtr<FWidgetBlueprintEditor> InBlueprintEditor);
@@ -126,7 +125,7 @@ public:
 	void Update();
 
 	/** Returns true if the view model needs to be updated */
-	bool NeedUpdate() const { return bRebuildRequested || bRebuildFavorites; }
+	bool NeedUpdate() const { return bRebuildRequested; }
 	   
 	/** Add the widget template to the list of favorites */
 	void AddToFavorites(const FWidgetTemplateViewModel* WidgetTemplateViewModel);
@@ -135,7 +134,6 @@ public:
 	void RemoveFromFavorites(const FWidgetTemplateViewModel* WidgetTemplateViewModel);
 
 	typedef TArray< TSharedPtr<FWidgetViewModel> > ViewModelsArray;
-	ViewModelsArray& GetFavoritesViewModels() { return WidgetFavoritesViewModels; }
 	ViewModelsArray& GetWidgetViewModels() { return WidgetViewModels; }
 
 	void SetSearchText(const FText& inSearchText) { SearchText = inSearchText; }
@@ -147,16 +145,12 @@ public:
 	/** Fires after the view model is updated */
 	FOnUpdated OnUpdated;
 
-	/** Fires after the list of favorites is updated */
-	FOnFavoritesUpdated OnFavoritesUpdated;
-
 private:
 	FPaletteViewModel() {};
 
 	UWidgetBlueprint* GetBlueprint() const;
 
 	void BuildWidgetList();
-	void BuildWidgetFavoriteList();
 	void BuildClassWidgetList();
 
 	static bool FilterAssetData(FAssetData &BPAssetData);
@@ -169,12 +163,15 @@ private:
 	/** Requests a rebuild of the widget list if a widget blueprint was compiled */
 	void OnBlueprintReinstanced();
 
+	/** Called when the favorite list is changed */
+	void OnFavoritesUpdated();
+
 	/** Requests a rebuild of the widget list */
 	void HandleOnHotReload(bool bWasTriggeredAutomatically);
 
 	/** Requests a rebuild of the widget list if a widget blueprint was deleted */
 	void HandleOnAssetsDeleted(const TArray<UClass*>& DeletedAssetClasses);
-
+	
 	TWeakPtr<class FWidgetBlueprintEditor> BlueprintEditor;
 
 	typedef TArray<TSharedPtr<FWidgetTemplate>> WidgetTemplateArray;
@@ -182,16 +179,10 @@ private:
 
 	/** The source root view models for the tree. */
 	ViewModelsArray WidgetViewModels;
-
-	/** The source root view models for the favorites. */
-	ViewModelsArray WidgetFavoritesViewModels;
 	   
 	/** Controls rebuilding the list of spawnable widgets */
 	bool bRebuildRequested;
 
-	/** Set to true when the Palette favorites list needs to be updated*/
-	bool bRebuildFavorites;
-
 	FText SearchText;
 };
 
diff --git a/Engine/Source/Editor/UMGEditor/Private/TabFactory/PaletteFavoritesTabSummoner.cpp b/Engine/Source/Editor/UMGEditor/Private/TabFactory/PaletteFavoritesTabSummoner.cpp
deleted file mode 100644
diff --git a/Engine/Source/Editor/UMGEditor/Private/TabFactory/PaletteFavoritesTabSummoner.h b/Engine/Source/Editor/UMGEditor/Private/TabFactory/PaletteFavoritesTabSummoner.h
deleted file mode 100644
