Fix canvas slot resize with mouse in designer when display scaling is different than 1.

diff --git a/Engine/Source/Editor/UMGEditor/Private/Designer/STransformHandle.cpp b/Engine/Source/Editor/UMGEditor/Private/Designer/STransformHandle.cpp
--- a/Engine/Source/Editor/UMGEditor/Private/Designer/STransformHandle.cpp
+++ b/Engine/Source/Editor/UMGEditor/Private/Designer/STransformHandle.cpp
@@ -128,8 +128,8 @@ FReply STransformHandle::OnMouseMove(const FGeometry& MyGeometry, const FPointer
 		UWidget* Preview = SelectedWidget.GetPreview();
 
 		{
-			FVector2D Delta = MouseEvent.GetScreenSpacePosition() - MouseDownPosition;
-			FVector2D TranslateAmount = Delta * ( 1.0f / Designer->GetPreviewScale() );
+			const FVector2D Delta = MouseEvent.GetScreenSpacePosition() - MouseDownPosition;
+			const FVector2D TranslateAmount = Delta * (1.0f / (Designer->GetPreviewScale() * MyGeometry.Scale));
 
 			Resize(Cast<UCanvasPanelSlot>(Preview->Slot), DragDirection, TranslateAmount);
 			Resize(Cast<UCanvasPanelSlot>(Template->Slot), DragDirection, TranslateAmount);
