UMG Editor: Minor improvement to Favorite. When add the first elements to favorites we force to expand it.

diff --git a/Engine/Source/Editor/UMGEditor/Private/Palette/SPaletteView.cpp b/Engine/Source/Editor/UMGEditor/Private/Palette/SPaletteView.cpp
--- a/Engine/Source/Editor/UMGEditor/Private/Palette/SPaletteView.cpp
+++ b/Engine/Source/Editor/UMGEditor/Private/Palette/SPaletteView.cpp
@@ -303,7 +303,7 @@ void SPaletteView::OnViewModelUpdated()
 	{
 		for (TSharedPtr<FWidgetViewModel>& ViewModel : PaletteViewModel->GetWidgetViewModels())
 		{
-			if (ViewModel->GetName().EqualTo(ExpandedItem->GetName()))
+			if (ViewModel->GetName().EqualTo(ExpandedItem->GetName()) || ViewModel->ShouldForceExpansion())
 			{
 				WidgetTemplatesView->SetItemExpansion(ViewModel, true);
 			}
diff --git a/Engine/Source/Editor/UMGEditor/Private/Palette/SPaletteViewModel.cpp b/Engine/Source/Editor/UMGEditor/Private/Palette/SPaletteViewModel.cpp
--- a/Engine/Source/Editor/UMGEditor/Private/Palette/SPaletteViewModel.cpp
+++ b/Engine/Source/Editor/UMGEditor/Private/Palette/SPaletteViewModel.cpp
@@ -104,6 +104,9 @@ FPaletteViewModel::FPaletteViewModel(TSharedPtr<FWidgetBlueprintEditor> InBluepr
 	: bRebuildRequested(true)
 {
 	BlueprintEditor = InBlueprintEditor;
+
+	FavoriteHeader = MakeShareable(new FWidgetHeaderViewModel());
+	FavoriteHeader->GroupName = LOCTEXT("Favorites", "Favorites");
 }
 
 void FPaletteViewModel::RegisterToEvents()
@@ -176,9 +179,9 @@ void FPaletteViewModel::BuildWidgetList()
 	// Generate a list of templates
 	BuildClassWidgetList();
 
-	// Build the Favorite section
-	TSharedPtr<FWidgetHeaderViewModel> FavoriteHeader = MakeShareable(new FWidgetHeaderViewModel());
-	FavoriteHeader->GroupName = LOCTEXT("Favorites", "Favorites");
+	// Clear the Favorite section
+	bool bHasFavorites = FavoriteHeader->Children.Num() != 0;
+	FavoriteHeader->Children.Reset();
 	
 	// Copy of the list of favorites to be able to do some cleanup in the real list
 	UWidgetPaletteFavorites* FavoritesPalette = GetDefault<UWidgetDesignerSettings>()->Favorites;
@@ -234,6 +237,8 @@ void FPaletteViewModel::BuildWidgetList()
 	// Add the Favorite section at the top
 	if (FavoriteHeader->Children.Num() != 0)
 	{
+		// We force expansion of the favorite header when we add favorites for the first time.
+		FavoriteHeader->SetForceExpansion(!bHasFavorites);
 		FavoriteHeader->Children.Sort([](TSharedPtr<FWidgetViewModel> L, TSharedPtr<FWidgetViewModel> R) { return R->GetName().CompareTo(L->GetName()) > 0; });
 		WidgetViewModels.Insert(FavoriteHeader, 0);
 	}
diff --git a/Engine/Source/Editor/UMGEditor/Private/Palette/SPaletteViewModel.h b/Engine/Source/Editor/UMGEditor/Private/Palette/SPaletteViewModel.h
--- a/Engine/Source/Editor/UMGEditor/Private/Palette/SPaletteViewModel.h
+++ b/Engine/Source/Editor/UMGEditor/Private/Palette/SPaletteViewModel.h
@@ -40,6 +40,8 @@ public:
 	virtual void SetFavorite()
 	{
 	}
+
+	virtual bool ShouldForceExpansion() const { return false; }
 };
 
 class FWidgetTemplateViewModel : public FWidgetViewModel
@@ -103,8 +105,16 @@ public:
 
 	virtual void GetChildren(TArray< TSharedPtr<FWidgetViewModel> >& OutChildren) override;
 
+	virtual bool ShouldForceExpansion() const { return bForceExpansion; }
+
+	void SetForceExpansion(bool bInForceExpansion) { bForceExpansion = bInForceExpansion; }
+
 	FText GroupName;
 	TArray< TSharedPtr<FWidgetViewModel> > Children;
+
+private:
+
+	bool bForceExpansion = false;
 };
 
 class FPaletteViewModel: public TSharedFromThis<FPaletteViewModel>
@@ -184,5 +194,7 @@ private:
 	bool bRebuildRequested;
 
 	FText SearchText;
+
+	TSharedPtr<FWidgetHeaderViewModel> FavoriteHeader;
 };
 
