Remove IsGenericName to replace it with UWidget::IsGeneratedName

diff --git a/Engine/Source/Editor/UMGEditor/Private/WidgetBlueprintEditorUtils.cpp b/Engine/Source/Editor/UMGEditor/Private/WidgetBlueprintEditorUtils.cpp
--- a/Engine/Source/Editor/UMGEditor/Private/WidgetBlueprintEditorUtils.cpp
+++ b/Engine/Source/Editor/UMGEditor/Private/WidgetBlueprintEditorUtils.cpp
@@ -783,11 +783,12 @@ void FWidgetBlueprintEditorUtils::ReplaceWidgetWithSelectedTemplate(TSharedRef<F
 		CurrentParent->ReplaceChild(ThisWidget, NewReplacementWidget);
 
 		FString ReplaceName = ThisWidget->GetName();
+		bool bIsGeneratedName = ThisWidget->IsGeneratedName();
 		// Rename the removed widget to the transient package so that it doesn't conflict with future widgets sharing the same name.
 		ThisWidget->Rename(nullptr, nullptr);
 
 		// Rename the new Widget to maintain the current name if it's not a generic name
-		if (!IsGenericName(ReplaceName, ThisWidget->GetClass()))
+		if (!bIsGeneratedName)
 		{
 			ReplaceName = FindNextValidName(BP->WidgetTree, ReplaceName);
 			NewReplacementWidget->Rename(*ReplaceName, BP->WidgetTree);
@@ -948,11 +949,12 @@ void FWidgetBlueprintEditorUtils::ReplaceWidgets(TSharedRef<FWidgetBlueprintEdit
 			}
 
 			FString ReplaceName = Item.GetTemplate()->GetName();
+			bool bIsGeneratedName = Item.GetTemplate()->IsGeneratedName();
 			// Rename the removed widget to the transient package so that it doesn't conflict with future widgets sharing the same name.
 			Item.GetTemplate()->Rename(nullptr, nullptr);
 
 			// Rename the new Widget to maintain the current name if it's not a generic name
-			if (!IsGenericName(ReplaceName, Item.GetTemplate()->GetClass()))
+			if (!bIsGeneratedName)
 			{
 				ReplaceName = FindNextValidName(BP->WidgetTree, ReplaceName);
 				NewReplacementWidget->Rename(*ReplaceName, BP->WidgetTree);
@@ -1560,10 +1562,4 @@ FString FWidgetBlueprintEditorUtils::FindNextValidName(UWidgetTree* WidgetTree,
 	return Name;
 }
 
-bool FWidgetBlueprintEditorUtils::IsGenericName(const FString& Name, const UClass* WidgetClass)
-{
-	FString NewName = RemoveSuffixFromName(Name);
-	return (NewName == WidgetClass->GetName());
-}
-
 #undef LOCTEXT_NAMESPACE
diff --git a/Engine/Source/Editor/UMGEditor/Private/WidgetBlueprintEditorUtils.h b/Engine/Source/Editor/UMGEditor/Private/WidgetBlueprintEditorUtils.h
--- a/Engine/Source/Editor/UMGEditor/Private/WidgetBlueprintEditorUtils.h
+++ b/Engine/Source/Editor/UMGEditor/Private/WidgetBlueprintEditorUtils.h
@@ -88,6 +88,4 @@ private:
 	static void ReplaceWidgets(TSharedRef<FWidgetBlueprintEditor> BlueprintEditor, UWidgetBlueprint* BP, TSet<FWidgetReference> Widgets, UClass* WidgetClass);
 
 	static FString FindNextValidName(UWidgetTree* WidgetTree, const FString& Name);
-
-	static bool IsGenericName(const FString& Name, const UClass* WidgetClass);
 };
