Fixed crash when using keyboard arrow keys or gamepad to navigate Editor UI.
There was a dependency with statically initialized variables.

diff --git a/Engine/Plugins/Experimental/RemoteSession/Source/RemoteSession/Private/MessageHandler/RecordingMessageHandler.cpp b/Engine/Plugins/Experimental/RemoteSession/Source/RemoteSession/Private/MessageHandler/RecordingMessageHandler.cpp
--- a/Engine/Plugins/Experimental/RemoteSession/Source/RemoteSession/Private/MessageHandler/RecordingMessageHandler.cpp
+++ b/Engine/Plugins/Experimental/RemoteSession/Source/RemoteSession/Private/MessageHandler/RecordingMessageHandler.cpp
@@ -135,7 +135,7 @@ FVector2D FRecordingMessageHandler::ConvertFromNormalizedScreenLocation(const FV
 				FWidgetPath WidgetPath(GameWindow.ToSharedRef(), JustWindow);
 				if (WidgetPath.ExtendPathTo(FWidgetMatcher(ViewportWidget.ToSharedRef()), EVisibility::Visible))
 				{
-					FArrangedWidget ArrangedWidget = WidgetPath.FindArrangedWidget(ViewportWidget.ToSharedRef()).Get(FArrangedWidget::NullWidget);
+					FArrangedWidget ArrangedWidget = WidgetPath.FindArrangedWidget(ViewportWidget.ToSharedRef()).Get(FArrangedWidget::GetNullWidget());
 
 					FVector2D WindowClientOffset = ArrangedWidget.Geometry.GetAbsolutePosition();
 					FVector2D WindowClientSize = ArrangedWidget.Geometry.GetAbsoluteSize();
diff --git a/Engine/Source/Editor/UMGEditor/Private/Designer/DesignTimeUtils.cpp b/Engine/Source/Editor/UMGEditor/Private/Designer/DesignTimeUtils.cpp
--- a/Engine/Source/Editor/UMGEditor/Private/Designer/DesignTimeUtils.cpp
+++ b/Engine/Source/Editor/UMGEditor/Private/Designer/DesignTimeUtils.cpp
@@ -20,7 +20,7 @@ bool FDesignTimeUtils::GetArrangedWidget(TSharedRef<SWidget> Widget, FArrangedWi
 	FWidgetPath WidgetPath;
 	if ( FSlateApplication::Get().GeneratePathToWidgetUnchecked(Widget, WidgetPath) )
 	{
-		ArrangedWidget = WidgetPath.FindArrangedWidget(Widget).Get(FArrangedWidget::NullWidget);
+		ArrangedWidget = WidgetPath.FindArrangedWidget(Widget).Get(FArrangedWidget::GetNullWidget());
 		return true;
 	}
 
@@ -40,7 +40,7 @@ bool FDesignTimeUtils::GetArrangedWidgetRelativeToWindow(TSharedRef<SWidget> Wid
 	FWidgetPath WidgetPath;
 	if ( FSlateApplication::Get().GeneratePathToWidgetUnchecked(Widget, WidgetPath) )
 	{
-		ArrangedWidget = WidgetPath.FindArrangedWidget(Widget).Get(FArrangedWidget::NullWidget);
+		ArrangedWidget = WidgetPath.FindArrangedWidget(Widget).Get(FArrangedWidget::GetNullWidget());
 		ArrangedWidget.Geometry.AppendTransform(FSlateLayoutTransform(Inverse(CurrentWindowRef->GetPositionInScreen())));
 		//ArrangedWidget.Geometry.AppendTransform(Inverse(CurrentWindowRef->GetLocalToScreenTransform()));
 		return true;
diff --git a/Engine/Source/Runtime/MovieSceneCapture/Private/FrameGrabber.cpp b/Engine/Source/Runtime/MovieSceneCapture/Private/FrameGrabber.cpp
--- a/Engine/Source/Runtime/MovieSceneCapture/Private/FrameGrabber.cpp
+++ b/Engine/Source/Runtime/MovieSceneCapture/Private/FrameGrabber.cpp
@@ -226,7 +226,7 @@ FFrameGrabber::FFrameGrabber(TSharedRef<FSceneViewport> Viewport, FIntPoint Desi
 			FWidgetPath WidgetPath(Window.ToSharedRef(), JustWindow);
 			if (WidgetPath.ExtendPathTo(FWidgetMatcher(ViewportWidget.ToSharedRef()), EVisibility::Visible))
 			{
-				FArrangedWidget ArrangedWidget = WidgetPath.FindArrangedWidget(ViewportWidget.ToSharedRef()).Get(FArrangedWidget::NullWidget);
+				FArrangedWidget ArrangedWidget = WidgetPath.FindArrangedWidget(ViewportWidget.ToSharedRef()).Get(FArrangedWidget::GetNullWidget());
 
 				FVector2D Position = ArrangedWidget.Geometry.GetAbsolutePosition();
 				FVector2D Size = ArrangedWidget.Geometry.GetAbsoluteSize();
diff --git a/Engine/Source/Runtime/Slate/Private/Framework/Application/SlateApplication.cpp b/Engine/Source/Runtime/Slate/Private/Framework/Application/SlateApplication.cpp
--- a/Engine/Source/Runtime/Slate/Private/Framework/Application/SlateApplication.cpp
+++ b/Engine/Source/Runtime/Slate/Private/Framework/Application/SlateApplication.cpp
@@ -4643,7 +4643,7 @@ bool FSlateApplication::TakeScreenshot(const TSharedRef<SWidget>& Widget, const
 	FWidgetPath WidgetPath;
 	FSlateApplication::Get().GeneratePathToWidgetChecked(Widget, WidgetPath);
 
-	FArrangedWidget ArrangedWidget = WidgetPath.FindArrangedWidget(Widget).Get(FArrangedWidget::NullWidget);
+	FArrangedWidget ArrangedWidget = WidgetPath.FindArrangedWidget(Widget).Get(FArrangedWidget::GetNullWidget());
 	FVector2D Position = ArrangedWidget.Geometry.AbsolutePosition;
 	FVector2D Size = ArrangedWidget.Geometry.GetDrawSize();
 	FVector2D WindowPosition = WidgetWindow->GetPositionInScreen();
diff --git a/Engine/Source/Runtime/Slate/Private/Widgets/SInvalidationPanel.cpp b/Engine/Source/Runtime/Slate/Private/Widgets/SInvalidationPanel.cpp
--- a/Engine/Source/Runtime/Slate/Private/Widgets/SInvalidationPanel.cpp
+++ b/Engine/Source/Runtime/Slate/Private/Widgets/SInvalidationPanel.cpp
@@ -613,7 +613,7 @@ int32 SInvalidationPanel::OnPaint( const FPaintArgs& Args, const FGeometry& Allo
 					FWidgetPath WidgetPath;
 					if ( FSlateApplication::Get().GeneratePathToWidgetUnchecked(SafeInvalidator.ToSharedRef(), WidgetPath, EVisibility::All) )
 					{
-						FArrangedWidget ArrangedWidget = WidgetPath.FindArrangedWidget(SafeInvalidator.ToSharedRef()).Get(FArrangedWidget::NullWidget);
+						FArrangedWidget ArrangedWidget = WidgetPath.FindArrangedWidget(SafeInvalidator.ToSharedRef()).Get(FArrangedWidget::GetNullWidget());
 						ArrangedWidget.Geometry.AppendTransform( FSlateLayoutTransform(Inverse(Args.GetWindowToDesktopTransform())) );
 
 						FSlateDrawElement::MakeBox(
diff --git a/Engine/Source/Runtime/SlateCore/Private/Input/Events.cpp b/Engine/Source/Runtime/SlateCore/Private/Input/Events.cpp
--- a/Engine/Source/Runtime/SlateCore/Private/Input/Events.cpp
+++ b/Engine/Source/Runtime/SlateCore/Private/Input/Events.cpp
@@ -15,7 +15,7 @@ const FTouchKeySet FTouchKeySet::EmptySet(EKeys::Invalid);
 
 FGeometry FInputEvent::FindGeometry(const TSharedRef<SWidget>& WidgetToFind) const
 {
-	return EventPath->FindArrangedWidget(WidgetToFind).Get(FArrangedWidget::NullWidget).Geometry;
+	return EventPath->FindArrangedWidget(WidgetToFind).Get(FArrangedWidget::GetNullWidget()).Geometry;
 }
 
 TSharedRef<SWindow> FInputEvent::GetWindow() const
diff --git a/Engine/Source/Runtime/SlateCore/Private/Layout/ArrangedWidget.cpp b/Engine/Source/Runtime/SlateCore/Private/Layout/ArrangedWidget.cpp
--- a/Engine/Source/Runtime/SlateCore/Private/Layout/ArrangedWidget.cpp
+++ b/Engine/Source/Runtime/SlateCore/Private/Layout/ArrangedWidget.cpp
@@ -4,8 +4,13 @@
 #include "Widgets/SNullWidget.h"
 #include "Widgets/SWidget.h"
 
-
-FArrangedWidget FArrangedWidget::NullWidget(SNullWidget::NullWidget, FGeometry());
+// Use function to initialize because SNullWidget::NullWidget is not statically initialized yet
+const FArrangedWidget& FArrangedWidget::GetNullWidget()
+{
+	static FArrangedWidget NullArrangedWidget(SNullWidget::NullWidget, FGeometry());
+	checkSlow(&SNullWidget::NullWidget.Get() != nullptr);
+	return NullArrangedWidget;
+}
 
 FString FArrangedWidget::ToString( ) const
 {
@@ -14,7 +19,7 @@ FString FArrangedWidget::ToString( ) const
 
 
 FWidgetAndPointer::FWidgetAndPointer()
-: FArrangedWidget(FArrangedWidget::NullWidget)
+: FArrangedWidget(FArrangedWidget::GetNullWidget())
 , PointerPosition(TSharedPtr<FVirtualPointerPosition>())
 {}
 
diff --git a/Engine/Source/Runtime/SlateCore/Private/Layout/WidgetPath.cpp b/Engine/Source/Runtime/SlateCore/Private/Layout/WidgetPath.cpp
--- a/Engine/Source/Runtime/SlateCore/Private/Layout/WidgetPath.cpp
+++ b/Engine/Source/Runtime/SlateCore/Private/Layout/WidgetPath.cpp
@@ -416,7 +416,7 @@ bool FWeakWidgetPath::ContainsWidget( const TSharedRef< const SWidget >& SomeWid
 
 FWidgetPath FWeakWidgetPath::ToNextFocusedPath(EUINavigation NavigationType) const
 {
-	return ToNextFocusedPath(NavigationType, FNavigationReply::Escape(), FArrangedWidget::NullWidget);
+	return ToNextFocusedPath(NavigationType, FNavigationReply::Escape(), FArrangedWidget::GetNullWidget());
 }
 
 FWidgetPath FWeakWidgetPath::ToNextFocusedPath(EUINavigation NavigationType, const FNavigationReply& NavigationReply, const FArrangedWidget& RuleWidget) const
diff --git a/Engine/Source/Runtime/SlateCore/Public/Layout/ArrangedWidget.h b/Engine/Source/Runtime/SlateCore/Public/Layout/ArrangedWidget.h
--- a/Engine/Source/Runtime/SlateCore/Public/Layout/ArrangedWidget.h
+++ b/Engine/Source/Runtime/SlateCore/Public/Layout/ArrangedWidget.h
@@ -21,7 +21,7 @@ public:
 		, Widget(InWidget)
 	{ }
 
-	SLATECORE_API static FArrangedWidget NullWidget;
+	SLATECORE_API static const FArrangedWidget& GetNullWidget();
 
 public:
 
