Fixed some places that were incorrectly using the native culture to avoid grouping

diff --git a/Engine/Source/Editor/UMGEditor/Private/Designer/SDesignerView.cpp b/Engine/Source/Editor/UMGEditor/Private/Designer/SDesignerView.cpp
--- a/Engine/Source/Editor/UMGEditor/Private/Designer/SDesignerView.cpp
+++ b/Engine/Source/Editor/UMGEditor/Private/Designer/SDesignerView.cpp
@@ -3100,10 +3100,9 @@ FReply SDesignerView::OnDrop(const FGeometry& MyGeometry, const FDragDropEvent&
 
 FText SDesignerView::GetResolutionText(int32 Width, int32 Height, const FString& AspectRatio) const
 {
-	FInternationalization& I18N = FInternationalization::Get();
 	FFormatNamedArguments Args;
-	Args.Add(TEXT("Width"), FText::AsNumber(Width, nullptr, I18N.GetInvariantCulture()));
-	Args.Add(TEXT("Height"), FText::AsNumber(Height, nullptr, I18N.GetInvariantCulture()));
+	Args.Add(TEXT("Width"), FText::AsNumber(Width, &FNumberFormattingOptions::DefaultNoGrouping()));
+	Args.Add(TEXT("Height"), FText::AsNumber(Height, &FNumberFormattingOptions::DefaultNoGrouping()));
 	Args.Add(TEXT("AspectRatio"), FText::FromString(AspectRatio));
 
 	return FText::Format(LOCTEXT("CommonResolutionFormat", "{Width} x {Height} ({AspectRatio})"), Args);
@@ -3116,9 +3115,7 @@ FText SDesignerView::GetCurrentResolutionText() const
 
 FText SDesignerView::GetCurrentDPIScaleText() const
 {
-	FInternationalization& I18N = FInternationalization::Get();
-
-	FNumberFormattingOptions Options;
+	FNumberFormattingOptions Options = FNumberFormattingOptions::DefaultNoGrouping();
 	Options.MinimumIntegralDigits = 1;
 	Options.MaximumFractionalDigits = 2;
 	Options.MinimumFractionalDigits = 1;
@@ -3134,7 +3131,7 @@ FText SDesignerView::GetCurrentDPIScaleText() const
 		}
 	}
 
-	FText DPIString = FText::AsNumber(GetPreviewDPIScale(), &Options, I18N.GetInvariantCulture());
+	FText DPIString = FText::AsNumber(GetPreviewDPIScale(), &Options);
 	return FText::Format(LOCTEXT("CurrentDPIScaleFormat", "DPI Scale {0}"), DPIString);
 }
 
@@ -3156,14 +3153,12 @@ FSlateColor SDesignerView::GetCurrentDPIScaleColor() const
 
 FText SDesignerView::GetCurrentScaleFactorText() const
 {
-	FInternationalization& I18N = FInternationalization::Get();
-
-	FNumberFormattingOptions Options;
+	FNumberFormattingOptions Options = FNumberFormattingOptions::DefaultNoGrouping();
 	Options.MinimumIntegralDigits = 1;
 	Options.MaximumFractionalDigits = 2;
 	Options.MinimumFractionalDigits = 1;
 
-	FText DPIString = FText::AsNumber(ScaleFactor, &Options, I18N.GetInvariantCulture());
+	FText DPIString = FText::AsNumber(ScaleFactor, &Options);
 	return FText::Format(LOCTEXT("CurrentContentScale", "Device Content Scale {0}"), DPIString);
 }
 
