Added selection dialog for choosing a root widget when creating a widget blueprint, can be enabled in project settings. Planning to make another pass once I gather some feedback to refine the set of common widgets and their descriptions.

#rn Added a selection dialog for choosing a root widget when creating a new widget blueprint. To enable, check Use Widget Template Selector in your project settings.

diff --git a/Engine/Source/Editor/UMGEditor/Classes/WidgetBlueprintFactory.h b/Engine/Source/Editor/UMGEditor/Classes/WidgetBlueprintFactory.h
--- a/Engine/Source/Editor/UMGEditor/Classes/WidgetBlueprintFactory.h
+++ b/Engine/Source/Editor/UMGEditor/Classes/WidgetBlueprintFactory.h
@@ -28,4 +28,6 @@ class UWidgetBlueprintFactory : public UFactory
 	virtual UObject* FactoryCreateNew(UClass* Class, UObject* InParent, FName Name, EObjectFlags Flags, UObject* Context, FFeedbackContext* Warn, FName CallingContext) override;
 	virtual UObject* FactoryCreateNew(UClass* Class, UObject* InParent, FName Name, EObjectFlags Flags, UObject* Context, FFeedbackContext* Warn) override;
 	//~ Begin UFactory Interface	
+
+	UClass* RootWidgetClass;
 };
diff --git a/Engine/Source/Editor/UMGEditor/Private/Settings/UMGEditorProjectSettings.cpp b/Engine/Source/Editor/UMGEditor/Private/Settings/UMGEditorProjectSettings.cpp
--- a/Engine/Source/Editor/UMGEditor/Private/Settings/UMGEditorProjectSettings.cpp
+++ b/Engine/Source/Editor/UMGEditor/Private/Settings/UMGEditorProjectSettings.cpp
@@ -14,6 +14,7 @@ UUMGEditorProjectSettings::UUMGEditorProjectSettings()
 	bShowWidgetsFromEngineContent = false;
 	bShowWidgetsFromDeveloperContent = true;
 
+	bUseWidgetTemplateSelector = false;
 	DefaultRootWidget = UCanvasPanel::StaticClass();
 
 	// Deprecated
diff --git a/Engine/Source/Editor/UMGEditor/Private/WidgetBlueprintFactory.cpp b/Engine/Source/Editor/UMGEditor/Private/WidgetBlueprintFactory.cpp
--- a/Engine/Source/Editor/UMGEditor/Private/WidgetBlueprintFactory.cpp
+++ b/Engine/Source/Editor/UMGEditor/Private/WidgetBlueprintFactory.cpp
@@ -10,6 +10,13 @@
 
 #include "Blueprint/WidgetTree.h"
 #include "UMGEditorProjectSettings.h"
+#include "ClassViewerModule.h"
+#include "Kismet2/SClassPickerDialog.h"
+#include "ClassViewerFilter.h"
+#include "Components/CanvasPanel.h"
+#include "Components/HorizontalBox.h"
+#include "Components/VerticalBox.h"
+#include "Components/GridPanel.h"
 
 #define LOCTEXT_NAMESPACE "UWidgetBlueprintFactory"
 
@@ -17,6 +24,29 @@
 	UWidgetBlueprintFactory implementation.
 ------------------------------------------------------------------------------*/
 
+class FWidgetClassFilter : public IClassViewerFilter
+{
+public:
+	/** All children of these classes will be included unless filtered out by another setting. */
+	TSet <const UClass*> AllowedChildrenOfClasses;
+
+	/** Disallowed class flags. */
+	EClassFlags DisallowedClassFlags;
+
+	virtual bool IsClassAllowed(const FClassViewerInitializationOptions& InInitOptions, const UClass* InClass, TSharedRef< FClassViewerFilterFuncs > InFilterFuncs) override
+	{
+		return !InClass->HasAnyClassFlags(DisallowedClassFlags)
+			&& InFilterFuncs->IfInChildOfClassesSet(AllowedChildrenOfClasses, InClass) != EFilterReturn::Failed;
+	}
+	
+	virtual bool IsUnloadedClassAllowed(const FClassViewerInitializationOptions& InInitOptions, const TSharedRef< const IUnloadedBlueprintData > InUnloadedClassData, TSharedRef< FClassViewerFilterFuncs > InFilterFuncs) override
+	{
+		return !InUnloadedClassData->HasAnyClassFlags(DisallowedClassFlags)
+			&& InFilterFuncs->IfInChildOfClassesSet(AllowedChildrenOfClasses, InUnloadedClassData) != EFilterReturn::Failed;
+	}
+
+};
+
 UWidgetBlueprintFactory::UWidgetBlueprintFactory(const FObjectInitializer& ObjectInitializer)
 	: Super(ObjectInitializer)
 {
@@ -28,7 +58,31 @@ UWidgetBlueprintFactory::UWidgetBlueprintFactory(const FObjectInitializer& Objec
 
 bool UWidgetBlueprintFactory::ConfigureProperties()
 {
-	// TODO Make config dialog like the anim blueprints.
+	if (GetDefault<UUMGEditorProjectSettings>()->bUseWidgetTemplateSelector)
+	{
+		// Load the classviewer module to display a class picker
+		FClassViewerModule& ClassViewerModule = FModuleManager::LoadModuleChecked<FClassViewerModule>("ClassViewer");
+
+		// Fill in options
+		FClassViewerInitializationOptions Options;
+		Options.Mode = EClassViewerMode::ClassPicker;
+		Options.bShowNoneOption = true;
+
+		Options.ExtraPickerCommonClasses.Add(UHorizontalBox::StaticClass());
+		Options.ExtraPickerCommonClasses.Add(UVerticalBox::StaticClass());
+		Options.ExtraPickerCommonClasses.Add(UGridPanel::StaticClass());
+		Options.ExtraPickerCommonClasses.Add(UCanvasPanel::StaticClass());
+
+		TSharedPtr<FWidgetClassFilter> Filter = MakeShareable(new FWidgetClassFilter);
+		Options.ClassFilter = Filter;
+
+		Filter->DisallowedClassFlags = CLASS_Abstract | CLASS_Deprecated | CLASS_NewerVersionExists;
+		Filter->AllowedChildrenOfClasses.Add(UPanelWidget::StaticClass());
+
+		const FText TitleText = LOCTEXT("CreateWidgetBlueprint", "Pick Root Widget for New Widget Blueprint");
+		return SClassPickerDialog::PickClass(TitleText, Options, RootWidgetClass, UPanelWidget::StaticClass());
+
+	}
 	return true;
 }
 
@@ -57,14 +111,19 @@ UObject* UWidgetBlueprintFactory::FactoryCreateNew(UClass* Class, UObject* InPar
 	}
 	else
 	{
+		if (!GetDefault<UUMGEditorProjectSettings>()->bUseWidgetTemplateSelector)
+		{
+			RootWidgetClass = GetDefault<UUMGEditorProjectSettings>()->DefaultRootWidget;
+		}
+
 		UWidgetBlueprint* NewBP = CastChecked<UWidgetBlueprint>(FKismetEditorUtilities::CreateBlueprint(ParentClass, InParent, Name, BlueprintType, UWidgetBlueprint::StaticClass(), UWidgetBlueprintGeneratedClass::StaticClass(), CallingContext));
 
 		// Create the desired root widget specified by the project
 		if ( NewBP->WidgetTree->RootWidget == nullptr )
 		{
-			if (TSubclassOf<UPanelWidget> RootWidgetClass = GetDefault<UUMGEditorProjectSettings>()->DefaultRootWidget)
+			if (TSubclassOf<UPanelWidget> RootWidgetPanel = RootWidgetClass)
 			{
-				UWidget* Root = NewBP->WidgetTree->ConstructWidget<UWidget>(RootWidgetClass);
+				UWidget* Root = NewBP->WidgetTree->ConstructWidget<UWidget>(RootWidgetPanel);
 				NewBP->WidgetTree->RootWidget = Root;
 			}
 		}
diff --git a/Engine/Source/Editor/UMGEditor/Public/UMGEditorProjectSettings.h b/Engine/Source/Editor/UMGEditor/Public/UMGEditorProjectSettings.h
--- a/Engine/Source/Editor/UMGEditor/Public/UMGEditorProjectSettings.h
+++ b/Engine/Source/Editor/UMGEditor/Public/UMGEditorProjectSettings.h
@@ -175,6 +175,10 @@ public:
 
 public:
 
+	/** Enables a dialog that lets you select a root widget before creating a widget blueprint */
+	UPROPERTY(EditAnywhere, config, Category = Designer)
+	bool bUseWidgetTemplateSelector;
+
 	/** The panel widget to place at the root of all newly constructed widget blueprints. Can be empty. */
 	UPROPERTY(EditAnywhere, config, Category = Designer)
 	TSubclassOf<UPanelWidget> DefaultRootWidget;
diff --git a/Engine/Source/Editor/UnrealEd/Private/SClassPickerDialog.cpp b/Engine/Source/Editor/UnrealEd/Private/SClassPickerDialog.cpp
--- a/Engine/Source/Editor/UnrealEd/Private/SClassPickerDialog.cpp
+++ b/Engine/Source/Editor/UnrealEd/Private/SClassPickerDialog.cpp
@@ -32,6 +32,7 @@
 void SClassPickerDialog::Construct(const FArguments& InArgs)
 {
 	WeakParentWindow = InArgs._ParentWindow;
+	bAllowNone = InArgs._Options.bShowNoneOption;
 
 	bPressedOk = false;
 	ChosenClass = NULL;
@@ -279,7 +280,7 @@ FReply SClassPickerDialog::OnDefaultClassPicked(UClass* InChosenClass)
 
 FReply SClassPickerDialog::OnClassPickerConfirmed()
 {
-	if (ChosenClass == NULL)
+	if (!bAllowNone && ChosenClass == NULL)
 	{
 		FMessageDialog::Open(EAppMsgType::Ok, NSLOCTEXT("EditorFactories", "MustChooseClassWarning", "You must choose a class."));
 	}
@@ -333,7 +334,7 @@ void SClassPickerDialog::OnCustomAreaExpansionChanged(bool bExpanded)
 EVisibility SClassPickerDialog::GetSelectButtonVisibility() const
 {
 	EVisibility ButtonVisibility = EVisibility::Hidden;
-	if( ChosenClass != nullptr )
+	if(bAllowNone || ChosenClass != nullptr )
 	{
 		ButtonVisibility = EVisibility::Visible;
 	}
diff --git a/Engine/Source/Editor/UnrealEd/Public/Kismet2/SClassPickerDialog.h b/Engine/Source/Editor/UnrealEd/Public/Kismet2/SClassPickerDialog.h
--- a/Engine/Source/Editor/UnrealEd/Public/Kismet2/SClassPickerDialog.h
+++ b/Engine/Source/Editor/UnrealEd/Public/Kismet2/SClassPickerDialog.h
@@ -79,4 +79,7 @@ private:
 
 	/**  An array of default classes used in the dialog **/
 	TArray< TSharedPtr<FClassPickerDefaults>  > AssetDefaultClasses;
+
+	/** Determines if selecting no class is permitted **/
+	bool bAllowNone;
 };
diff --git a/Engine/Source/Runtime/UMG/Public/Components/CanvasPanel.h b/Engine/Source/Runtime/UMG/Public/Components/CanvasPanel.h
--- a/Engine/Source/Runtime/UMG/Public/Components/CanvasPanel.h
+++ b/Engine/Source/Runtime/UMG/Public/Components/CanvasPanel.h
@@ -21,7 +21,7 @@ class UCanvasPanelSlot;
  * * Absolute Layout
  * * Anchors
  */
-UCLASS()
+UCLASS(meta = (ShortTooltip = "A designer-friendly panel useful for laying out top-level widgets. Use sparingly."))
 class UMG_API UCanvasPanel : public UPanelWidget
 {
 	GENERATED_UCLASS_BODY()
diff --git a/Engine/Source/Runtime/UMG/Public/Components/HorizontalBox.h b/Engine/Source/Runtime/UMG/Public/Components/HorizontalBox.h
--- a/Engine/Source/Runtime/UMG/Public/Components/HorizontalBox.h
+++ b/Engine/Source/Runtime/UMG/Public/Components/HorizontalBox.h
@@ -17,7 +17,7 @@ class UHorizontalBoxSlot;
  * * Many Children
  * * Flow Horizontal
  */
-UCLASS()
+UCLASS(meta = (ShortTooltip = "A layout panel for automatically laying child widgets out horizontally"))
 class UMG_API UHorizontalBox : public UPanelWidget
 {
 	GENERATED_UCLASS_BODY()
diff --git a/Engine/Source/Runtime/UMG/Public/Components/VerticalBox.h b/Engine/Source/Runtime/UMG/Public/Components/VerticalBox.h
--- a/Engine/Source/Runtime/UMG/Public/Components/VerticalBox.h
+++ b/Engine/Source/Runtime/UMG/Public/Components/VerticalBox.h
@@ -18,7 +18,7 @@ class UVerticalBoxSlot;
  * * Many Children
  * * Flows Vertical
  */
-UCLASS()
+UCLASS(meta = (ShortTooltip = "A layout panel for automatically laying child widgets out vertically"))
 class UMG_API UVerticalBox : public UPanelWidget
 {
 	GENERATED_UCLASS_BODY()
