add missing header in UMGEditor WidgetBlueprintFactory
#rnx

diff --git a/Engine/Source/Editor/UMGEditor/Private/WidgetBlueprintFactory.cpp b/Engine/Source/Editor/UMGEditor/Private/WidgetBlueprintFactory.cpp
--- a/Engine/Source/Editor/UMGEditor/Private/WidgetBlueprintFactory.cpp
+++ b/Engine/Source/Editor/UMGEditor/Private/WidgetBlueprintFactory.cpp
@@ -7,6 +7,7 @@
 #include "Blueprint/WidgetBlueprintGeneratedClass.h"
 #include "WidgetBlueprint.h"
 #include "Kismet2/KismetEditorUtilities.h"
+#include "Modules/ModuleManager.h"
 
 #include "Blueprint/WidgetTree.h"
 #include "UMGEditorProjectSettings.h"
