Fix to IsDesignTime() returning false when called inside Initialize overide. In this specific case, we need to change the call To SetDesignTimeFlag() before duplication to make sure the design time flags is set before the Initialize called during duplication.
Fix to PanelWidget SetDesignerFlags() not being virtual/overide like it should be.

diff --git a/Engine/Source/Editor/UMGEditor/Private/WidgetBlueprintEditor.cpp b/Engine/Source/Editor/UMGEditor/Private/WidgetBlueprintEditor.cpp
--- a/Engine/Source/Editor/UMGEditor/Private/WidgetBlueprintEditor.cpp
+++ b/Engine/Source/Editor/UMGEditor/Private/WidgetBlueprintEditor.cpp
@@ -1002,6 +1002,16 @@ void FWidgetBlueprintEditor::UpdatePreview(UBlueprint* InBlueprint, bool bInForc
 			// The preview widget should not be transactional.
 			PreviewUserWidget->ClearFlags(RF_Transactional);
 
+			// Establish the widget as being in design time before initializing and before duplication 
+            // (so that IsDesignTime is reliable within both calls to Initialize)
+            // The preview widget is also the outer widget that will update all child flags
+			PreviewUserWidget->SetDesignerFlags(GetCurrentDesignerFlags());
+
+			if ( ULocalPlayer* Player = PreviewScene.GetWorld()->GetFirstLocalPlayerFromController() )
+			{
+				PreviewUserWidget->SetPlayerContext(FLocalPlayerContext(Player));
+			}
+
 			UWidgetTree* LatestWidgetTree = PreviewBlueprint->WidgetTree;
 
 			// If there is no RootWidget, we look for a WidgetTree in the parents classes until we find one.
@@ -1012,21 +1022,11 @@ void FWidgetBlueprintEditor::UpdatePreview(UBlueprint* InBlueprint, bool bInForc
 				{
 					LatestWidgetTree = BGClass->WidgetTree;
 				}
-			 }
+			}
 
 			// Update the widget tree directly to match the blueprint tree.  That way the preview can update
 			// without needing to do a full recompile.
 			PreviewUserWidget->DuplicateAndInitializeFromWidgetTree(LatestWidgetTree);
-
-			if ( ULocalPlayer* Player = PreviewScene.GetWorld()->GetFirstLocalPlayerFromController() )
-			{
-				PreviewUserWidget->SetPlayerContext(FLocalPlayerContext(Player));
-			}
-
-			// Establish the widget as being in design time before initializing (so that IsDesignTime is reliable within Initialize)
-			PreviewUserWidget->SetDesignerFlags(GetCurrentDesignerFlags());
-
-			PreviewUserWidget->Initialize();
 		}
 
 		// Store a reference to the preview actor.
diff --git a/Engine/Source/Runtime/UMG/Private/UserWidget.cpp b/Engine/Source/Runtime/UMG/Private/UserWidget.cpp
--- a/Engine/Source/Runtime/UMG/Private/UserWidget.cpp
+++ b/Engine/Source/Runtime/UMG/Private/UserWidget.cpp
@@ -453,6 +453,14 @@ void UUserWidget::PostDuplicate(bool bDuplicateForPIE)
 	
 	if ( bInitializingFromWidgetTree )
 	{
+		// If this is a sub-widget of another UserWidget, default designer flags to match those of the owning widget before initialize.
+		if (UUserWidget* OwningUserWidget = GetTypedOuter<UUserWidget>())
+		{
+#if WITH_EDITOR
+			SetDesignerFlags(OwningUserWidget->GetDesignerFlags());	
+#endif
+			SetPlayerContext(OwningUserWidget->GetPlayerContext());
+		}
 		Initialize();
 	}
 }
diff --git a/Engine/Source/Runtime/UMG/Public/Components/PanelWidget.h b/Engine/Source/Runtime/UMG/Public/Components/PanelWidget.h
--- a/Engine/Source/Runtime/UMG/Public/Components/PanelWidget.h
+++ b/Engine/Source/Runtime/UMG/Public/Components/PanelWidget.h
@@ -88,7 +88,7 @@ public:
 	 */
 	void ShiftChild(int32 Index, UWidget* Child);
 
-	void SetDesignerFlags(EWidgetDesignFlags NewFlags);
+	virtual void SetDesignerFlags(EWidgetDesignFlags NewFlags) override;
 
 #endif
 
