Fix to a crash when we delete the content in a named slot or drag it out of the named slot.

diff --git a/Engine/Source/Editor/UMGEditor/Private/WidgetBlueprintEditorUtils.cpp b/Engine/Source/Editor/UMGEditor/Private/WidgetBlueprintEditorUtils.cpp
--- a/Engine/Source/Editor/UMGEditor/Private/WidgetBlueprintEditorUtils.cpp
+++ b/Engine/Source/Editor/UMGEditor/Private/WidgetBlueprintEditorUtils.cpp
@@ -532,7 +532,10 @@ bool FWidgetBlueprintEditorUtils::ReplaceNamedSlotHostContent(UWidget* WidgetTem
 		{
 			if (SlotContent == WidgetTemplate)
 			{
-				NewContentWidget->Modify();
+				if (NewContentWidget)
+				{
+					NewContentWidget->Modify();
+				}
 				NamedSlotHost.GetObject()->Modify();
 				NamedSlotHost->SetContentForSlot(SlotName, NewContentWidget);
 				return true;
