Fix to drag and drop of widgets not working in the Designer view

diff --git a/Engine/Source/Editor/UMGEditor/Private/WidgetBlueprintEditor.cpp b/Engine/Source/Editor/UMGEditor/Private/WidgetBlueprintEditor.cpp
--- a/Engine/Source/Editor/UMGEditor/Private/WidgetBlueprintEditor.cpp
+++ b/Engine/Source/Editor/UMGEditor/Private/WidgetBlueprintEditor.cpp
@@ -1027,6 +1027,10 @@ void FWidgetBlueprintEditor::UpdatePreview(UBlueprint* InBlueprint, bool bInForc
 			// Update the widget tree directly to match the blueprint tree.  That way the preview can update
 			// without needing to do a full recompile.
 			PreviewUserWidget->DuplicateAndInitializeFromWidgetTree(LatestWidgetTree);
+
+			// Establish the widget as being in design time before initializing (so that IsDesignTime is reliable within Initialize)
+            // We have to call it to make sure that all the WidgetTree had the DesignerFlags set correctly
+			PreviewUserWidget->SetDesignerFlags(GetCurrentDesignerFlags());
 		}
 
 		// Store a reference to the preview actor.
