Sequencer: RefreshTree to refilter UMG selection when ShowOnlySelectedNodes is active

diff --git a/Engine/Source/Editor/UMGEditor/Private/WidgetBlueprintEditor.cpp b/Engine/Source/Editor/UMGEditor/Private/WidgetBlueprintEditor.cpp
--- a/Engine/Source/Editor/UMGEditor/Private/WidgetBlueprintEditor.cpp
+++ b/Engine/Source/Editor/UMGEditor/Private/WidgetBlueprintEditor.cpp
@@ -1751,6 +1751,11 @@ void FWidgetBlueprintEditor::SyncSequencerSelectionToSelectedWidgets()
 
 	TGuardValue<bool> Guard(bUpdatingSequencerSelection, true);
 
+	if (GetSequencer()->GetSequencerSettings()->GetShowSelectedNodesOnly())
+	{
+		GetSequencer()->RefreshTree();
+	}
+
 	GetSequencer()->ExternalSelectionHasChanged();
 }
 
