UMG Editor: Add named slot content to Replace With context menu in Hierarchy view

#rn UMG Editor: The Replace with context menu in the Hierarchy view now support to replace a widget with the content of one of it's Named Slot.

diff --git a/Engine/Source/Editor/UMGEditor/Private/WidgetBlueprintEditorUtils.cpp b/Engine/Source/Editor/UMGEditor/Private/WidgetBlueprintEditorUtils.cpp
--- a/Engine/Source/Editor/UMGEditor/Private/WidgetBlueprintEditorUtils.cpp
+++ b/Engine/Source/Editor/UMGEditor/Private/WidgetBlueprintEditorUtils.cpp
@@ -36,6 +36,7 @@
 #include "Blueprint/WidgetNavigation.h"
 #include "Subsystems/AssetEditorSubsystem.h"
 #include "UObject/ScriptInterface.h"
+#include "Components/NamedSlotInterface.h"
 
 #define LOCTEXT_NAMESPACE "UMG"
 
@@ -699,6 +700,27 @@ void FWidgetBlueprintEditorUtils::BuildReplaceWithMenu(FMenuBuilder& Menu, TShar
 					));
 
 				Menu.AddMenuSeparator();
+			}			
+			if (TScriptInterface<INamedSlotInterface> NamedSlotHost = TScriptInterface<INamedSlotInterface>(Widget.GetTemplate()))
+			{								
+				TArray<FName> SlotNames;
+				NamedSlotHost->GetSlotNames(SlotNames);
+				for (const FName SlotName : SlotNames)
+				{
+					const FText SlotNameTxt = FText::FromString(SlotName.ToString());
+					if (UWidget* Content = NamedSlotHost->GetContentForSlot(SlotName))
+					{
+						Menu.AddMenuEntry(
+							FText::Format(LOCTEXT("ReplaceWithNamedSlot", "Replace With '{0}'"), SlotNameTxt),
+							FText::Format(LOCTEXT("ReplaceWithChildTooltip", "Remove this widget and insert '{0}' content into the parent."), SlotNameTxt),
+							FSlateIcon(),
+							FUIAction(
+								FExecuteAction::CreateStatic(&FWidgetBlueprintEditorUtils::ReplaceWidgetWithNamedSlot, BlueprintEditor, BP, Widget, SlotName),
+								FCanExecuteAction()
+							));
+					}
+				}
+				Menu.AddMenuSeparator();
 			}
 		}
 
@@ -915,6 +937,49 @@ void FWidgetBlueprintEditorUtils::ReplaceWidgetWithChildren(TSharedRef<FWidgetBl
 	}
 }
 
+void FWidgetBlueprintEditorUtils::ReplaceWidgetWithNamedSlot(TSharedRef<FWidgetBlueprintEditor> BlueprintEditor, UWidgetBlueprint* BP, FWidgetReference Widget, FName NamedSlot)
+{
+	UWidget* WidgetTemplate = Widget.GetTemplate();
+	if (INamedSlotInterface* ExistingNamedSlotContainerTemplate = Cast<INamedSlotInterface>(WidgetTemplate))
+	{
+		UWidget* NamedSlotContentTemplate = ExistingNamedSlotContainerTemplate->GetContentForSlot(NamedSlot);
+
+		FScopedTransaction Transaction(LOCTEXT("ReplaceWidgets", "Replace Widgets"));
+
+		WidgetTemplate->Modify();
+		NamedSlotContentTemplate->Modify();
+
+		if (UPanelWidget* PanelParentTemplate = WidgetTemplate->GetParent())
+		{
+			PanelParentTemplate->Modify();
+
+			NamedSlotContentTemplate->RemoveFromParent();
+			PanelParentTemplate->ReplaceChild(WidgetTemplate, NamedSlotContentTemplate);
+		}
+		else if (WidgetTemplate == BP->WidgetTree->RootWidget)
+		{
+			NamedSlotContentTemplate->RemoveFromParent();
+
+			BP->WidgetTree->Modify();
+			BP->WidgetTree->RootWidget = NamedSlotContentTemplate;
+		}
+		else if (TScriptInterface<INamedSlotInterface> NamedSlotHost = FindNamedSlotHostForContent(WidgetTemplate, BP->WidgetTree))
+		{
+			ReplaceNamedSlotHostContent(WidgetTemplate, NamedSlotHost, NamedSlotContentTemplate);
+		}
+		else
+		{
+			Transaction.Cancel();
+			return;
+		}
+
+		// Rename the removed widget to the transient package so that it doesn't conflict with future widgets sharing the same name.
+		WidgetTemplate->Rename(nullptr, nullptr);
+
+		FBlueprintEditorUtils::MarkBlueprintAsStructurallyModified(BP);
+	}
+}
+
 void FWidgetBlueprintEditorUtils::ReplaceWidgets(TSharedRef<FWidgetBlueprintEditor> BlueprintEditor, UWidgetBlueprint* BP, TSet<FWidgetReference> Widgets, UClass* WidgetClass)
 {
 	const FScopedTransaction Transaction(LOCTEXT("ReplaceWidgets", "Replace Widgets"));
diff --git a/Engine/Source/Editor/UMGEditor/Private/WidgetBlueprintEditorUtils.h b/Engine/Source/Editor/UMGEditor/Private/WidgetBlueprintEditorUtils.h
--- a/Engine/Source/Editor/UMGEditor/Private/WidgetBlueprintEditorUtils.h
+++ b/Engine/Source/Editor/UMGEditor/Private/WidgetBlueprintEditorUtils.h
@@ -85,6 +85,8 @@ private:
 
 	static void ReplaceWidgetWithChildren(TSharedRef<FWidgetBlueprintEditor> BlueprintEditor, UWidgetBlueprint* BP, FWidgetReference Widget);
 
+	static void ReplaceWidgetWithNamedSlot(TSharedRef<FWidgetBlueprintEditor> BlueprintEditor, UWidgetBlueprint* BP, FWidgetReference Widget, FName NamedSlot);
+
 	static void ReplaceWidgets(TSharedRef<FWidgetBlueprintEditor> BlueprintEditor, UWidgetBlueprint* BP, TSet<FWidgetReference> Widgets, UClass* WidgetClass);
 
 	static FString FindNextValidName(UWidgetTree* WidgetTree, const FString& Name);
