Fixed a crash when compiling a running Editor Utility Blueprint containing a details view.

# Rebuilding the details view widget after a compile is now synchronous operation. Compiling multiple times would cause the underlying widgets to be cleaned up on the main thread from the compile. The async details view task was having widgets it created destroyed by the main thread's compile before finishing the task.

diff --git a/Engine/Source/Editor/UMGEditor/Private/Components/DetailsView.cpp b/Engine/Source/Editor/UMGEditor/Private/Components/DetailsView.cpp
--- a/Engine/Source/Editor/UMGEditor/Private/Components/DetailsView.cpp
+++ b/Engine/Source/Editor/UMGEditor/Private/Components/DetailsView.cpp
@@ -115,7 +115,7 @@ void UDetailsView::OnObjectChanged()
 	}
 	else
 	{
-		AsynBuildContentWidget();
+		BuildContentWidget();
 	}
 }
 
