UMG - The widget blueprint compiler now does additional checking to make sure you've correctly setup your UPROPERTY if you're using WidgetBindAnim, so that the property is transient.

diff --git a/Engine/Source/Editor/UMGEditor/Private/WidgetBlueprintCompiler.cpp b/Engine/Source/Editor/UMGEditor/Private/WidgetBlueprintCompiler.cpp
--- a/Engine/Source/Editor/UMGEditor/Private/WidgetBlueprintCompiler.cpp
+++ b/Engine/Source/Editor/UMGEditor/Private/WidgetBlueprintCompiler.cpp
@@ -814,6 +814,12 @@ void FWidgetBlueprintCompilerContext::FinishCompilingClass(UClass* Class)
 						MessageLog.Error(*RequiredWidgetAnimNotBoundError.ToString(), WidgetAnimProperty);
 					}
 				}
+
+				if (!WidgetAnimProperty->HasAnyPropertyFlags(RF_Transient))
+				{
+					const FText BindWidgetAnimTransientError = LOCTEXT("BindWidgetAnimTransient", "The property @@ uses BindWidgetAnim, but isn't Transient!");
+					MessageLog.Error(*BindWidgetAnimTransientError.ToString(), WidgetAnimProperty);
+				}
 			}
 		}
 	}
