Blueprint CPF instead of RF.

diff --git a/Engine/Source/Editor/UMGEditor/Private/WidgetBlueprintCompiler.cpp b/Engine/Source/Editor/UMGEditor/Private/WidgetBlueprintCompiler.cpp
--- a/Engine/Source/Editor/UMGEditor/Private/WidgetBlueprintCompiler.cpp
+++ b/Engine/Source/Editor/UMGEditor/Private/WidgetBlueprintCompiler.cpp
@@ -815,7 +815,7 @@ void FWidgetBlueprintCompilerContext::FinishCompilingClass(UClass* Class)
 					}
 				}
 
-				if (!WidgetAnimProperty->HasAnyPropertyFlags(RF_Transient))
+				if (!WidgetAnimProperty->HasAnyPropertyFlags(CPF_Transient))
 				{
 					const FText BindWidgetAnimTransientError = LOCTEXT("BindWidgetAnimTransient", "The property @@ uses BindWidgetAnim, but isn't Transient!");
 					MessageLog.Error(*BindWidgetAnimTransientError.ToString(), WidgetAnimProperty);
