Prevent circular dependency with Dynamic Entry Box

diff --git a/Engine/Source/Editor/UMGEditor/Private/WidgetBlueprintCompiler.cpp b/Engine/Source/Editor/UMGEditor/Private/WidgetBlueprintCompiler.cpp
--- a/Engine/Source/Editor/UMGEditor/Private/WidgetBlueprintCompiler.cpp
+++ b/Engine/Source/Editor/UMGEditor/Private/WidgetBlueprintCompiler.cpp
@@ -831,11 +831,17 @@ void FWidgetBlueprintCompilerContext::FinishCompilingClass(UClass* Class)
 class FBlueprintCompilerLog : public IWidgetCompilerLog
 {
 public:
-	FBlueprintCompilerLog(FCompilerResultsLog& InMessageLog)
+	FBlueprintCompilerLog(FCompilerResultsLog& InMessageLog, TSubclassOf<UUserWidget> InClassContext)
 		: MessageLog(InMessageLog)
+		, ClassContext(InClassContext)
 	{
 	}
 
+	virtual TSubclassOf<UUserWidget> GetContextClass() const override
+	{
+		return ClassContext;
+	}
+
 	virtual void InternalLogMessage(TSharedRef<FTokenizedMessage>& InMessage) override
 	{
 		MessageLog.AddTokenizedMessage(InMessage);
@@ -844,6 +850,7 @@ public:
 private:
 	// Compiler message log (errors, warnings, notes)
 	FCompilerResultsLog& MessageLog;
+	TSubclassOf<UUserWidget> ClassContext;
 };
 
 
@@ -870,7 +877,7 @@ void FWidgetBlueprintCompilerContext::OnPostCDOCompiled()
 
 	if (!Blueprint->bIsRegeneratingOnLoad && bIsFullCompile)
 	{
-		FBlueprintCompilerLog BlueprintLog(MessageLog);
+		FBlueprintCompilerLog BlueprintLog(MessageLog, WidgetClass);
 		WidgetClass->GetDefaultObject<UUserWidget>()->ValidateBlueprint(*WidgetBP->WidgetTree, BlueprintLog);
 
 		if (MessageLog.NumErrors == 0 && WidgetClass->bAllowTemplate)
diff --git a/Engine/Source/Runtime/UMG/Private/Components/DynamicEntryBox.cpp b/Engine/Source/Runtime/UMG/Private/Components/DynamicEntryBox.cpp
--- a/Engine/Source/Runtime/UMG/Private/Components/DynamicEntryBox.cpp
+++ b/Engine/Source/Runtime/UMG/Private/Components/DynamicEntryBox.cpp
@@ -31,6 +31,10 @@ void UDynamicEntryBox::ValidateCompiledDefaults(IWidgetCompilerLog& CompileLog)
 	{
 		CompileLog.Error(FText::Format(LOCTEXT("Error_DynamicEntryBox_MissingEntryClass", "{0} has no EntryWidgetClass specified - required for any Dynamic Entry Box to function."), FText::FromString(GetName())));
 	}
+	else if (CompileLog.GetContextClass() && EntryWidgetClass->IsChildOf(CompileLog.GetContextClass()))
+	{
+		CompileLog.Error(FText::Format(LOCTEXT("Error_DynamicEntryBox_RecursiveEntryClass", "{0} has a recursive EntryWidgetClass specified (reference itself)."), FText::FromString(GetName())));
+	}
 }
 #endif
 
@@ -42,7 +46,7 @@ void UDynamicEntryBox::SynchronizeProperties()
 #if WITH_EDITORONLY_DATA
 	if (IsDesignTime() && MyPanelWidget.IsValid())
 	{
-		if (!EntryWidgetClass)
+		if (!EntryWidgetClass || !IsEntryClassValid(EntryWidgetClass))
 		{
 			// We have no entry class, so clear everything out
 			Reset(true);
@@ -73,7 +77,7 @@ UUserWidget* UDynamicEntryBox::BP_CreateEntry()
 
 UUserWidget* UDynamicEntryBox::BP_CreateEntryOfClass(TSubclassOf<UUserWidget> EntryClass)
 {
-	if (EntryClass)
+	if (EntryClass && IsEntryClassValid(EntryClass))
 	{
 		return CreateEntryInternal(EntryClass);
 	}
diff --git a/Engine/Source/Runtime/UMG/Private/Components/DynamicEntryBoxBase.cpp b/Engine/Source/Runtime/UMG/Private/Components/DynamicEntryBoxBase.cpp
--- a/Engine/Source/Runtime/UMG/Private/Components/DynamicEntryBoxBase.cpp
+++ b/Engine/Source/Runtime/UMG/Private/Components/DynamicEntryBoxBase.cpp
@@ -2,6 +2,8 @@
 
 #include "Components/DynamicEntryBoxBase.h"
 #include "UMGPrivate.h"
+#include "Blueprint/WidgetTree.h"
+#include "Blueprint/UserWidget.h"
 #include "Widgets/Layout/SWrapBox.h"
 #include "Widgets/SBoxPanel.h"
 #include "Editor/WidgetCompilerLog.h"
@@ -242,6 +244,26 @@ void UDynamicEntryBoxBase::SynchronizeProperties()
 #endif
 }
 
+bool UDynamicEntryBoxBase::IsEntryClassValid(TSubclassOf<UUserWidget> InEntryClass) const
+{
+	if (InEntryClass)
+	{
+		// Would InEntryClass create an instance of the same DynamicEntryBox
+		if (UWidgetTree* WidgetTree = Cast<UWidgetTree>(GetOuter()))
+		{
+			if (UUserWidget* UserWidget = Cast<UUserWidget>(WidgetTree->GetOuter()))
+			{
+				if (InEntryClass->IsChildOf(UserWidget->GetClass()))
+				{
+					return false;
+				}
+			}
+		}
+	}
+
+	return true;
+}
+
 UUserWidget* UDynamicEntryBoxBase::CreateEntryInternal(TSubclassOf<UUserWidget> InEntryClass)
 {
 	UUserWidget* NewEntryWidget = EntryWidgetPool.GetOrCreateInstance(InEntryClass);
diff --git a/Engine/Source/Runtime/UMG/Public/Components/DynamicEntryBox.h b/Engine/Source/Runtime/UMG/Public/Components/DynamicEntryBox.h
--- a/Engine/Source/Runtime/UMG/Public/Components/DynamicEntryBox.h
+++ b/Engine/Source/Runtime/UMG/Public/Components/DynamicEntryBox.h
@@ -25,7 +25,7 @@ public:
 	WidgetT* CreateEntry(const TSubclassOf<WidgetT>& ExplicitEntryClass = nullptr)
 	{
 		TSubclassOf<UUserWidget> EntryClass = ExplicitEntryClass ? ExplicitEntryClass : EntryWidgetClass;
-		if (EntryClass && EntryClass->IsChildOf(WidgetT::StaticClass()))
+		if (EntryClass && EntryClass->IsChildOf(WidgetT::StaticClass()) && IsEntryClassValid(EntryClass))
 		{
 			return Cast<WidgetT>(CreateEntryInternal(EntryClass));
 		}
diff --git a/Engine/Source/Runtime/UMG/Public/Components/DynamicEntryBoxBase.h b/Engine/Source/Runtime/UMG/Public/Components/DynamicEntryBoxBase.h
--- a/Engine/Source/Runtime/UMG/Public/Components/DynamicEntryBoxBase.h
+++ b/Engine/Source/Runtime/UMG/Public/Components/DynamicEntryBoxBase.h
@@ -65,6 +65,7 @@ protected:
 	virtual void SynchronizeProperties() override;
 	virtual void AddEntryChild(UUserWidget& ChildWidget);
 
+	bool IsEntryClassValid(TSubclassOf<UUserWidget> InEntryClass) const;
 	UUserWidget* CreateEntryInternal(TSubclassOf<UUserWidget> InEntryClass);
 	void RemoveEntryInternal(UUserWidget* EntryWidget);
 	FMargin BuildEntryPadding(const FVector2D& DesiredSpacing);
diff --git a/Engine/Source/Runtime/UMG/Public/Editor/WidgetCompilerLog.h b/Engine/Source/Runtime/UMG/Public/Editor/WidgetCompilerLog.h
--- a/Engine/Source/Runtime/UMG/Public/Editor/WidgetCompilerLog.h
+++ b/Engine/Source/Runtime/UMG/Public/Editor/WidgetCompilerLog.h
@@ -9,6 +9,8 @@
 
 #if WITH_EDITOR
 
+class UUserWidget;
+
 class IWidgetCompilerLog
 {
 public:
@@ -35,6 +37,8 @@ public:
 		InternalLogMessage(Line);
 		return Line;
 	}
+
+	virtual TSubclassOf<UUserWidget> GetContextClass() const = 0;
 	
 protected:
 	virtual void InternalLogMessage(TSharedRef<FTokenizedMessage>& Message) = 0;
