Remove unused border from UMG designer view which is causing a huge gray block to appear in the new style

diff --git a/Engine/Source/Editor/UMGEditor/Private/Designer/SDesignerView.cpp b/Engine/Source/Editor/UMGEditor/Private/Designer/SDesignerView.cpp
--- a/Engine/Source/Editor/UMGEditor/Private/Designer/SDesignerView.cpp
+++ b/Engine/Source/Editor/UMGEditor/Private/Designer/SDesignerView.cpp
@@ -421,16 +421,6 @@ void SDesignerView::Construct(const FArguments& InArgs, TSharedPtr<FWidgetBluepr
 					.ViewOffset(this, &SDesignerView::GetViewOffset)
 					[
 						SNew(SOverlay)
-
-						+ SOverlay::Slot()
-						[
-							SNew(SBorder)
-							[
-								SNew(SSpacer)
-								.Size(FVector2D(1, 1))
-							]
-						]
-						
 						+ SOverlay::Slot()
 						[
 							SNew(SBorder)
