Updated UMG horizontal and vertical alignment to use segmented control UI.

diff --git a/Engine/Source/Editor/EditorStyle/Private/StarshipStyle.cpp b/Engine/Source/Editor/EditorStyle/Private/StarshipStyle.cpp
--- a/Engine/Source/Editor/EditorStyle/Private/StarshipStyle.cpp
+++ b/Engine/Source/Editor/EditorStyle/Private/StarshipStyle.cpp
@@ -7608,15 +7608,15 @@ void FStarshipEditorStyle::FStyle::SetupUMGEditorStyles()
 		.SetCheckedPressedImage(IMAGE_BRUSH("Icons/Star_16x", Icon10x10, FLinearColor(0.2f, 0.2f, 0.2f, 1.f)));
 	Set("UMGEditor.Palette.FavoriteToggleStyle", UMGEditorFavoriteToggleStyle);
 
-	Set("HorizontalAlignment_Left", new IMAGE_BRUSH("Icons/UMG/Alignment/Horizontal_Left", Icon20x20));
-	Set("HorizontalAlignment_Center", new IMAGE_BRUSH("Icons/UMG/Alignment/Horizontal_Center", Icon20x20));
-	Set("HorizontalAlignment_Right", new IMAGE_BRUSH("Icons/UMG/Alignment/Horizontal_Right", Icon20x20));
-	Set("HorizontalAlignment_Fill", new IMAGE_BRUSH("Icons/UMG/Alignment/Horizontal_Fill", Icon20x20));
-
-	Set("VerticalAlignment_Top", new IMAGE_BRUSH("Icons/UMG/Alignment/Vertical_Top", Icon20x20));
-	Set("VerticalAlignment_Center", new IMAGE_BRUSH("Icons/UMG/Alignment/Vertical_Center", Icon20x20));
-	Set("VerticalAlignment_Bottom", new IMAGE_BRUSH("Icons/UMG/Alignment/Vertical_Bottom", Icon20x20));
-	Set("VerticalAlignment_Fill", new IMAGE_BRUSH("Icons/UMG/Alignment/Vertical_Fill", Icon20x20));
+	Set("HorizontalAlignment_Left", new IMAGE_BRUSH("Icons/UMG/Alignment/Horizontal_Left", Icon16x16));
+	Set("HorizontalAlignment_Center", new IMAGE_BRUSH("Icons/UMG/Alignment/Horizontal_Center", Icon16x16));
+	Set("HorizontalAlignment_Right", new IMAGE_BRUSH("Icons/UMG/Alignment/Horizontal_Right", Icon16x16));
+	Set("HorizontalAlignment_Fill", new IMAGE_BRUSH("Icons/UMG/Alignment/Horizontal_Fill", Icon16x16));
+
+	Set("VerticalAlignment_Top", new IMAGE_BRUSH("Icons/UMG/Alignment/Vertical_Top", Icon16x16));
+	Set("VerticalAlignment_Center", new IMAGE_BRUSH("Icons/UMG/Alignment/Vertical_Center", Icon16x16));
+	Set("VerticalAlignment_Bottom", new IMAGE_BRUSH("Icons/UMG/Alignment/Vertical_Bottom", Icon16x16));
+	Set("VerticalAlignment_Fill", new IMAGE_BRUSH("Icons/UMG/Alignment/Vertical_Fill", Icon16x16));
 
 	const FTextBlockStyle NoAnimationFont =
 		FTextBlockStyle(NormalText)
diff --git a/Engine/Source/Editor/UMGEditor/Private/Customizations/HorizontalAlignmentCustomization.cpp b/Engine/Source/Editor/UMGEditor/Private/Customizations/HorizontalAlignmentCustomization.cpp
--- a/Engine/Source/Editor/UMGEditor/Private/Customizations/HorizontalAlignmentCustomization.cpp
+++ b/Engine/Source/Editor/UMGEditor/Private/Customizations/HorizontalAlignmentCustomization.cpp
@@ -4,6 +4,7 @@
 #include "Widgets/Images/SImage.h"
 #include "Widgets/Input/SCheckBox.h"
 #include "DetailWidgetRow.h"
+#include "Widgets/Input/SSegmentedControl.h"
 
 
 #define LOCTEXT_NAMESPACE "UMG"
@@ -19,72 +20,23 @@ void FHorizontalAlignmentCustomization::CustomizeHeader(TSharedRef<IPropertyHand
 		PropertyHandle->CreatePropertyNameWidget()
 	]
 	.ValueContent()
+	.MaxDesiredWidth(0)
 	[
-		SNew(SHorizontalBox)
-
-		+ SHorizontalBox::Slot()
-		.AutoWidth()
-		.Padding(OuterPadding)
-		[
-			SNew( SCheckBox )
-			.Style(FEditorStyle::Get(), "ToggleButtonCheckbox")
-			.ToolTipText(LOCTEXT("HAlignLeft", "Horizontally Align Left"))
-			.Padding(ContentPadding)
-			.OnCheckStateChanged(this, &FHorizontalAlignmentCustomization::HandleCheckStateChanged, PropertyHandle, HAlign_Left)
-			.IsChecked(this, &FHorizontalAlignmentCustomization::GetCheckState, PropertyHandle, HAlign_Left)
-			[
-				SNew(SImage)
-				.Image(FEditorStyle::GetBrush("HorizontalAlignment_Left"))
-			]
-		]
-
-		+ SHorizontalBox::Slot()
-		.AutoWidth()
-		.Padding(OuterPadding)
-		[
-			SNew(SCheckBox)
-			.Style(FEditorStyle::Get(), "ToggleButtonCheckbox")
-			.ToolTipText(LOCTEXT("HAlignCenter", "Horizontally Align Center"))
-			.Padding(ContentPadding)
-			.OnCheckStateChanged(this, &FHorizontalAlignmentCustomization::HandleCheckStateChanged, PropertyHandle, HAlign_Center)
-			.IsChecked(this, &FHorizontalAlignmentCustomization::GetCheckState, PropertyHandle, HAlign_Center)
-			[
-				SNew(SImage)
-				.Image(FEditorStyle::GetBrush("HorizontalAlignment_Center"))
-			]
-		]
-
-		+ SHorizontalBox::Slot()
-		.AutoWidth()
-		.Padding(OuterPadding)
-		[
-			SNew(SCheckBox)
-			.Style(FEditorStyle::Get(), "ToggleButtonCheckbox")
-			.ToolTipText(LOCTEXT("HAlignRight", "Horizontally Align Right"))
-			.Padding(ContentPadding)
-			.OnCheckStateChanged(this, &FHorizontalAlignmentCustomization::HandleCheckStateChanged, PropertyHandle, HAlign_Right)
-			.IsChecked(this, &FHorizontalAlignmentCustomization::GetCheckState, PropertyHandle, HAlign_Right)
-			[
-				SNew(SImage)
-				.Image(FEditorStyle::GetBrush("HorizontalAlignment_Right"))
-			]
-		]
-
-		+ SHorizontalBox::Slot()
-		.AutoWidth()
-		.Padding(OuterPadding)
-		[
-			SNew(SCheckBox)
-			.Style(FEditorStyle::Get(), "ToggleButtonCheckbox")
-			.ToolTipText(LOCTEXT("HAlignFill", "Horizontally Align Fill"))
-			.Padding(ContentPadding)
-			.OnCheckStateChanged(this, &FHorizontalAlignmentCustomization::HandleCheckStateChanged, PropertyHandle, HAlign_Fill)
-			.IsChecked(this, &FHorizontalAlignmentCustomization::GetCheckState, PropertyHandle, HAlign_Fill)
-			[
-				SNew(SImage)
-				.Image(FEditorStyle::GetBrush("HorizontalAlignment_Fill"))
-			]
-		]
+		SNew(SSegmentedControl<EHorizontalAlignment>)
+		.Value(this, &FHorizontalAlignmentCustomization::GetCurrentAlignment, PropertyHandle)
+		.OnValueChanged(this, &FHorizontalAlignmentCustomization::OnCurrentAlignmentChanged, PropertyHandle)
+		+SSegmentedControl<EHorizontalAlignment>::Slot(EHorizontalAlignment::HAlign_Left)
+			.Icon(FEditorStyle::GetBrush("HorizontalAlignment_Left"))
+			.ToolTip(LOCTEXT("HAlignLeft", "Left Align Horizontally"))
+		+ SSegmentedControl<EHorizontalAlignment>::Slot(EHorizontalAlignment::HAlign_Center)
+			.Icon(FEditorStyle::GetBrush("HorizontalAlignment_Center"))
+			.ToolTip(LOCTEXT("HAlignCenter", "Center Align Horizontally"))
+		+ SSegmentedControl<EHorizontalAlignment>::Slot(EHorizontalAlignment::HAlign_Right)
+			.Icon(FEditorStyle::GetBrush("HorizontalAlignment_Right"))
+			.ToolTip(LOCTEXT("HAlignRight", "Right Align Horizontally"))
+		+ SSegmentedControl<EHorizontalAlignment>::Slot(EHorizontalAlignment::HAlign_Fill)
+			.Icon(FEditorStyle::GetBrush("HorizontalAlignment_Fill"))
+			.ToolTip(LOCTEXT("HAlignFill", "Fill Horizontally"))
 	];
 }
 
@@ -92,20 +44,20 @@ void FHorizontalAlignmentCustomization::CustomizeChildren(TSharedRef<IPropertyHa
 {
 }
 
-void FHorizontalAlignmentCustomization::HandleCheckStateChanged(ECheckBoxState InCheckboxState, TSharedRef<IPropertyHandle> PropertyHandle, EHorizontalAlignment ToAlignment)
-{
-	PropertyHandle->SetValue((uint8)ToAlignment);
-}
-
-ECheckBoxState FHorizontalAlignmentCustomization::GetCheckState(TSharedRef<IPropertyHandle> PropertyHandle, EHorizontalAlignment ForAlignment) const
+EHorizontalAlignment FHorizontalAlignmentCustomization::GetCurrentAlignment(TSharedRef<IPropertyHandle> PropertyHandle) const
 {
 	uint8 Value;
-	if ( PropertyHandle->GetValue(Value) == FPropertyAccess::Result::Success )
+	if (PropertyHandle->GetValue(Value) == FPropertyAccess::Result::Success)
 	{
-		return Value == ForAlignment ? ECheckBoxState::Checked : ECheckBoxState::Unchecked;
+		return (EHorizontalAlignment)Value;
 	}
 
-	return ECheckBoxState::Unchecked;
+	return HAlign_Fill;
+}
+
+void FHorizontalAlignmentCustomization::OnCurrentAlignmentChanged(EHorizontalAlignment NewAlignment, TSharedRef<IPropertyHandle> PropertyHandle)
+{
+	PropertyHandle->SetValue((uint8)NewAlignment);
 }
 
 #undef LOCTEXT_NAMESPACE
diff --git a/Engine/Source/Editor/UMGEditor/Private/Customizations/HorizontalAlignmentCustomization.h b/Engine/Source/Editor/UMGEditor/Private/Customizations/HorizontalAlignmentCustomization.h
--- a/Engine/Source/Editor/UMGEditor/Private/Customizations/HorizontalAlignmentCustomization.h
+++ b/Engine/Source/Editor/UMGEditor/Private/Customizations/HorizontalAlignmentCustomization.h
@@ -26,4 +26,8 @@ public:
 	void HandleCheckStateChanged(ECheckBoxState InCheckboxState, TSharedRef<IPropertyHandle> PropertyHandle, EHorizontalAlignment ToAlignment);
 
 	ECheckBoxState GetCheckState(TSharedRef<IPropertyHandle> PropertyHandle, EHorizontalAlignment ForAlignment) const;
+
+	EHorizontalAlignment GetCurrentAlignment(TSharedRef<IPropertyHandle> PropertyHandle) const;
+
+	void OnCurrentAlignmentChanged(EHorizontalAlignment NewAlignment, TSharedRef<IPropertyHandle> PropertyHandle);
 };
diff --git a/Engine/Source/Editor/UMGEditor/Private/Customizations/VerticalAlignmentCustomization.cpp b/Engine/Source/Editor/UMGEditor/Private/Customizations/VerticalAlignmentCustomization.cpp
--- a/Engine/Source/Editor/UMGEditor/Private/Customizations/VerticalAlignmentCustomization.cpp
+++ b/Engine/Source/Editor/UMGEditor/Private/Customizations/VerticalAlignmentCustomization.cpp
@@ -9,6 +9,7 @@
 #include "Widgets/Input/SCheckBox.h"
 #include "PropertyHandle.h"
 #include "DetailWidgetRow.h"
+#include "Widgets/Input/SSegmentedControl.h"
 
 
 #define LOCTEXT_NAMESPACE "UMG"
@@ -24,72 +25,23 @@ void FVerticalAlignmentCustomization::CustomizeHeader(TSharedRef<IPropertyHandle
 		PropertyHandle->CreatePropertyNameWidget()
 	]
 	.ValueContent()
+	.MaxDesiredWidth(0)
 	[
-		SNew(SHorizontalBox)
-
-		+ SHorizontalBox::Slot()
-		.AutoWidth()
-		.Padding(OuterPadding)
-		[
-			SNew( SCheckBox )
-			.Style(FEditorStyle::Get(), "ToggleButtonCheckbox")
-			.ToolTipText(LOCTEXT("VAlignTop", "Vertically Align Top"))
-			.Padding(ContentPadding)
-			.OnCheckStateChanged(this, &FVerticalAlignmentCustomization::HandleCheckStateChanged, PropertyHandle, VAlign_Top)
-			.IsChecked(this, &FVerticalAlignmentCustomization::GetCheckState, PropertyHandle, VAlign_Top)
-			[
-				SNew(SImage)
-				.Image(FEditorStyle::GetBrush("VerticalAlignment_Top"))
-			]
-		]
-
-		+ SHorizontalBox::Slot()
-		.AutoWidth()
-		.Padding(OuterPadding)
-		[
-			SNew(SCheckBox)
-			.Style(FEditorStyle::Get(), "ToggleButtonCheckbox")
-			.ToolTipText(LOCTEXT("VAlignCenter", "Vertically Align Center"))
-			.Padding(ContentPadding)
-			.OnCheckStateChanged(this, &FVerticalAlignmentCustomization::HandleCheckStateChanged, PropertyHandle, VAlign_Center)
-			.IsChecked(this, &FVerticalAlignmentCustomization::GetCheckState, PropertyHandle, VAlign_Center)
-			[
-				SNew(SImage)
-				.Image(FEditorStyle::GetBrush("VerticalAlignment_Center"))
-			]
-		]
-
-		+ SHorizontalBox::Slot()
-		.AutoWidth()
-		.Padding(OuterPadding)
-		[
-			SNew(SCheckBox)
-			.Style(FEditorStyle::Get(), "ToggleButtonCheckbox")
-			.ToolTipText(LOCTEXT("VAlignBottom", "Vertically Align Bottom"))
-			.Padding(ContentPadding)
-			.OnCheckStateChanged(this, &FVerticalAlignmentCustomization::HandleCheckStateChanged, PropertyHandle, VAlign_Bottom)
-			.IsChecked(this, &FVerticalAlignmentCustomization::GetCheckState, PropertyHandle, VAlign_Bottom)
-			[
-				SNew(SImage)
-				.Image(FEditorStyle::GetBrush("VerticalAlignment_Bottom"))
-			]
-		]
-
-		+ SHorizontalBox::Slot()
-		.AutoWidth()
-		.Padding(OuterPadding)
-		[
-			SNew(SCheckBox)
-			.Style(FEditorStyle::Get(), "ToggleButtonCheckbox")
-			.ToolTipText(LOCTEXT("VAlignFill", "Vertically Align Fill"))
-			.Padding(ContentPadding)
-			.OnCheckStateChanged(this, &FVerticalAlignmentCustomization::HandleCheckStateChanged, PropertyHandle, VAlign_Fill)
-			.IsChecked(this, &FVerticalAlignmentCustomization::GetCheckState, PropertyHandle, VAlign_Fill)
-			[
-				SNew(SImage)
-				.Image(FEditorStyle::GetBrush("VerticalAlignment_Fill"))
-			]
-		]
+		SNew(SSegmentedControl<EVerticalAlignment>)
+		.Value(this, &FVerticalAlignmentCustomization::GetCurrentAlignment, PropertyHandle)
+		.OnValueChanged(this, &FVerticalAlignmentCustomization::OnCurrentAlignmentChanged, PropertyHandle)
+		+ SSegmentedControl<EVerticalAlignment>::Slot(EVerticalAlignment::VAlign_Top)
+			.Icon(FEditorStyle::GetBrush("VerticalAlignment_Top"))
+			.ToolTip(LOCTEXT("VAlignTop", "Top Align Vertically"))
+		+ SSegmentedControl<EVerticalAlignment>::Slot(EVerticalAlignment::VAlign_Center)
+			.Icon(FEditorStyle::GetBrush("VerticalAlignment_Center"))
+			.ToolTip(LOCTEXT("VAlignCenter", "Center Align Vertically"))
+		+ SSegmentedControl<EVerticalAlignment>::Slot(EVerticalAlignment::VAlign_Bottom)
+			.Icon(FEditorStyle::GetBrush("VerticalAlignment_Bottom"))
+			.ToolTip(LOCTEXT("VAlignBottom", "Bottom Align Vertically"))
+		+ SSegmentedControl<EVerticalAlignment>::Slot(EVerticalAlignment::VAlign_Fill)
+			.Icon(FEditorStyle::GetBrush("VerticalAlignment_Fill"))
+			.ToolTip(LOCTEXT("VAlignFill", "Fill Vertically"))
 	];
 }
 
@@ -97,20 +49,20 @@ void FVerticalAlignmentCustomization::CustomizeChildren(TSharedRef<IPropertyHand
 {
 }
 
-void FVerticalAlignmentCustomization::HandleCheckStateChanged(ECheckBoxState InCheckboxState, TSharedRef<IPropertyHandle> PropertyHandle, EVerticalAlignment ToAlignment)
-{
-	PropertyHandle->SetValue((uint8)ToAlignment);
-}
-
-ECheckBoxState FVerticalAlignmentCustomization::GetCheckState(TSharedRef<IPropertyHandle> PropertyHandle, EVerticalAlignment ForAlignment) const
+EVerticalAlignment FVerticalAlignmentCustomization::GetCurrentAlignment(TSharedRef<IPropertyHandle> PropertyHandle) const
 {
 	uint8 Value;
-	if ( PropertyHandle->GetValue(Value) == FPropertyAccess::Result::Success)
+	if (PropertyHandle->GetValue(Value) == FPropertyAccess::Result::Success)
 	{
-		return Value == ForAlignment ? ECheckBoxState::Checked : ECheckBoxState::Unchecked;
+		return (EVerticalAlignment)Value;
 	}
 
-	return ECheckBoxState::Unchecked;
+	return VAlign_Fill;
+}
+
+void FVerticalAlignmentCustomization::OnCurrentAlignmentChanged(EVerticalAlignment NewAlignment, TSharedRef<IPropertyHandle> PropertyHandle)
+{
+	PropertyHandle->SetValue((uint8)NewAlignment);
 }
 
 #undef LOCTEXT_NAMESPACE
diff --git a/Engine/Source/Editor/UMGEditor/Private/Customizations/VerticalAlignmentCustomization.h b/Engine/Source/Editor/UMGEditor/Private/Customizations/VerticalAlignmentCustomization.h
--- a/Engine/Source/Editor/UMGEditor/Private/Customizations/VerticalAlignmentCustomization.h
+++ b/Engine/Source/Editor/UMGEditor/Private/Customizations/VerticalAlignmentCustomization.h
@@ -24,8 +24,8 @@ public:
 	virtual void CustomizeHeader(TSharedRef<IPropertyHandle> PropertyHandle, class FDetailWidgetRow& HeaderRow, IPropertyTypeCustomizationUtils& CustomizationUtils) override;
 
 	virtual void CustomizeChildren(TSharedRef<IPropertyHandle> PropertyHandle, class IDetailChildrenBuilder& ChildBuilder, IPropertyTypeCustomizationUtils& CustomizationUtils) override;
-	
-	void HandleCheckStateChanged(ECheckBoxState InCheckboxState, TSharedRef<IPropertyHandle> PropertyHandle, EVerticalAlignment ToAlignment);
 
-	ECheckBoxState GetCheckState(TSharedRef<IPropertyHandle> PropertyHandle, EVerticalAlignment ForAlignment) const;
+	EVerticalAlignment GetCurrentAlignment(TSharedRef<IPropertyHandle> PropertyHandle) const;
+
+	void OnCurrentAlignmentChanged(EVerticalAlignment NewAlignment, TSharedRef<IPropertyHandle> PropertyHandle);
 };
