Fix to Undo/Redo of adding a Widget to the Designer View causing it to be renamed with a Trash_ prefix.

diff --git a/Engine/Source/Editor/UMGEditor/Private/Designer/SDesignerView.cpp b/Engine/Source/Editor/UMGEditor/Private/Designer/SDesignerView.cpp
--- a/Engine/Source/Editor/UMGEditor/Private/Designer/SDesignerView.cpp
+++ b/Engine/Source/Editor/UMGEditor/Private/Designer/SDesignerView.cpp
@@ -2789,6 +2789,7 @@ void SDesignerView::ProcessDropAndAddWidget(const FGeometry& MyGeometry, const F
 
 	FGeometry WidgetUnderCursorGeometry = HitResult.WidgetArranged.Geometry;
 	
+	FScopedTransaction DragAndDropTransaction(LOCTEXT("Designer_DragAddDrop", "Drag and Drop Widget"));
 	TArray<UWidget*> DragDropPreviewWidgets;
 	DetermineDragDropPreviewWidgets(DragDropPreviewWidgets, DragDropEvent, TargetTree);
 
@@ -2874,12 +2875,12 @@ void SDesignerView::ProcessDropAndAddWidget(const FGeometry& MyGeometry, const F
 			DragOperation->SetCursorOverride(EMouseCursor::SlashedCircle);
 
 			// Cancel the transaction even if it's not a preview, since we can't do anything
-			Transaction.Cancel();
+			DragAndDropTransaction.Cancel();
 		}
 
 		if (bIsPreview)
 		{
-			Transaction.Cancel();
+			DragAndDropTransaction.Cancel();
 		}
 
 		// Remove widgets tracked by the 'DropPreviews' set. We don't consider them to be transient at this point because they have been inserted into the widget tree hierarchy.
@@ -3120,7 +3121,7 @@ void SDesignerView::ProcessDropAndAddWidget(const FGeometry& MyGeometry, const F
 
 		if (bIsPreview || !bWidgetMoved)
 		{
-			Transaction.Cancel();
+			DragAndDropTransaction.Cancel();
 		}
 	}
 
