PR #7283: Fixup widget editor crash when working with struct or instanced properties (Contributed by [REDACTED_PERSON])

diff --git a/Engine/Source/Editor/UMGEditor/Private/WidgetBlueprintEditor.cpp b/Engine/Source/Editor/UMGEditor/Private/WidgetBlueprintEditor.cpp
--- a/Engine/Source/Editor/UMGEditor/Private/WidgetBlueprintEditor.cpp
+++ b/Engine/Source/Editor/UMGEditor/Private/WidgetBlueprintEditor.cpp
@@ -619,7 +619,7 @@ static bool MigratePropertyValue(UObject* SourceObject, UObject* DestinationObje
 	ensure(SourceObject->GetClass() == DestinationObject->GetClass());
 
 	// If the current property is an array, map or set, short-circuit current progress so that we copy the whole container.
-	if ( CastField<FArrayProperty>(CurrentProperty) || CastField<FMapProperty>(CurrentProperty) || CastField<FSetProperty>(CurrentProperty) )
+	if ( CastField<FArrayProperty>(CurrentProperty) || CastField<FMapProperty>(CurrentProperty) || CastField<FSetProperty>(CurrentProperty) || CastField<FStructProperty>(CurrentProperty))
 	{
 		NextNode = nullptr;
 	}
diff --git a/Engine/Source/Runtime/Engine/Private/ObjectEditorUtils.cpp b/Engine/Source/Runtime/Engine/Private/ObjectEditorUtils.cpp
--- a/Engine/Source/Runtime/Engine/Private/ObjectEditorUtils.cpp
+++ b/Engine/Source/Runtime/Engine/Private/ObjectEditorUtils.cpp
@@ -246,7 +246,7 @@ namespace FObjectEditorUtils
 						UObject* ExistingObject = StaticFindObject(UObject::StaticClass(), InDestinationObject, *Value->GetFName().ToString());
 						if (ExistingObject)
 						{
-							ExistingObject->Rename(nullptr, GetTransientPackage());
+							ExistingObject->Rename(nullptr, GetTransientPackage(), REN_DontCreateRedirectors);
 						}
 
 						UObject* DuplicateValue = StaticDuplicateObject(Value, InDestinationObject, Value->GetFName(), RF_AllFlags, nullptr, EDuplicateMode::Normal, EInternalObjectFlags::AllFlags);
