PR #7209: UE-92202: Prevent crash due to drag/drop named slot element (Contributed by [REDACTED_PERSON])

diff --git a/Engine/Source/Editor/UMGEditor/Private/Designer/SDesignerView.cpp b/Engine/Source/Editor/UMGEditor/Private/Designer/SDesignerView.cpp
--- a/Engine/Source/Editor/UMGEditor/Private/Designer/SDesignerView.cpp
+++ b/Engine/Source/Editor/UMGEditor/Private/Designer/SDesignerView.cpp
@@ -1891,7 +1891,8 @@ FReply SDesignerView::OnMouseMove(const FGeometry& MyGeometry, const FPointerEve
 				bool bIsRootWidgetSelected = false;
 				for (const auto& SelectedWidget : SelectedWidgets)
 				{
-					if (SelectedWidget.GetTemplate()->GetParent() == nullptr)
+					UWidget* ParentWidget = SelectedWidget.GetTemplate()->GetParent();
+					if (!ParentWidget || Cast<UNamedSlot>(ParentWidget))
 					{
 						bIsRootWidgetSelected = true;
 						break;
@@ -2956,7 +2957,7 @@ void SDesignerView::ProcessDropAndAddWidget(const FGeometry& MyGeometry, const F
 				UPanelWidget* NewParent = Cast<UPanelWidget>(Target);
 
 				UWidget* Widget = bIsPreview ? DraggedWidget.Preview : DraggedWidget.Template;
-				UPanelWidget* ParentWidget = bIsPreview ? Cast<UPanelWidget>(DraggedWidget.ParentWidget.GetPreview()) : Cast<UPanelWidget>(DraggedWidget.ParentWidget.GetTemplate());
+				UWidget* ParentWidget = bIsPreview ? (DraggedWidget.ParentWidget.GetPreview()) : (DraggedWidget.ParentWidget.GetTemplate());
 				if (ensure(Widget))
 				{
 					bool bIsChangingParent = ParentWidget != NewParent;
@@ -3044,9 +3045,9 @@ void SDesignerView::ProcessDropAndAddWidget(const FGeometry& MyGeometry, const F
 					{
 						Slot = NewParent->AddChild(Widget);
 					}
-					else
+					else if (UPanelWidget* ParentwidgetAsPanel = Cast<UPanelWidget>(ParentWidget))
 					{
-						Slot = ParentWidget->InsertChildAt(ParentWidget->GetChildIndex(Widget), Widget);
+						Slot = ParentwidgetAsPanel->InsertChildAt(ParentwidgetAsPanel->GetChildIndex(Widget), Widget);
 					}
 
 					if (Slot != nullptr)
diff --git a/Engine/Source/Editor/UMGEditor/Private/WidgetBlueprintEditorUtils.cpp b/Engine/Source/Editor/UMGEditor/Private/WidgetBlueprintEditorUtils.cpp
--- a/Engine/Source/Editor/UMGEditor/Private/WidgetBlueprintEditorUtils.cpp
+++ b/Engine/Source/Editor/UMGEditor/Private/WidgetBlueprintEditorUtils.cpp
@@ -498,7 +498,7 @@ void FWidgetBlueprintEditorUtils::DeleteWidgets(UWidgetBlueprint* Blueprint, TSe
 			// Modify the widget being removed.
 			WidgetTemplate->Modify();
 
-			bRemoved = Blueprint->WidgetTree->RemoveWidget(WidgetTemplate);
+			bRemoved |= Blueprint->WidgetTree->RemoveWidget(WidgetTemplate);
 
 			// If the widget we're removing doesn't have a parent it may be rooted in a named slot,
 			// so check there as well.
diff --git a/Engine/Source/Runtime/UMG/Private/Components/PanelWidget.cpp b/Engine/Source/Runtime/UMG/Private/Components/PanelWidget.cpp
--- a/Engine/Source/Runtime/UMG/Private/Components/PanelWidget.cpp
+++ b/Engine/Source/Runtime/UMG/Private/Components/PanelWidget.cpp
@@ -21,7 +21,7 @@ void UPanelWidget::ReleaseSlateResources(bool bReleaseChildren)
 	{
 		for ( int32 SlotIndex = 0; SlotIndex < Slots.Num(); SlotIndex++ )
 		{
-			if ( Slots[SlotIndex]->Content != nullptr )
+			if ( Slots[SlotIndex] && Slots[SlotIndex]->Content != nullptr )
 			{
 				Slots[SlotIndex]->ReleaseSlateResources(bReleaseChildren);
 			}
