Edigrate CL 14337650 from UE4 Main

Fix Duplicating a Widget in the Hierarchy does not retain the selection index

diff --git a/Engine/Source/Editor/UMGEditor/Private/Hierarchy/SHierarchyView.cpp b/Engine/Source/Editor/UMGEditor/Private/Hierarchy/SHierarchyView.cpp
--- a/Engine/Source/Editor/UMGEditor/Private/Hierarchy/SHierarchyView.cpp
+++ b/Engine/Source/Editor/UMGEditor/Private/Hierarchy/SHierarchyView.cpp
@@ -429,10 +429,16 @@ void SHierarchyView::RecursiveExpand(TSharedPtr<FHierarchyModel>& Model, EExpand
 
 void SHierarchyView::RestoreSelectedItems()
 {
+	// Use single selection mode to update the navigation item as we select.
+	// Otherwise keyboard navigation will be off.
+	WidgetTreeView->SetSelectionMode(ESelectionMode::Single);
+
 	for ( TSharedPtr<FHierarchyModel>& Model : RootWidgets )
 	{
 		RecursiveSelection(Model);
 	}
+
+	WidgetTreeView->SetSelectionMode(ESelectionMode::Multi);
 }
 
 void SHierarchyView::RecursiveSelection(TSharedPtr<FHierarchyModel>& Model)
