Expose NameOverride on SinglePropertyView widget

diff --git a/Engine/Source/Editor/UMGEditor/Private/Components/SinglePropertyView.cpp b/Engine/Source/Editor/UMGEditor/Private/Components/SinglePropertyView.cpp
--- a/Engine/Source/Editor/UMGEditor/Private/Components/SinglePropertyView.cpp
+++ b/Engine/Source/Editor/UMGEditor/Private/Components/SinglePropertyView.cpp
@@ -121,6 +121,20 @@ void USinglePropertyView::SetPropertyName(FName InPropertyName)
 	}
 }
 
+FText USinglePropertyView::GetNameOverride() const
+{
+	return NameOverride;
+}
+
+
+void USinglePropertyView::SetNameOverride(FText InNameOverride)
+{
+	if (!NameOverride.EqualTo(InNameOverride))
+	{
+		NameOverride = InNameOverride;
+		BuildContentWidget();
+	}
+}
 
 void USinglePropertyView::OnObjectChanged()
 {
diff --git a/Engine/Source/Editor/UMGEditor/Public/Components/SinglePropertyView.h b/Engine/Source/Editor/UMGEditor/Public/Components/SinglePropertyView.h
--- a/Engine/Source/Editor/UMGEditor/Public/Components/SinglePropertyView.h
+++ b/Engine/Source/Editor/UMGEditor/Public/Components/SinglePropertyView.h
@@ -32,6 +32,12 @@ public:
 	UFUNCTION(BlueprintCallable, Category = "View")
 	void SetPropertyName(FName NewPropertyName);
 
+	UFUNCTION(BlueprintCallable, Category = "View")
+	FText GetNameOverride() const;
+
+	UFUNCTION(BlueprintCallable, Category = "View")
+	void SetNameOverride(FText NewPropertyName);
+
 private:
 	void InternalSinglePropertyChanged();
 
