UMG: Add extensibility manager to the UMG editor.

diff --git a/Engine/Source/Editor/UMGEditor/Private/WidgetBlueprintEditor.cpp b/Engine/Source/Editor/UMGEditor/Private/WidgetBlueprintEditor.cpp
--- a/Engine/Source/Editor/UMGEditor/Private/WidgetBlueprintEditor.cpp
+++ b/Engine/Source/Editor/UMGEditor/Private/WidgetBlueprintEditor.cpp
@@ -53,6 +53,7 @@
 #include "Framework/Notifications/NotificationManager.h"
 #include "Widgets/Notifications/SNotificationList.h"
 #include "UMGEditorActions.h"
+#include "UMGEditorModule.h"
 #include "GameProjectGenerationModule.h"
 
 #include "SPaletteViewModel.h"
@@ -161,6 +162,9 @@ void FWidgetBlueprintEditor::InitalizeExtenders()
 {
 	FBlueprintEditor::InitalizeExtenders();
 
+	IUMGEditorModule& UMGEditorModule = FModuleManager::LoadModuleChecked<IUMGEditorModule>("UMGEditor");
+	AddMenuExtender(UMGEditorModule.GetMenuExtensibilityManager()->GetAllExtenders(GetToolkitCommands(), GetEditingObjects()));
+
 	AddMenuExtender(CreateMenuExtender());
 }
 
