UMG: Remove widget mark as "HideDropDown" from the palette.

diff --git a/Engine/Source/Editor/UMGEditor/Private/Palette/SPaletteViewModel.cpp b/Engine/Source/Editor/UMGEditor/Private/Palette/SPaletteViewModel.cpp
--- a/Engine/Source/Editor/UMGEditor/Private/Palette/SPaletteViewModel.cpp
+++ b/Engine/Source/Editor/UMGEditor/Private/Palette/SPaletteViewModel.cpp
@@ -274,6 +274,11 @@ void FPaletteViewModel::BuildClassWidgetList()
 			continue;
 		}
 
+		if (WidgetClass->HasAnyClassFlags(CLASS_HideDropDown))
+		{
+			continue;
+		}
+
 		// Initialize AssetData for checking PackagePath
 		FAssetData WidgetAssetData = FAssetData(WidgetClass);
 
@@ -411,7 +416,7 @@ void FPaletteViewModel::BuildClassWidgetList()
 			}
 
 			uint32 BPFlags = WidgetBPAssetData.GetTagValueRef<uint32>(FBlueprintTags::ClassFlags);
-			if (!(BPFlags & (CLASS_Abstract | CLASS_Deprecated)))
+			if ((BPFlags & (CLASS_Abstract | CLASS_Deprecated | CLASS_HideDropDown)) != 0)
 			{
 				auto Template = MakeShareable(new FWidgetTemplateBlueprintClass(WidgetBPAssetData, WidgetBPClass));
 
diff --git a/Engine/Source/Editor/UMGEditor/Private/WidgetBlueprintEditorUtils.cpp b/Engine/Source/Editor/UMGEditor/Private/WidgetBlueprintEditorUtils.cpp
--- a/Engine/Source/Editor/UMGEditor/Private/WidgetBlueprintEditorUtils.cpp
+++ b/Engine/Source/Editor/UMGEditor/Private/WidgetBlueprintEditorUtils.cpp
@@ -673,7 +673,7 @@ void FWidgetBlueprintEditorUtils::BuildWrapWithMenu(FMenuBuilder& Menu, TSharedR
 		UClass* WidgetClass = *ClassIt;
 		if ( FWidgetBlueprintEditorUtils::IsUsableWidgetClass(WidgetClass) )
 		{
-			if ( WidgetClass->IsChildOf(UPanelWidget::StaticClass()) )
+			if ( WidgetClass->IsChildOf(UPanelWidget::StaticClass()) && !WidgetClass->HasAnyClassFlags(CLASS_HideDropDown) )
 			{
 				WrapperClasses.Add(WidgetClass);
 			}
@@ -836,7 +836,7 @@ void FWidgetBlueprintEditorUtils::BuildReplaceWithMenu(FMenuBuilder& Menu, TShar
 			UClass* WidgetClass = *ClassIt;
 			if ( FWidgetBlueprintEditorUtils::IsUsableWidgetClass(WidgetClass) )
 			{
-				if ( WidgetClass->IsChildOf(UPanelWidget::StaticClass()) )
+				if ( WidgetClass->IsChildOf(UPanelWidget::StaticClass()) && !WidgetClass->HasAnyClassFlags(CLASS_HideDropDown) )
 				{
 					// Only allow replacement with panels that accept multiple children
 					if ( WidgetClass->GetDefaultObject<UPanelWidget>()->CanHaveMultipleChildren() )
