Fix to UMG Widget Palette not showing Widget blueprint assets

diff --git a/Engine/Source/Editor/UMGEditor/Private/Palette/SPaletteViewModel.cpp b/Engine/Source/Editor/UMGEditor/Private/Palette/SPaletteViewModel.cpp
--- a/Engine/Source/Editor/UMGEditor/Private/Palette/SPaletteViewModel.cpp
+++ b/Engine/Source/Editor/UMGEditor/Private/Palette/SPaletteViewModel.cpp
@@ -416,7 +416,7 @@ void FPaletteViewModel::BuildClassWidgetList()
 			}
 
 			uint32 BPFlags = WidgetBPAssetData.GetTagValueRef<uint32>(FBlueprintTags::ClassFlags);
-			if ((BPFlags & (CLASS_Abstract | CLASS_Deprecated | CLASS_HideDropDown)) != 0)
+			if ((BPFlags & (CLASS_Abstract | CLASS_Deprecated | CLASS_HideDropDown)) == 0)
 			{
 				auto Template = MakeShareable(new FWidgetTemplateBlueprintClass(WidgetBPAssetData, WidgetBPClass));
 
