Mark preview widget as being in Design mode before destroying it, to prevent it from invoking NativeDestruct. Reported in UDN-00292247

diff --git a/Engine/Source/Editor/UMGEditor/Private/WidgetBlueprintEditor.cpp b/Engine/Source/Editor/UMGEditor/Private/WidgetBlueprintEditor.cpp
--- a/Engine/Source/Editor/UMGEditor/Private/WidgetBlueprintEditor.cpp
+++ b/Engine/Source/Editor/UMGEditor/Private/WidgetBlueprintEditor.cpp
@@ -932,6 +932,9 @@ void FWidgetBlueprintEditor::DestroyPreview()
 	{
 		check(PreviewScene.GetWorld());
 
+		// Establish the widget as being in design time before destroying it
+		PreviewUserWidget->SetDesignerFlags(GetCurrentDesignerFlags());
+
 		// Immediately release the preview ptr to let people know it's gone.
 		PreviewWidgetPtr.Reset();
 
