PR #7533: Added AssetTypeActions_SlateVectorArtData class which re-enables creation of USlateVectorArtData assets. (Contributed by [REDACTED_PERSON])

diff --git a/Engine/Source/Editor/UMGEditor/Private/AssetTypeActions_SlateVectorArtData.h b/Engine/Source/Editor/UMGEditor/Private/AssetTypeActions_SlateVectorArtData.h
new file mode 100644
--- /dev/null
+++ b/Engine/Source/Editor/UMGEditor/Private/AssetTypeActions_SlateVectorArtData.h
@@ -0,0 +1,18 @@
+// Copyright Epic Games, Inc. All Rights Reserved.
+
+#pragma once
+
+#include "CoreMinimal.h"
+#include "AssetTypeActions_Base.h"
+#include "Slate/SlateVectorArtData.h"
+
+class FAssetTypeActions_SlateVectorArtData : public FAssetTypeActions_Base
+{
+public:
+	// IAssetTypeActions Implementation
+	virtual FText GetName() const override { return NSLOCTEXT("AssetTypeActions", "AssetTypeActions_SlateVectorArtData", "Slate Vector Art Data"); }
+	virtual FColor GetTypeColor() const override { return FColor(105, 165, 60); }
+	virtual UClass* GetSupportedClass() const override { return USlateVectorArtData::StaticClass(); }
+	virtual uint32 GetCategories() override { return EAssetTypeCategories::UI; }
+	virtual bool HasActions(const TArray<UObject*>& InObjects) const override { return false; }
+}; 
\ No newline at end of file
diff --git a/Engine/Source/Editor/UMGEditor/Private/UMGEditorModule.cpp b/Engine/Source/Editor/UMGEditor/Private/UMGEditorModule.cpp
--- a/Engine/Source/Editor/UMGEditor/Private/UMGEditorModule.cpp
+++ b/Engine/Source/Editor/UMGEditor/Private/UMGEditorModule.cpp
@@ -13,6 +13,7 @@
 
 #include "AssetToolsModule.h"
 #include "IAssetTypeActions.h"
+#include "AssetTypeActions_SlateVectorArtData.h"
 #include "AssetTypeActions_WidgetBlueprint.h"
 #include "AssetTypeActions_WidgetBlueprintGeneratedClass.h"
 #include "KismetCompilerModule.h"
@@ -73,6 +74,7 @@ public:
 
 		// Register asset types
 		IAssetTools& AssetTools = FModuleManager::LoadModuleChecked<FAssetToolsModule>("AssetTools").Get();
+		RegisterAssetTypeAction(AssetTools, MakeShareable(new FAssetTypeActions_SlateVectorArtData()));
 		RegisterAssetTypeAction(AssetTools, MakeShareable(new FAssetTypeActions_WidgetBlueprint()));
 		RegisterAssetTypeAction(AssetTools, MakeShareable(new FAssetTypeActions_WidgetBlueprintGeneratedClass()));
 
