Ensure that when a widget blueprint is compiled, that the newly created WidgetTree and animations supplant the old ones in the export map.

diff --git a/Engine/Source/Editor/UMGEditor/Private/WidgetBlueprintCompiler.cpp b/Engine/Source/Editor/UMGEditor/Private/WidgetBlueprintCompiler.cpp
--- a/Engine/Source/Editor/UMGEditor/Private/WidgetBlueprintCompiler.cpp
+++ b/Engine/Source/Editor/UMGEditor/Private/WidgetBlueprintCompiler.cpp
@@ -96,6 +96,7 @@ bool FWidgetBlueprintCompiler::GetBlueprintTypesForClass(UClass* ParentClass, UC
 FWidgetBlueprintCompilerContext::FWidgetBlueprintCompilerContext(UWidgetBlueprint* SourceSketch, FCompilerResultsLog& InMessageLog, const FKismetCompilerOptions& InCompilerOptions)
 	: Super(SourceSketch, InMessageLog, InCompilerOptions)
 	, NewWidgetBlueprintClass(nullptr)
+	, OldWidgetTree(nullptr)
 	, WidgetSchema(nullptr)
 {
 }
@@ -305,6 +306,14 @@ void FWidgetBlueprintCompilerContext::SaveSubObjectsFromCleanAndSanitizeClass(FS
 	check(ClassToClean == NewClass);
 	NewWidgetBlueprintClass = CastChecked<UWidgetBlueprintGeneratedClass>((UObject*)NewClass);
 
+	OldWidgetTree = nullptr;
+	OldWidgetAnimations.Empty();
+	if (NewWidgetBlueprintClass)
+	{
+		OldWidgetTree = NewWidgetBlueprintClass->GetWidgetTreeArchetype();
+		OldWidgetAnimations.Append(NewWidgetBlueprintClass->Animations);
+	}
+
 	UWidgetBlueprint* WidgetBP = WidgetBlueprint();
 
 	// We need to save the widget tree to survive the initial sub-object clean blitz, 
@@ -606,18 +615,32 @@ void FWidgetBlueprintCompilerContext::FinishCompilingClass(UClass* Class)
 			WidgetBP->WidgetTree->ClearFlags(RF_ArchetypeObject);
 
 			UWidgetTree* NewWidgetTree = Cast<UWidgetTree>(StaticDuplicateObject(WidgetBP->WidgetTree, BPGClass, NAME_None, RF_AllFlags & ~RF_DefaultSubObject));
+
 			BPGClass->SetWidgetTreeArchetype(NewWidgetTree);
+			if (OldWidgetTree)
+			{
+				FLinkerLoad::PRIVATE_PatchNewObjectIntoExport(OldWidgetTree, NewWidgetTree);
+			}
+			OldWidgetTree = nullptr;
 
 			WidgetBP->WidgetTree->SetFlags(PreviousFlags);
 		}
 
+		int32 AnimIndex = 0;
 		for ( const UWidgetAnimation* Animation : WidgetBP->Animations )
 		{
 			UWidgetAnimation* ClonedAnimation = DuplicateObject<UWidgetAnimation>(Animation, BPGClass, *( Animation->GetName() + TEXT("_INST") ));
 			//ClonedAnimation->SetFlags(RF_Public); // Needs to be marked public so that it can be referenced from widget instances.
 
+			if (OldWidgetAnimations[AnimIndex] && (AnimIndex < OldWidgetAnimations.Num()))
+			{
+				FLinkerLoad::PRIVATE_PatchNewObjectIntoExport(OldWidgetAnimations[AnimIndex], ClonedAnimation);
+			}
+
 			BPGClass->Animations.Add(ClonedAnimation);
+			AnimIndex++;
 		}
+		OldWidgetAnimations.Empty();
 
 		// Only check bindings on a full compile.  Also don't check them if we're regenerating on load,
 		// that has a nasty tendency to fail because the other dependent classes that may also be blueprints
diff --git a/Engine/Source/Editor/UMGEditor/Public/WidgetBlueprintCompiler.h b/Engine/Source/Editor/UMGEditor/Public/WidgetBlueprintCompiler.h
--- a/Engine/Source/Editor/UMGEditor/Public/WidgetBlueprintCompiler.h
+++ b/Engine/Source/Editor/UMGEditor/Public/WidgetBlueprintCompiler.h
@@ -81,6 +81,10 @@ protected:
 
 	UWidgetBlueprintGeneratedClass* NewWidgetBlueprintClass;
 
+	UWidgetTree* OldWidgetTree;
+
+	TArray<UWidgetAnimation*> OldWidgetAnimations;
+
 	class UWidgetGraphSchema* WidgetSchema;
 
 	// Map of properties created for widgets; to aid in debug data generation
