Fix to rename of widget that add Brackets when renaming to an empty name.

diff --git a/Engine/Source/Editor/UMGEditor/Private/Hierarchy/SHierarchyViewItem.cpp b/Engine/Source/Editor/UMGEditor/Private/Hierarchy/SHierarchyViewItem.cpp
--- a/Engine/Source/Editor/UMGEditor/Private/Hierarchy/SHierarchyViewItem.cpp
+++ b/Engine/Source/Editor/UMGEditor/Private/Hierarchy/SHierarchyViewItem.cpp
@@ -1054,6 +1054,7 @@ FHierarchyWidget::FHierarchyWidget(FWidgetReference InItem, TSharedPtr<FWidgetBl
 	: FHierarchyModel(InBlueprintEditor)
 	, Item(InItem)
 	, bEditing(false)
+	, bNameTextValid(false)
 {
 }
 
@@ -1209,13 +1210,18 @@ FReply FHierarchyWidget::HandleAcceptDrop(const FDragDropEvent& DragDropEvent, E
 }
 
 bool FHierarchyWidget::OnVerifyNameTextChanged(const FText& InText, FText& OutErrorMessage)
-{
-	return FWidgetBlueprintEditorUtils::VerifyWidgetRename(BlueprintEditor.Pin().ToSharedRef(), Item, InText, OutErrorMessage);
+{	
+	 bNameTextValid = FWidgetBlueprintEditorUtils::VerifyWidgetRename(BlueprintEditor.Pin().ToSharedRef(), Item, InText, OutErrorMessage);
+	 return bNameTextValid;
 }
 
 void FHierarchyWidget::OnNameTextCommited(const FText& InText, ETextCommit::Type CommitInfo)
 {
-	FWidgetBlueprintEditorUtils::RenameWidget(BlueprintEditor.Pin().ToSharedRef(), Item.GetTemplate()->GetFName(), InText.ToString());
+	if (CommitInfo == ETextCommit::OnEnter || bNameTextValid)
+	{
+		FWidgetBlueprintEditorUtils::RenameWidget(BlueprintEditor.Pin().ToSharedRef(), Item.GetTemplate()->GetFName(), InText.ToString());
+	}
+	bNameTextValid = false;
 }
 
 void FHierarchyWidget::GetChildren(TArray< TSharedPtr<FHierarchyModel> >& Children)
diff --git a/Engine/Source/Editor/UMGEditor/Private/Hierarchy/SHierarchyViewItem.h b/Engine/Source/Editor/UMGEditor/Private/Hierarchy/SHierarchyViewItem.h
--- a/Engine/Source/Editor/UMGEditor/Private/Hierarchy/SHierarchyViewItem.h
+++ b/Engine/Source/Editor/UMGEditor/Private/Hierarchy/SHierarchyViewItem.h
@@ -330,6 +330,7 @@ protected:
 private:
 	FWidgetReference Item;
 	bool bEditing;
+	bool bNameTextValid;
 };
 
 /**
diff --git a/Engine/Source/Editor/UMGEditor/Private/WidgetBlueprintEditorUtils.cpp b/Engine/Source/Editor/UMGEditor/Private/WidgetBlueprintEditorUtils.cpp
--- a/Engine/Source/Editor/UMGEditor/Private/WidgetBlueprintEditorUtils.cpp
+++ b/Engine/Source/Editor/UMGEditor/Private/WidgetBlueprintEditorUtils.cpp
@@ -92,7 +92,7 @@ public:
 };
 
 bool FWidgetBlueprintEditorUtils::VerifyWidgetRename(TSharedRef<class FWidgetBlueprintEditor> BlueprintEditor, FWidgetReference Widget, const FText& NewName, FText& OutErrorMessage)
-{
+{	
 	if (NewName.IsEmptyOrWhitespace())
 	{
 		OutErrorMessage = LOCTEXT("EmptyWidgetName", "Empty Widget Name");
@@ -226,15 +226,15 @@ bool FWidgetBlueprintEditorUtils::RenameWidget(TSharedRef<FWidgetBlueprintEditor
 		// Rename Preview before renaming the template widget so the preview widget can be found
 		UWidget* WidgetPreview = BlueprintEditor->GetReferenceFromTemplate(Widget).GetPreview();
 		if(WidgetPreview)
-		{
-			WidgetPreview->SetDisplayLabel(NewDisplayName);
+		{			
+			WidgetPreview->SetDisplayLabel(NewNameStr);
 			WidgetPreview->Rename(*NewNameStr);
 		}
 
 		if (!WidgetPreview || WidgetPreview != Widget)
 		{
 			// Find and update all variable references in the graph
-			Widget->SetDisplayLabel(NewDisplayName);
+			Widget->SetDisplayLabel(NewNameStr);
 			Widget->Rename(*NewNameStr);
 		}
 
