Fix to a crash caused by a null ptr when pasting Widgets.

diff --git a/Engine/Source/Editor/UMGEditor/Private/WidgetBlueprintEditorUtils.cpp b/Engine/Source/Editor/UMGEditor/Private/WidgetBlueprintEditorUtils.cpp
--- a/Engine/Source/Editor/UMGEditor/Private/WidgetBlueprintEditorUtils.cpp
+++ b/Engine/Source/Editor/UMGEditor/Private/WidgetBlueprintEditorUtils.cpp
@@ -1402,7 +1402,9 @@ TArray<UWidget*> FWidgetBlueprintEditorUtils::PasteWidgetsInternal(TSharedRef<FW
 
 	// If we're pasting into a content widget of the same type, treat it as a sibling duplication
 	UWidget* FirstPastedWidget = *PastedWidgets.CreateIterator();
-	if (FirstPastedWidget->IsA(UContentWidget::StaticClass()) && FirstPastedWidget->GetClass() == ParentWidgetRef.GetTemplate()->GetClass())
+	if (FirstPastedWidget->IsA(UContentWidget::StaticClass()) && 
+		ParentWidgetRef.IsValid() &&
+		FirstPastedWidget->GetClass() == ParentWidgetRef.GetTemplate()->GetClass())
 	{
 		UPanelWidget* TargetParentWidget = ParentWidgetRef.GetTemplate()->GetParent();
 		if (TargetParentWidget && TargetParentWidget->CanAddMoreChildren())
@@ -1614,11 +1616,14 @@ TArray<UWidget*> FWidgetBlueprintEditorUtils::PasteWidgetsInternal(TSharedRef<FW
 
 		BP->WidgetTree->Modify();
 
-		NamedSlotHostWidget->SetFlags(RF_Transactional);
-		NamedSlotHostWidget->Modify();
+		if (NamedSlotHostWidget)
+		{
+			NamedSlotHostWidget->SetFlags(RF_Transactional);
+			NamedSlotHostWidget->Modify();
 
-		INamedSlotInterface* NamedSlotInterface = Cast<INamedSlotInterface>(NamedSlotHostWidget);
-		NamedSlotInterface->SetContentForSlot(SlotName, RootPasteWidgets[0]);
+			INamedSlotInterface* NamedSlotInterface = Cast<INamedSlotInterface>(NamedSlotHostWidget);
+			NamedSlotInterface->SetContentForSlot(SlotName, RootPasteWidgets[0]);
+		}
 
 		FBlueprintEditorUtils::MarkBlueprintAsStructurallyModified(BP);
 	}
