fix for a out of order Array element check

diff --git a/Engine/Source/Editor/UMGEditor/Private/WidgetBlueprintCompiler.cpp b/Engine/Source/Editor/UMGEditor/Private/WidgetBlueprintCompiler.cpp
--- a/Engine/Source/Editor/UMGEditor/Private/WidgetBlueprintCompiler.cpp
+++ b/Engine/Source/Editor/UMGEditor/Private/WidgetBlueprintCompiler.cpp
@@ -631,8 +631,8 @@ void FWidgetBlueprintCompilerContext::FinishCompilingClass(UClass* Class)
 		{
 			UWidgetAnimation* ClonedAnimation = DuplicateObject<UWidgetAnimation>(Animation, BPGClass, *( Animation->GetName() + TEXT("_INST") ));
 			//ClonedAnimation->SetFlags(RF_Public); // Needs to be marked public so that it can be referenced from widget instances.
-
-			if (OldWidgetAnimations[AnimIndex] && (AnimIndex < OldWidgetAnimations.Num()))
+			
+			if (OldWidgetAnimations.IsValidIndex(AnimIndex) && OldWidgetAnimations[AnimIndex])
 			{
 				FLinkerLoad::PRIVATE_PatchNewObjectIntoExport(OldWidgetAnimations[AnimIndex], ClonedAnimation);
 			}
