Arrays of structs in Editor Utility Widgets can now be expanded.

diff --git a/Engine/Source/Editor/UMGEditor/Private/Components/DetailsView.cpp b/Engine/Source/Editor/UMGEditor/Private/Components/DetailsView.cpp
--- a/Engine/Source/Editor/UMGEditor/Private/Components/DetailsView.cpp
+++ b/Engine/Source/Editor/UMGEditor/Private/Components/DetailsView.cpp
@@ -159,7 +159,6 @@ void UDetailsView::ToggleWhitelistedProperties()
 	}
 }
 
-
 bool UDetailsView::GetIsPropertyVisible(const FPropertyAndParent& PropertyAndParent) const
 {
     if (!bShowOnlyWhitelisted)
@@ -170,7 +169,22 @@ bool UDetailsView::GetIsPropertyVisible(const FPropertyAndParent& PropertyAndPar
 	{
 		return true;
 	}
-	if (CategoriesToShow.Contains(FObjectEditorUtils::GetCategoryFName(&PropertyAndParent.Property)))
+	// get the topmost parent's category name if the property has one
+	const FProperty* Property = PropertyAndParent.ParentProperties.Num() > 0 ? PropertyAndParent.ParentProperties.Last() : &PropertyAndParent.Property;
+	const FString CategoryString = FObjectEditorUtils::GetCategoryFName(Property).ToString();
+	FString SubcategoryString = CategoryString;
+	
+	int32 SubcategoryStart;
+	if (CategoryString.FindChar(TEXT('|'), SubcategoryStart))
+	{
+		SubcategoryString = CategoryString.Left(SubcategoryStart);
+	}
+
+	if (CategoriesToShow.ContainsByPredicate([&CategoryString, &SubcategoryString](const FName& Element)
+		{
+			FString ElementString = Element.ToString();
+			return CategoryString == ElementString || SubcategoryString == ElementString;
+		}))
 	{
 		return true;
 	}
diff --git a/Engine/Source/Runtime/Engine/Private/ObjectEditorUtils.cpp b/Engine/Source/Runtime/Engine/Private/ObjectEditorUtils.cpp
--- a/Engine/Source/Runtime/Engine/Private/ObjectEditorUtils.cpp
+++ b/Engine/Source/Runtime/Engine/Private/ObjectEditorUtils.cpp
@@ -67,22 +67,20 @@ namespace FObjectEditorUtils
 	FName GetCategoryFName( const FField* InField )
 	{
 		static const FName CategoryKey(TEXT("Category"));
-		if (InField)
-		{
-			const FField* FieldToConsider = InField;
 
-			const bool bIsInsideContainerProperty = InField->GetOwner<FArrayProperty>() || InField->GetOwner<FSetProperty>() || InField->GetOwner<FMapProperty>();
-			if (bIsInsideContainerProperty)
-			{
-				FieldToConsider = InField->GetOwner<FField>();
-			}
+		const FField* CurrentField = InField;
+		FString CategoryString;
+		while (CurrentField != nullptr && CategoryString.IsEmpty())
+		{
+			CategoryString = CurrentField->GetMetaData(CategoryKey);
+			CurrentField = CurrentField->GetOwner<FField>();
+		} 
 
-			FString CategoryString = FieldToConsider->GetMetaData(CategoryKey);
-			if (!CategoryString.IsEmpty())
-			{
-				return FName(*CategoryString);
-			}
+		if (!CategoryString.IsEmpty())
+		{
+			return FName(*CategoryString);
 		}
+
 		return NAME_None;
 	}
 
