Correction to the ReplaceNamedSlotHostContent function to properly remove the widget from its parent before adding it as a NamedSlot.

#rn UMG Designer: Fix to a crash that can happen when using Replace with Child option in the context menu of the hierarchy panel if the Widget being replaced by its child is in a NamedSlot.

diff --git a/Engine/Source/Editor/UMGEditor/Private/WidgetBlueprintEditorUtils.cpp b/Engine/Source/Editor/UMGEditor/Private/WidgetBlueprintEditorUtils.cpp
--- a/Engine/Source/Editor/UMGEditor/Private/WidgetBlueprintEditorUtils.cpp
+++ b/Engine/Source/Editor/UMGEditor/Private/WidgetBlueprintEditorUtils.cpp
@@ -686,6 +686,11 @@ bool FWidgetBlueprintEditorUtils::ReplaceNamedSlotHostContent(UWidget* WidgetTem
 				if (NewContentWidget)
 				{
 					NewContentWidget->Modify();
+					if (UPanelWidget* Parent = NewContentWidget->GetParent())
+					{
+						Parent->Modify();
+						NewContentWidget->RemoveFromParent();
+					}
 				}
 				NamedSlotHost.GetObject()->Modify();
 				NamedSlotHost->SetContentForSlot(SlotName, NewContentWidget);
