Fixed potential issue where transaction buffer could hold a reference to the dummy world created in UWidgetBlueprint::DetectSlateWidgetLeaks
See UDN: https://epicgames.lightning.force.com/lightning/r/Case/5004z00001WeTx3AAF/view

diff --git a/Engine/Source/Editor/UMGEditor/Private/WidgetBlueprint.cpp b/Engine/Source/Editor/UMGEditor/Private/WidgetBlueprint.cpp
--- a/Engine/Source/Editor/UMGEditor/Private/WidgetBlueprint.cpp
+++ b/Engine/Source/Editor/UMGEditor/Private/WidgetBlueprint.cpp
@@ -702,6 +702,7 @@ bool UWidgetBlueprint::DetectSlateWidgetLeaks(TArray<FText>& ValidationErrors)
 		}
 	});
 
+	DummyWorld->MarkObjectsPendingKill();
 	return bFoundLeak;
 }
 
