Fix UMG size rule and text justification not showing their value properly in the details panel

diff --git a/Engine/Source/Editor/UMGEditor/Private/Customizations/SlateChildSizeCustomization.cpp b/Engine/Source/Editor/UMGEditor/Private/Customizations/SlateChildSizeCustomization.cpp
--- a/Engine/Source/Editor/UMGEditor/Private/Customizations/SlateChildSizeCustomization.cpp
+++ b/Engine/Source/Editor/UMGEditor/Private/Customizations/SlateChildSizeCustomization.cpp
@@ -32,51 +32,26 @@ void FSlateChildSizeCustomization::CustomizeHeader(TSharedRef<IPropertyHandle> P
 	.MaxDesiredWidth(TOptional<float>())
 	[
 		SNew(SHorizontalBox)
-
 		+ SHorizontalBox::Slot()
 		.AutoWidth()
 		.Padding(OuterPadding)
 		[
-			SNew(SUniformGridPanel)
-			.SlotPadding(OuterPadding)
-
-			+ SUniformGridPanel::Slot(0, 0)
-			[
-				SNew(SCheckBox)
-				.Style(FEditorStyle::Get(), "ToggleButtonCheckbox")
-				.ToolTipText(LOCTEXT("Auto_ToolTip", "Only requests as much room as it needs based on the widgets desired size."))
-				.Padding(ContentPadding)
-				.OnCheckStateChanged(this, &FSlateChildSizeCustomization::HandleCheckStateChanged, RuleHandle, ESlateSizeRule::Automatic)
-				.IsChecked(this, &FSlateChildSizeCustomization::GetCheckState, RuleHandle, ESlateSizeRule::Automatic)
-				.HAlign(HAlign_Center)
-				[
-					SNew(STextBlock)
-					.Text(LOCTEXT("Auto", "Auto"))
-				]
-			]
-
-			+ SUniformGridPanel::Slot(1, 0)
-			[
-				SNew(SCheckBox)
-				.Style(FEditorStyle::Get(), "ToggleButtonCheckbox")
-				.ToolTipText(LOCTEXT("Fill_ToolTip", "Greedily attempts to fill all available room based on the percentage value 0..1"))
-				.Padding(ContentPadding)
-				.OnCheckStateChanged(this, &FSlateChildSizeCustomization::HandleCheckStateChanged, RuleHandle, ESlateSizeRule::Fill)
-				.IsChecked(this, &FSlateChildSizeCustomization::GetCheckState, RuleHandle, ESlateSizeRule::Fill)
-				.HAlign(HAlign_Center)
-				[
-					SNew(STextBlock)
-					.Text(LOCTEXT("Fill", "Fill"))
-				]
-			]
+			SNew(SSegmentedControl<ESlateSizeRule::Type>)
+			.Value(this, &FSlateChildSizeCustomization::GetCurrentSizeRule, RuleHandle)
+			.OnValueChanged(this, &FSlateChildSizeCustomization::OnSizeRuleChanged, RuleHandle)
+			+ SSegmentedControl<ESlateSizeRule::Type>::Slot(ESlateSizeRule::Automatic)
+			.Text(LOCTEXT("Auto", "Auto"))
+			.ToolTip(LOCTEXT("Auto_ToolTip", "Only requests as much room as it needs based on the widgets desired size."))
+			+ SSegmentedControl<ESlateSizeRule::Type>::Slot(ESlateSizeRule::Fill)
+			.Text(LOCTEXT("Fill", "Fill"))
+			.ToolTip(LOCTEXT("Fill_ToolTip", "Greedily attempts to fill all available room based on the percentage value 0..1"))
 		]
-
 		+ SHorizontalBox::Slot()
 		.AutoWidth()
 		.Padding(OuterPadding)
 		[
 			SNew(SBox)
-			.WidthOverride(40)
+			.WidthOverride(45)
 			[
 				SNew( SNumericEntryBox<float> )
 				.LabelVAlign(VAlign_Center)
@@ -94,20 +69,20 @@ void FSlateChildSizeCustomization::CustomizeChildren(TSharedRef<IPropertyHandle>
 	
 }
 
-void FSlateChildSizeCustomization::HandleCheckStateChanged(ECheckBoxState InCheckboxState, TSharedPtr<IPropertyHandle> PropertyHandle, ESlateSizeRule::Type ToRule)
+void FSlateChildSizeCustomization::OnSizeRuleChanged(ESlateSizeRule::Type ToRule, TSharedPtr<IPropertyHandle> PropertyHandle)
 {
 	PropertyHandle->SetValue((uint8)ToRule);
 }
 
-ECheckBoxState FSlateChildSizeCustomization::GetCheckState(TSharedPtr<IPropertyHandle> PropertyHandle, ESlateSizeRule::Type ForRule) const
+ESlateSizeRule::Type FSlateChildSizeCustomization::GetCurrentSizeRule(TSharedPtr<IPropertyHandle> PropertyHandle) const
 {
 	uint8 Value;
-	if ( PropertyHandle->GetValue(Value) == FPropertyAccess::Result::Success )
+	if (PropertyHandle->GetValue(Value) == FPropertyAccess::Result::Success)
 	{
-		return Value == ForRule ? ECheckBoxState::Checked : ECheckBoxState::Unchecked;
+		return (ESlateSizeRule::Type)Value;
 	}
 
-	return ECheckBoxState::Unchecked;
+	return ESlateSizeRule::Automatic;
 }
 
 TOptional<float> FSlateChildSizeCustomization::GetValue(TSharedPtr<IPropertyHandle> ValueHandle) const
diff --git a/Engine/Source/Editor/UMGEditor/Private/Customizations/SlateChildSizeCustomization.h b/Engine/Source/Editor/UMGEditor/Private/Customizations/SlateChildSizeCustomization.h
--- a/Engine/Source/Editor/UMGEditor/Private/Customizations/SlateChildSizeCustomization.h
+++ b/Engine/Source/Editor/UMGEditor/Private/Customizations/SlateChildSizeCustomization.h
@@ -25,9 +25,9 @@ public:
 
 	virtual void CustomizeChildren(TSharedRef<IPropertyHandle> PropertyHandle, class IDetailChildrenBuilder& ChildBuilder, IPropertyTypeCustomizationUtils& CustomizationUtils) override;
 
-	void HandleCheckStateChanged(ECheckBoxState InCheckboxState, TSharedPtr<IPropertyHandle> PropertyHandle, ESlateSizeRule::Type ToRule);
+	void OnSizeRuleChanged(ESlateSizeRule::Type ToRule, TSharedPtr<IPropertyHandle> PropertyHandle);
 
-	ECheckBoxState GetCheckState(TSharedPtr<IPropertyHandle> PropertyHandle, ESlateSizeRule::Type ForRule) const;
+	ESlateSizeRule::Type GetCurrentSizeRule(TSharedPtr<IPropertyHandle> PropertyHandle) const;
 
 	TOptional<float> GetValue(TSharedPtr<IPropertyHandle> ValueHandle) const;
 
diff --git a/Engine/Source/Editor/UMGEditor/Private/Customizations/TextJustifyCustomization.cpp b/Engine/Source/Editor/UMGEditor/Private/Customizations/TextJustifyCustomization.cpp
--- a/Engine/Source/Editor/UMGEditor/Private/Customizations/TextJustifyCustomization.cpp
+++ b/Engine/Source/Editor/UMGEditor/Private/Customizations/TextJustifyCustomization.cpp
@@ -9,11 +9,7 @@
 #define LOCTEXT_NAMESPACE "UMG"
 
 void FTextJustifyCustomization::CustomizeHeader(TSharedRef<IPropertyHandle> PropertyHandle, FDetailWidgetRow& HeaderRow, IPropertyTypeCustomizationUtils& CustomizationUtils)
-{
-	const FMargin OuterPadding(2);
-	const FMargin ContentPadding(2);
-
-	
+{	
 	HeaderRow
 	.IsEnabled(TAttribute<bool>(PropertyHandle, &IPropertyHandle::IsEditable))
 	.NameContent()
@@ -22,55 +18,18 @@ void FTextJustifyCustomization::CustomizeHeader(TSharedRef<IPropertyHandle> Prop
 	]
 	.ValueContent()
 	[
-		SNew(SHorizontalBox)
-
-		+ SHorizontalBox::Slot()
-		.AutoWidth()
-		.Padding(OuterPadding)
-		[
-			SNew( SCheckBox )
-			.Style(FEditorStyle::Get(), "ToggleButtonCheckbox")
-			.ToolTipText(LOCTEXT("AlignTextLeft", "Align Text Left"))
-			.Padding(ContentPadding)
-			.OnCheckStateChanged(this, &FTextJustifyCustomization::HandleCheckStateChanged, PropertyHandle, ETextJustify::Left)
-			.IsChecked(this, &FTextJustifyCustomization::GetCheckState, PropertyHandle, ETextJustify::Left)
-			[
-				SNew(SImage)
-				.Image(FEditorStyle::GetBrush("HorizontalAlignment_Left"))
-			]
-		]
-
-		+ SHorizontalBox::Slot()
-		.AutoWidth()
-		.Padding(OuterPadding)
-		[
-			SNew(SCheckBox)
-			.Style(FEditorStyle::Get(), "ToggleButtonCheckbox")
-			.ToolTipText(LOCTEXT("AlignTextCenter", "Align Text Center"))
-			.Padding(ContentPadding)
-			.OnCheckStateChanged(this, &FTextJustifyCustomization::HandleCheckStateChanged, PropertyHandle, ETextJustify::Center)
-			.IsChecked(this, &FTextJustifyCustomization::GetCheckState, PropertyHandle, ETextJustify::Center)
-			[
-				SNew(SImage)
-				.Image(FEditorStyle::GetBrush("HorizontalAlignment_Center"))
-			]
-		]
-
-		+ SHorizontalBox::Slot()
-		.AutoWidth()
-		.Padding(OuterPadding)
-		[
-			SNew(SCheckBox)
-			.Style(FEditorStyle::Get(), "ToggleButtonCheckbox")
-			.ToolTipText(LOCTEXT("AlignTextRight", "Align Text Right"))
-			.Padding(ContentPadding)
-			.OnCheckStateChanged(this, &FTextJustifyCustomization::HandleCheckStateChanged, PropertyHandle, ETextJustify::Right)
-			.IsChecked(this, &FTextJustifyCustomization::GetCheckState, PropertyHandle, ETextJustify::Right)
-			[
-				SNew(SImage)
-				.Image(FEditorStyle::GetBrush("HorizontalAlignment_Right"))
-			]
-		]
+		SNew(SSegmentedControl<ETextJustify::Type>)
+		.Value(this, &FTextJustifyCustomization::GetCurrentJustification, PropertyHandle)
+		.OnValueChanged(this, &FTextJustifyCustomization::OnJustificationChanged, PropertyHandle)
+		+ SSegmentedControl<ETextJustify::Type>::Slot(ETextJustify::Left)
+		.Icon(FEditorStyle::GetBrush("HorizontalAlignment_Left"))
+		.ToolTip(LOCTEXT("AlignTextLeft", "Align Text Left"))
+		+ SSegmentedControl<ETextJustify::Type>::Slot(ETextJustify::Center)
+		.Icon(FEditorStyle::GetBrush("HorizontalAlignment_Center"))
+		.ToolTip(LOCTEXT("AlignTextCenter", "Align Text Center"))
+		+ SSegmentedControl<ETextJustify::Type>::Slot(ETextJustify::Right)
+		.Icon(FEditorStyle::GetBrush("HorizontalAlignment_Right"))
+		.ToolTip(LOCTEXT("AlignTextRight", "Align Text Right"))
 	];
 }
 
@@ -78,20 +37,20 @@ void FTextJustifyCustomization::CustomizeChildren(TSharedRef<IPropertyHandle> Pr
 {
 }
 
-void FTextJustifyCustomization::HandleCheckStateChanged(ECheckBoxState InCheckboxState, TSharedRef<IPropertyHandle> PropertyHandle, ETextJustify::Type ToAlignment)
+void FTextJustifyCustomization::OnJustificationChanged(ETextJustify::Type NewState, TSharedRef<IPropertyHandle> PropertyHandle)
 {
-	PropertyHandle->SetValue((uint8)ToAlignment);
+	PropertyHandle->SetValue((uint8)NewState);
 }
 
-ECheckBoxState FTextJustifyCustomization::GetCheckState(TSharedRef<IPropertyHandle> PropertyHandle, ETextJustify::Type ForAlignment) const
+ETextJustify::Type FTextJustifyCustomization::GetCurrentJustification(TSharedRef<IPropertyHandle> PropertyHandle) const
 {
 	uint8 Value;
 	if ( PropertyHandle->GetValue(Value) == FPropertyAccess::Result::Success)
 	{
-		return Value == ForAlignment ? ECheckBoxState::Checked : ECheckBoxState::Unchecked;
+		return ETextJustify::Type(Value);
 	}
 
-	return ECheckBoxState::Unchecked;
+	return ETextJustify::Left;
 }
 
 #undef LOCTEXT_NAMESPACE
diff --git a/Engine/Source/Editor/UMGEditor/Private/Customizations/TextJustifyCustomization.h b/Engine/Source/Editor/UMGEditor/Private/Customizations/TextJustifyCustomization.h
--- a/Engine/Source/Editor/UMGEditor/Private/Customizations/TextJustifyCustomization.h
+++ b/Engine/Source/Editor/UMGEditor/Private/Customizations/TextJustifyCustomization.h
@@ -23,8 +23,7 @@ public:
 	virtual void CustomizeHeader(TSharedRef<IPropertyHandle> PropertyHandle, class FDetailWidgetRow& HeaderRow, IPropertyTypeCustomizationUtils& CustomizationUtils) override;
 
 	virtual void CustomizeChildren(TSharedRef<IPropertyHandle> PropertyHandle, class IDetailChildrenBuilder& ChildBuilder, IPropertyTypeCustomizationUtils& CustomizationUtils) override;
-	
-	void HandleCheckStateChanged(ECheckBoxState InCheckboxState, TSharedRef<IPropertyHandle> PropertyHandle, ETextJustify::Type ToAlignment);
 
-	ECheckBoxState GetCheckState(TSharedRef<IPropertyHandle> PropertyHandle, ETextJustify::Type ForAlignment) const;
+	void OnJustificationChanged(ETextJustify::Type NewState, TSharedRef<IPropertyHandle> PropertyHandle);
+	ETextJustify::Type GetCurrentJustification(TSharedRef<IPropertyHandle> PropertyHandle) const;
 };
