Fix CIS

diff --git a/Engine/Source/Editor/UMGEditor/Private/Customizations/SlateChildSizeCustomization.cpp b/Engine/Source/Editor/UMGEditor/Private/Customizations/SlateChildSizeCustomization.cpp
--- a/Engine/Source/Editor/UMGEditor/Private/Customizations/SlateChildSizeCustomization.cpp
+++ b/Engine/Source/Editor/UMGEditor/Private/Customizations/SlateChildSizeCustomization.cpp
@@ -7,6 +7,7 @@
 #include "Widgets/Input/SCheckBox.h"
 #include "DetailWidgetRow.h"
 #include "Widgets/Input/SNumericEntryBox.h"
+#include "Widgets/Input/SSegmentedControl.h"
 
 #define LOCTEXT_NAMESPACE "UMG"
 
diff --git a/Engine/Source/Editor/UMGEditor/Private/Customizations/TextJustifyCustomization.cpp b/Engine/Source/Editor/UMGEditor/Private/Customizations/TextJustifyCustomization.cpp
--- a/Engine/Source/Editor/UMGEditor/Private/Customizations/TextJustifyCustomization.cpp
+++ b/Engine/Source/Editor/UMGEditor/Private/Customizations/TextJustifyCustomization.cpp
@@ -4,6 +4,7 @@
 #include "Widgets/Images/SImage.h"
 #include "Widgets/Input/SCheckBox.h"
 #include "DetailWidgetRow.h"
+#include "Widgets/Input/SSegmentedControl.h"
 
 
 #define LOCTEXT_NAMESPACE "UMG"
