Created thumbnails in the content browser for widget BPs #RB patrick.boutot #jira none

diff --git a/Engine/Source/Editor/UMGEditor/Private/UMGEditorModule.cpp b/Engine/Source/Editor/UMGEditor/Private/UMGEditorModule.cpp
--- a/Engine/Source/Editor/UMGEditor/Private/UMGEditorModule.cpp
+++ b/Engine/Source/Editor/UMGEditor/Private/UMGEditorModule.cpp
@@ -38,6 +38,7 @@
 #include "PropertyEditorModule.h"
 #include "DynamicEntryBoxDetails.h"
 #include "ListViewBaseDetails.h"
+#include "WidgetBlueprintThumbnailRenderer.h"
 
 #define LOCTEXT_NAMESPACE "UMG"
 
@@ -94,11 +95,17 @@ public:
 		PropertyModule.RegisterCustomClassLayout(TEXT("DynamicEntryBoxBase"), FOnGetDetailCustomizationInstance::CreateStatic(&FDynamicEntryBoxBaseDetails::MakeInstance));
 		PropertyModule.RegisterCustomClassLayout(TEXT("DynamicEntryBox"), FOnGetDetailCustomizationInstance::CreateStatic(&FDynamicEntryBoxDetails::MakeInstance));
 		PropertyModule.RegisterCustomClassLayout(TEXT("ListViewBase"), FOnGetDetailCustomizationInstance::CreateStatic(&FListViewBaseDetails::MakeInstance));
+
+		UThumbnailManager::Get().RegisterCustomRenderer(UWidgetBlueprint::StaticClass(), UWidgetBlueprintThumbnailRenderer::StaticClass());
 	}
 
 	/** Called before the module is unloaded, right before the module object is destroyed. */
 	virtual void ShutdownModule() override
 	{
+		if (UObjectInitialized()) 
+		{
+			UThumbnailManager::Get().UnregisterCustomRenderer(UWidgetBlueprint::StaticClass());
+		}
 		MenuExtensibilityManager.Reset();
 		ToolBarExtensibilityManager.Reset();
 
diff --git a/Engine/Source/Editor/UMGEditor/Private/WidgetBlueprintThumbnailRenderer.cpp b/Engine/Source/Editor/UMGEditor/Private/WidgetBlueprintThumbnailRenderer.cpp
new file mode 100644
--- /dev/null
+++ b/Engine/Source/Editor/UMGEditor/Private/WidgetBlueprintThumbnailRenderer.cpp
@@ -0,0 +1,135 @@
+// Copyright Epic Games, Inc. All Rights Reserved.
+
+#include "WidgetBlueprintThumbnailRenderer.h"
+
+#include "Blueprint/UserWidget.h"
+#include "CanvasItem.h"
+#include "Engine/TextureRenderTarget2D.h"
+#include "Input/HittestGrid.h"
+#include "Slate/WidgetRenderer.h"
+#include "WidgetBlueprint.h"
+#include "Widgets/SVirtualWindow.h"
+
+#define LOCTEXT_NAMESPACE "UWidgetBlueprintThumbnailRenderer"
+
+bool UWidgetBlueprintThumbnailRenderer::CanVisualizeAsset(UObject* Object)
+{
+	UWidgetBlueprint* Blueprint = Cast<UWidgetBlueprint>(Object);
+	return (Blueprint && Blueprint->GeneratedClass && Blueprint->GeneratedClass->IsChildOf(UWidget::StaticClass()));
+}
+
+void UWidgetBlueprintThumbnailRenderer::Draw(UObject* Object, int32 X, int32 Y, uint32 Width, uint32 Height, FRenderTarget* RenderTarget, FCanvas* Canvas, bool bAdditionalViewFamily)
+{
+	// Create a plain gray background for the thumbnail
+	const int32 SizeOfUV = 1;
+	FLinearColor GrayBackgroundColor(FVector4(.03f, .03f, .03f, 1.f));
+	Canvas->DrawTile(
+		0.0f, 0.0f, Width, Height,
+		0.0f, 0.0f, SizeOfUV, SizeOfUV,
+		GrayBackgroundColor);
+
+
+	// Convert the object from UObject to SWidget
+	UWidgetBlueprint* WidgetBlueprintToRender = Cast<UWidgetBlueprint>(Object);
+
+	const bool bIsBlueprintValid = IsValid(WidgetBlueprintToRender)
+		&& IsValid(WidgetBlueprintToRender->GeneratedClass)
+		&& WidgetBlueprintToRender->bHasBeenRegenerated
+		//&& Blueprint->IsUpToDate() - This condition makes the thumbnail blank whenever the BP is dirty. It seems too strict.
+		&& !WidgetBlueprintToRender->bBeingCompiled
+		&& !WidgetBlueprintToRender->HasAnyFlags(RF_Transient);
+
+	TSharedPtr<SWidget> WindowContent;
+	UClass* ClassToGenerate = WidgetBlueprintToRender->GeneratedClass;
+	if (bIsBlueprintValid && ClassToGenerate->IsChildOf(UWidget::StaticClass()))
+	{
+		UWidget* WidgetInstance = NewObject<UWidget>(GetTransientPackage(), ClassToGenerate);
+		WindowContent = WidgetInstance->TakeWidget();
+	}
+
+	if (WindowContent == nullptr)
+	{
+		return;
+	}
+
+	//Create Window
+	FVector2D WidgetSize(Width, Height);
+	TSharedRef<SVirtualWindow> Window = SNew(SVirtualWindow).Size(WidgetSize);
+	TUniquePtr<FHittestGrid> HitTestGrid = MakeUnique<FHittestGrid>();
+	Window->SetContent(WindowContent.ToSharedRef());
+	Window->Resize(WidgetSize);
+
+	// Store the desired size to maintain the aspect ratio later
+	FGeometry WindowGeometry = FGeometry::MakeRoot(WidgetSize, FSlateLayoutTransform(1.0f));
+	Window->SlatePrepass(1.0f);
+	FVector2D DesiredSizeWindow = Window->GetDesiredSize();
+
+	
+	if (DesiredSizeWindow.X < SMALL_NUMBER || DesiredSizeWindow.Y < SMALL_NUMBER)
+	{
+		return;
+	}
+
+	// Create Renderer Target and WidgetRenderer
+	bool bApplyGammaCorrection = false;
+	FWidgetRenderer* WidgetRenderer = new FWidgetRenderer(bApplyGammaCorrection);
+	WidgetRenderer->SetIsPrepassNeeded(false);
+	const bool bIsLinearSpace = !bApplyGammaCorrection;
+
+	UTextureRenderTarget2D* RenderTarget2D = NewObject<UTextureRenderTarget2D>();
+	RenderTarget2D->Filter = TF_Bilinear;
+	RenderTarget2D->ClearColor = FLinearColor::Transparent;
+	RenderTarget2D->SRGB = true;
+	RenderTarget2D->TargetGamma = 1;
+
+	const EPixelFormat RequestedFormat = FSlateApplication::Get().GetRenderer()->GetSlateRecommendedColorFormat();
+	RenderTarget2D->InitCustomFormat(DesiredSizeWindow.X, DesiredSizeWindow.Y, RequestedFormat, true);
+
+	WidgetRenderer->DrawWindow(RenderTarget2D, *HitTestGrid, Window, 1.0f, DesiredSizeWindow, 0.1f);
+
+	// Draw the SWidget on canvas
+
+	TTuple<FVector2D, FVector2D> ScaledSizeAndOffset = GetScaledSizeAndOffset(DesiredSizeWindow.X, DesiredSizeWindow.Y, Width, Height);
+	FVector2D DrawOffset = ScaledSizeAndOffset.Get<1>();
+	WidgetSize = ScaledSizeAndOffset.Get<0>();
+
+	FCanvasTileItem CanvasTile(FVector2D(X + DrawOffset.X, Y + DrawOffset.Y), RenderTarget2D->GetResource(), WidgetSize, FLinearColor::White);
+	CanvasTile.BlendMode = SE_BLEND_Translucent;
+	CanvasTile.Draw(Canvas);
+
+	if (WidgetRenderer)
+	{
+		BeginCleanup(WidgetRenderer);
+		WidgetRenderer = nullptr;
+	}
+
+}
+
+TTuple<FVector2D, FVector2D> UWidgetBlueprintThumbnailRenderer::GetScaledSizeAndOffset(float ImgWidth, float ImgHeight, float ThumbnailWidth, float ThumbnailHeight) const
+{
+	// Scale the widget blueprint image to fit in the thumbnail
+
+	checkf(ImgWidth > 0.f && ImgHeight > 0.f, TEXT("The size should have been previously checked to be > 0."));
+
+	FVector2D ScaledSize(0.0f,0.0f);
+	FVector2D Offset(0.0f, 0.0f);
+	if (ImgWidth > ImgHeight)
+	{
+		float RatioAdjust = ImgHeight / ImgWidth;
+		Offset.Y = ThumbnailHeight * (1 - RatioAdjust) / 2;
+		ScaledSize.Y = ThumbnailHeight * RatioAdjust;
+		ScaledSize.X = ThumbnailWidth;
+	}
+	else
+	{
+		float RatioAdjust = ImgWidth / ImgHeight;
+		Offset.X = ThumbnailWidth * (1 - RatioAdjust) / 2;
+		ScaledSize.X = ThumbnailWidth * RatioAdjust;
+		ScaledSize.Y = ThumbnailHeight;
+	}
+	
+	
+	return TTuple<FVector2D, FVector2D>(ScaledSize, Offset);
+}
+
+#undef LOCTEXT_NAMESPACE
\ No newline at end of file
diff --git a/Engine/Source/Editor/UMGEditor/Public/WidgetBlueprintThumbnailRenderer.h b/Engine/Source/Editor/UMGEditor/Public/WidgetBlueprintThumbnailRenderer.h
new file mode 100644
--- /dev/null
+++ b/Engine/Source/Editor/UMGEditor/Public/WidgetBlueprintThumbnailRenderer.h
@@ -0,0 +1,24 @@
+// Copyright Epic Games, Inc. All Rights Reserved.
+
+#pragma once
+
+
+#include "ThumbnailRendering/DefaultSizedThumbnailRenderer.h"
+#include "WidgetBlueprintThumbnailRenderer.generated.h"
+
+class FCanvas;
+class FRenderTarget;
+
+UCLASS()
+class UMGEDITOR_API UWidgetBlueprintThumbnailRenderer : public UDefaultSizedThumbnailRenderer
+{
+	GENERATED_BODY()
+		
+	//~ Begin UThumbnailRenderer Object
+	bool CanVisualizeAsset(UObject* Object) override;
+	void Draw(UObject* Object, int32 X, int32 Y, uint32 Width, uint32 Height, FRenderTarget* RenderTarget, FCanvas* Canvas, bool bAdditionalViewFamily) override;
+
+private:
+	TTuple<FVector2D, FVector2D> GetScaledSizeAndOffset(float ImgWidth, float ImgHeight, float ThumbnailWidth, float ThumbnailHeight) const;
+
+};
