Fix for no unity build error
#rnx

diff --git a/Engine/Source/Editor/UMGEditor/Private/WidgetBlueprintThumbnailRenderer.cpp b/Engine/Source/Editor/UMGEditor/Private/WidgetBlueprintThumbnailRenderer.cpp
--- a/Engine/Source/Editor/UMGEditor/Private/WidgetBlueprintThumbnailRenderer.cpp
+++ b/Engine/Source/Editor/UMGEditor/Private/WidgetBlueprintThumbnailRenderer.cpp
@@ -4,6 +4,7 @@
 
 #include "Blueprint/UserWidget.h"
 #include "CanvasItem.h"
+#include "CanvasTypes.h"
 #include "Engine/TextureRenderTarget2D.h"
 #include "Input/HittestGrid.h"
 #include "Slate/WidgetRenderer.h"
