UMG - When pasting widgets, handle the case where the pasted child widget can't be imported.  Can happen when pasting between versions that may not have all the same versions.

diff --git a/Engine/Source/Editor/UMGEditor/Private/WidgetBlueprintEditorUtils.cpp b/Engine/Source/Editor/UMGEditor/Private/WidgetBlueprintEditorUtils.cpp
--- a/Engine/Source/Editor/UMGEditor/Private/WidgetBlueprintEditorUtils.cpp
+++ b/Engine/Source/Editor/UMGEditor/Private/WidgetBlueprintEditorUtils.cpp
@@ -1701,7 +1701,11 @@ void FWidgetBlueprintEditorUtils::ImportWidgetsFromText(UWidgetBlueprint* BP, co
 			TArray<UPanelSlot*> PanelSlots = PanelWidget->GetSlots();
 			for (int32 i = 0; i < PanelWidget->GetChildrenCount(); i++)
 			{
-				PanelWidget->GetChildAt(i)->Slot = PanelSlots[i];
+				UWidget* PanelChild = PanelWidget->GetChildAt(i);
+				if (ensure(PanelChild))
+				{
+					PanelChild->Slot = PanelSlots[i];
+				}
 			}
 		}
 
