UMG - Thumbnail renderer needs to include the designer tags so that we don't run widget logic.

diff --git a/Engine/Source/Editor/UMGEditor/Private/WidgetBlueprintThumbnailRenderer.cpp b/Engine/Source/Editor/UMGEditor/Private/WidgetBlueprintThumbnailRenderer.cpp
--- a/Engine/Source/Editor/UMGEditor/Private/WidgetBlueprintThumbnailRenderer.cpp
+++ b/Engine/Source/Editor/UMGEditor/Private/WidgetBlueprintThumbnailRenderer.cpp
@@ -45,6 +45,7 @@ void UWidgetBlueprintThumbnailRenderer::Draw(UObject* Object, int32 X, int32 Y,
 	if (bIsBlueprintValid && ClassToGenerate->IsChildOf(UWidget::StaticClass()))
 	{
 		UWidget* WidgetInstance = NewObject<UWidget>(GetTransientPackage(), ClassToGenerate);
+		WidgetInstance->SetDesignerFlags(EWidgetDesignFlags::Designing | EWidgetDesignFlags::ExecutePreConstruct);
 		WindowContent = WidgetInstance->TakeWidget();
 	}
 
