UMG - Thumbnail renderer needs to initialize the widget before we attempt to take to take the widget, which will do the initialize if it hasnt happened, but too late, for us to initialize the designer flags correctly.

diff --git a/Engine/Source/Editor/UMGEditor/Private/WidgetBlueprintThumbnailRenderer.cpp b/Engine/Source/Editor/UMGEditor/Private/WidgetBlueprintThumbnailRenderer.cpp
--- a/Engine/Source/Editor/UMGEditor/Private/WidgetBlueprintThumbnailRenderer.cpp
+++ b/Engine/Source/Editor/UMGEditor/Private/WidgetBlueprintThumbnailRenderer.cpp
@@ -44,9 +44,13 @@ void UWidgetBlueprintThumbnailRenderer::Draw(UObject* Object, int32 X, int32 Y,
 	UClass* ClassToGenerate = WidgetBlueprintToRender->GeneratedClass;
 	if (bIsBlueprintValid && ClassToGenerate->IsChildOf(UWidget::StaticClass()))
 	{
-		UWidget* WidgetInstance = NewObject<UWidget>(GetTransientPackage(), ClassToGenerate);
-		WidgetInstance->SetDesignerFlags(EWidgetDesignFlags::Designing | EWidgetDesignFlags::ExecutePreConstruct);
-		WindowContent = WidgetInstance->TakeWidget();
+		if (UUserWidget* WidgetInstance = NewObject<UUserWidget>(GetTransientPackage(), ClassToGenerate))
+		{
+			WidgetInstance->Initialize();
+			WidgetInstance->SetDesignerFlags(EWidgetDesignFlags::Designing | EWidgetDesignFlags::ExecutePreConstruct);
+
+			WindowContent = WidgetInstance->TakeWidget();
+		}
 	}
 
 	if (WindowContent == nullptr)
