UMG: Reenable SDesignerView::DrawSafeZone.

diff --git a/Engine/Source/Editor/UMGEditor/Private/Designer/SDesignerView.cpp b/Engine/Source/Editor/UMGEditor/Private/Designer/SDesignerView.cpp
--- a/Engine/Source/Editor/UMGEditor/Private/Designer/SDesignerView.cpp
+++ b/Engine/Source/Editor/UMGEditor/Private/Designer/SDesignerView.cpp
@@ -2169,7 +2169,7 @@ void SDesignerView::PopulateWidgetGeometryCache_Loop(FArrangedWidget& CurrentWid
 int32 SDesignerView::HandleEffectsPainting(const FOnPaintHandlerParams& PaintArgs)
 {
 	DrawSelectionAndHoverOutline(PaintArgs);
-	//DrawSafeZone(PaintArgs);
+	DrawSafeZone(PaintArgs);
 
 	return PaintArgs.Layer + 1;
 }
