- Cache Widget Thumbnail Render Target to prevent excessive memory usage #RB patrick.boutot #jira ue-119556

diff --git a/Engine/Source/Editor/UMGEditor/Private/WidgetBlueprintThumbnailRenderer.cpp b/Engine/Source/Editor/UMGEditor/Private/WidgetBlueprintThumbnailRenderer.cpp
--- a/Engine/Source/Editor/UMGEditor/Private/WidgetBlueprintThumbnailRenderer.cpp
+++ b/Engine/Source/Editor/UMGEditor/Private/WidgetBlueprintThumbnailRenderer.cpp
@@ -82,11 +82,14 @@ void UWidgetBlueprintThumbnailRenderer::Draw(UObject* Object, int32 X, int32 Y,
 	WidgetRenderer->SetIsPrepassNeeded(false);
 	const bool bIsLinearSpace = !bApplyGammaCorrection;
 
-	UTextureRenderTarget2D* RenderTarget2D = NewObject<UTextureRenderTarget2D>();
-	RenderTarget2D->Filter = TF_Bilinear;
-	RenderTarget2D->ClearColor = FLinearColor::Transparent;
-	RenderTarget2D->SRGB = true;
-	RenderTarget2D->TargetGamma = 1;
+	if (!RenderTarget2D)
+	{
+		RenderTarget2D  = NewObject<UTextureRenderTarget2D>();
+		RenderTarget2D->Filter = TF_Bilinear;
+		RenderTarget2D->ClearColor = FLinearColor::Transparent;
+		RenderTarget2D->SRGB = true;
+		RenderTarget2D->TargetGamma = 1;
+	}
 
 	const EPixelFormat RequestedFormat = FSlateApplication::Get().GetRenderer()->GetSlateRecommendedColorFormat();
 	RenderTarget2D->InitCustomFormat(DesiredSizeWindow.X, DesiredSizeWindow.Y, RequestedFormat, true);
@@ -101,6 +104,7 @@ void UWidgetBlueprintThumbnailRenderer::Draw(UObject* Object, int32 X, int32 Y,
 
 	FCanvasTileItem CanvasTile(FVector2D(X + DrawOffset.X, Y + DrawOffset.Y), RenderTarget2D->GetResource(), WidgetSize, FLinearColor::White);
 	CanvasTile.BlendMode = SE_BLEND_Translucent;
+	FlushRenderingCommands();
 	CanvasTile.Draw(Canvas);
 
 	if (WidgetRenderer)
diff --git a/Engine/Source/Editor/UMGEditor/Public/WidgetBlueprintThumbnailRenderer.h b/Engine/Source/Editor/UMGEditor/Public/WidgetBlueprintThumbnailRenderer.h
--- a/Engine/Source/Editor/UMGEditor/Public/WidgetBlueprintThumbnailRenderer.h
+++ b/Engine/Source/Editor/UMGEditor/Public/WidgetBlueprintThumbnailRenderer.h
@@ -8,6 +8,7 @@
 
 class FCanvas;
 class FRenderTarget;
+class UTextureRenderTarget2D;
 
 UCLASS()
 class UMGEDITOR_API UWidgetBlueprintThumbnailRenderer : public UDefaultSizedThumbnailRenderer
@@ -21,4 +22,7 @@ class UMGEDITOR_API UWidgetBlueprintThumbnailRenderer : public UDefaultSizedThum
 private:
 	TTuple<FVector2D, FVector2D> GetScaledSizeAndOffset(float ImgWidth, float ImgHeight, float ThumbnailWidth, float ThumbnailHeight) const;
 
+	UPROPERTY(Transient)
+	UTextureRenderTarget2D* RenderTarget2D;
+
 };
