- Add a CVar option to disable widget thumbnail rendering #RB patrick.boutot #jira ue-119462

diff --git a/Engine/Source/Editor/UMGEditor/Private/UMGEditorModule.cpp b/Engine/Source/Editor/UMGEditor/Private/UMGEditorModule.cpp
--- a/Engine/Source/Editor/UMGEditor/Private/UMGEditorModule.cpp
+++ b/Engine/Source/Editor/UMGEditor/Private/UMGEditorModule.cpp
@@ -43,6 +43,11 @@
 #define LOCTEXT_NAMESPACE "UMG"
 
 const FName UMGEditorAppIdentifier = FName(TEXT("UMGEditorApp"));
+static  TAutoConsoleVariable<bool> CVarThumbnailRenderEnable(
+	TEXT("UMG.Thumbnail.Renderer.Enable"),
+	true,
+	TEXT("Option to enable/disable thumbnail rendering.")
+);
 
 class FUMGEditorModule : public IUMGEditorModule, public FGCObject
 {
@@ -96,16 +101,24 @@ public:
 		PropertyModule.RegisterCustomClassLayout(TEXT("DynamicEntryBox"), FOnGetDetailCustomizationInstance::CreateStatic(&FDynamicEntryBoxDetails::MakeInstance));
 		PropertyModule.RegisterCustomClassLayout(TEXT("ListViewBase"), FOnGetDetailCustomizationInstance::CreateStatic(&FListViewBaseDetails::MakeInstance));
 
-		UThumbnailManager::Get().RegisterCustomRenderer(UWidgetBlueprint::StaticClass(), UWidgetBlueprintThumbnailRenderer::StaticClass());
+	
+		CVarThumbnailRenderEnable->AsVariable()->SetOnChangedCallback(FConsoleVariableDelegate::CreateStatic(&FUMGEditorModule::ThumbnailRenderingEnabled));
+		
+		if (IConsoleManager::Get().FindConsoleVariable(TEXT("UMG.Thumbnail.Renderer.Enable"))->GetBool())
+		{
+			UThumbnailManager::Get().RegisterCustomRenderer(UWidgetBlueprint::StaticClass(), UWidgetBlueprintThumbnailRenderer::StaticClass());
+
+		}
 	}
 
 	/** Called before the module is unloaded, right before the module object is destroyed. */
 	virtual void ShutdownModule() override
 	{
-		if (UObjectInitialized()) 
+		if (UObjectInitialized() && IConsoleManager::Get().FindConsoleVariable(TEXT("UMGEditor.ThumbnailRenderer.Enable"))->GetBool())
 		{
 			UThumbnailManager::Get().UnregisterCustomRenderer(UWidgetBlueprint::StaticClass());
 		}
+
 		MenuExtensibilityManager.Reset();
 		ToolBarExtensibilityManager.Reset();
 
@@ -204,6 +217,25 @@ private:
 		CreatedAssetTypeActions.Add(Action);
 	}
 
+
+	static void ThumbnailRenderingEnabled(IConsoleVariable* Variable)
+	{
+		if (UObjectInitialized())
+		{
+			if (Variable->GetBool())
+			{
+				UThumbnailManager::Get().UnregisterCustomRenderer(UWidgetBlueprint::StaticClass());
+				UThumbnailManager::Get().RegisterCustomRenderer(UWidgetBlueprint::StaticClass(), UWidgetBlueprintThumbnailRenderer::StaticClass());
+			}
+			else
+			{
+				UThumbnailManager::Get().UnregisterCustomRenderer(UWidgetBlueprint::StaticClass());
+
+			}
+		}
+
+	}
+
 private:
 	TSharedPtr<FExtensibilityManager> MenuExtensibilityManager;
 	TSharedPtr<FExtensibilityManager> ToolBarExtensibilityManager;
