Register UMG thumbnail renders after engine loop init to avoid forcing early init of other modules using thumbnail renderer.

diff --git a/Engine/Source/Editor/UMGEditor/Private/UMGEditorModule.cpp b/Engine/Source/Editor/UMGEditor/Private/UMGEditorModule.cpp
--- a/Engine/Source/Editor/UMGEditor/Private/UMGEditorModule.cpp
+++ b/Engine/Source/Editor/UMGEditor/Private/UMGEditorModule.cpp
@@ -43,8 +43,8 @@
 #define LOCTEXT_NAMESPACE "UMG"
 
 const FName UMGEditorAppIdentifier = FName(TEXT("UMGEditorApp"));
-static  TAutoConsoleVariable<bool> CVarThumbnailRenderEnable(
-	TEXT("UMG.Thumbnail.Renderer.Enable"),
+static TAutoConsoleVariable<bool> CVarThumbnailRenderEnable(
+	TEXT("UMG.ThumbnailRenderer.Enable"),
 	true,
 	TEXT("Option to enable/disable thumbnail rendering.")
 );
@@ -55,6 +55,8 @@ public:
 	/** Constructor, set up console commands and variables **/
 	FUMGEditorModule()
 		: Settings(nullptr)
+		, bThumbnailRenderersRegistered(false)
+		, bOnPostEngineInitHandled(false)
 	{
 	}
 
@@ -63,6 +65,9 @@ public:
 	{
 		FModuleManager::LoadModuleChecked<IUMGModule>("UMG");
 
+		// Any attempt to use GEditor right now will fail as it hasn't been initialized yet. Waiting for post engine init resolves that.
+		FCoreDelegates::OnPostEngineInit.AddRaw(this, &FUMGEditorModule::OnPostEngineInit);
+
 		if (GIsEditor)
 		{
 			FDesignerCommands::Register();
@@ -103,18 +108,14 @@ public:
 
 	
 		CVarThumbnailRenderEnable->AsVariable()->SetOnChangedCallback(FConsoleVariableDelegate::CreateStatic(&FUMGEditorModule::ThumbnailRenderingEnabled));
-		
-		if (IConsoleManager::Get().FindConsoleVariable(TEXT("UMG.Thumbnail.Renderer.Enable"))->GetBool())
-		{
-			UThumbnailManager::Get().RegisterCustomRenderer(UWidgetBlueprint::StaticClass(), UWidgetBlueprintThumbnailRenderer::StaticClass());
-
-		}
 	}
 
 	/** Called before the module is unloaded, right before the module object is destroyed. */
 	virtual void ShutdownModule() override
 	{
-		if (UObjectInitialized() && IConsoleManager::Get().FindConsoleVariable(TEXT("UMGEditor.ThumbnailRenderer.Enable"))->GetBool())
+		FCoreDelegates::OnPostEngineInit.RemoveAll(this);
+
+		if (UObjectInitialized() && bThumbnailRenderersRegistered && IConsoleManager::Get().FindConsoleVariable(TEXT("UMGEditor.ThumbnailRenderer.Enable"))->GetBool())
 		{
 			UThumbnailManager::Get().UnregisterCustomRenderer(UWidgetBlueprint::StaticClass());
 		}
@@ -202,7 +203,7 @@ public:
 
 	virtual FString GetReferencerName() const override
 	{
-		return "FUMGEditorModule";
+		return "UMGEditorModule";
 	}
 
 	virtual FWidgetBlueprintCompiler* GetRegisteredCompiler() override
@@ -217,23 +218,33 @@ private:
 		CreatedAssetTypeActions.Add(Action);
 	}
 
+	void OnPostEngineInit()
+	{
+		if (GIsEditor)
+		{
+			if (IConsoleManager::Get().FindConsoleVariable(TEXT("UMG.ThumbnailRenderer.Enable"))->GetBool())
+			{
+				UThumbnailManager::Get().RegisterCustomRenderer(UWidgetBlueprint::StaticClass(), UWidgetBlueprintThumbnailRenderer::StaticClass());
+				bThumbnailRenderersRegistered = true;
+			}
+		}
+		bOnPostEngineInitHandled = true;
+	}
 
 	static void ThumbnailRenderingEnabled(IConsoleVariable* Variable)
 	{
-		if (UObjectInitialized())
+		FUMGEditorModule* UMGEditorModule = FModuleManager::GetModulePtr<FUMGEditorModule>(TEXT("UMGEditor"));
+		if (UObjectInitialized() && UMGEditorModule && UMGEditorModule->bOnPostEngineInitHandled)
 		{
 			if (Variable->GetBool())
 			{
-				UThumbnailManager::Get().UnregisterCustomRenderer(UWidgetBlueprint::StaticClass());
 				UThumbnailManager::Get().RegisterCustomRenderer(UWidgetBlueprint::StaticClass(), UWidgetBlueprintThumbnailRenderer::StaticClass());
 			}
 			else
 			{
 				UThumbnailManager::Get().UnregisterCustomRenderer(UWidgetBlueprint::StaticClass());
-
 			}
 		}
-
 	}
 
 private:
@@ -253,6 +264,9 @@ private:
 
 	/** Compiler customization for Widgets */
 	FWidgetBlueprintCompiler WidgetBlueprintCompiler;
+
+	bool bThumbnailRenderersRegistered;
+	bool bOnPostEngineInitHandled;
 };
 
 IMPLEMENT_MODULE(FUMGEditorModule, UMGEditor);
