Fixed aliasing and custom size issue for Widget Thumbnail Renderer #RB patrick.boutot #jira ue-119520

diff --git a/Engine/Source/Editor/UMGEditor/Private/Designer/SDesignerView.cpp b/Engine/Source/Editor/UMGEditor/Private/Designer/SDesignerView.cpp
--- a/Engine/Source/Editor/UMGEditor/Private/Designer/SDesignerView.cpp
+++ b/Engine/Source/Editor/UMGEditor/Private/Designer/SDesignerView.cpp
@@ -1065,32 +1065,32 @@ ETransformMode::Type SDesignerView::GetTransformMode() const
 
 FOptionalSize SDesignerView::GetPreviewAreaWidth() const
 {
-	FVector2D Area, Size;
-	GetPreviewAreaAndSize(Area, Size);
+	TTuple<FVector2D, FVector2D> AreaAndSize = FWidgetBlueprintEditorUtils::GetWidgetPreviewAreaAndSize(GetDefaultWidget(), CachedPreviewDesiredSize, FVector2D(PreviewWidth, PreviewHeight));
+	FVector2D Area = AreaAndSize.Get<0>();
 
 	return Area.X;
 }
 
 FOptionalSize SDesignerView::GetPreviewAreaHeight() const
 {
-	FVector2D Area, Size;
-	GetPreviewAreaAndSize(Area, Size);
+	TTuple<FVector2D, FVector2D> AreaAndSize = FWidgetBlueprintEditorUtils::GetWidgetPreviewAreaAndSize(GetDefaultWidget(), CachedPreviewDesiredSize, FVector2D(PreviewWidth, PreviewHeight));
+	FVector2D Area = AreaAndSize.Get<0>();;
 
 	return Area.Y;
 }
 
 FOptionalSize SDesignerView::GetPreviewSizeWidth() const
 {
-	FVector2D Area, Size;
-	GetPreviewAreaAndSize(Area, Size);
+	TTuple<FVector2D, FVector2D> AreaAndSize = FWidgetBlueprintEditorUtils::GetWidgetPreviewAreaAndSize(GetDefaultWidget(), CachedPreviewDesiredSize, FVector2D(PreviewWidth, PreviewHeight));
+	FVector2D Size = AreaAndSize.Get<1>();
 
 	return Size.X;
 }
 
 FOptionalSize SDesignerView::GetPreviewSizeHeight() const
 {
-	FVector2D Area, Size;
-	GetPreviewAreaAndSize(Area, Size);
+	TTuple<FVector2D, FVector2D> AreaAndSize = FWidgetBlueprintEditorUtils::GetWidgetPreviewAreaAndSize(GetDefaultWidget(), CachedPreviewDesiredSize, FVector2D(PreviewWidth, PreviewHeight));
+	FVector2D Size = AreaAndSize.Get<1>();
 
 	return Size.Y;
 }
@@ -1191,66 +1191,9 @@ const FSlateBrush* SDesignerView::GetPreviewBackground() const
 	return nullptr;
 }
 
-void SDesignerView::GetPreviewAreaAndSize(FVector2D& Area, FVector2D& Size) const
-{
-	Area = FVector2D(PreviewWidth, PreviewHeight);
-	Size = FVector2D(PreviewWidth, PreviewHeight);
-
-	if ( UUserWidget* DefaultWidget = GetDefaultWidget() )
-	{
-		switch ( DefaultWidget->DesignSizeMode )
-		{
-		case EDesignPreviewSizeMode::Custom:
-			Area = DefaultWidget->DesignTimeSize;
-			// If the custom size is 0 in some dimension, use the desired size instead.
-			if (Area.X == 0)
-			{
-				Area.X = CachedPreviewDesiredSize.X;
-			}
-			if (Area.Y == 0)
-			{
-				Area.Y = CachedPreviewDesiredSize.Y;
-			}
-			Size = Area;
-			break;
-		case EDesignPreviewSizeMode::CustomOnScreen:
-			Size = DefaultWidget->DesignTimeSize;
-
-			// If the custom size is 0 in some dimension, use the desired size instead.
-			if (Size.X == 0)
-			{
-				Size.X = CachedPreviewDesiredSize.X;
-			}
-			if (Size.Y == 0)
-			{
-				Size.Y = CachedPreviewDesiredSize.Y;
-			}
-			return;
-		case EDesignPreviewSizeMode::Desired:
-			Area = CachedPreviewDesiredSize;
-			// Fall through to DesiredOnScreen
-		case EDesignPreviewSizeMode::DesiredOnScreen:
-			Size = CachedPreviewDesiredSize;
-			return;
-		case EDesignPreviewSizeMode::FillScreen:
-			break;
-		}
-	}
-}
-
 float SDesignerView::GetPreviewDPIScale() const
 {
-	// If the user is using a custom size then we disable the DPI scaling logic.
-	if ( UUserWidget* DefaultWidget = GetDefaultWidget() )
-	{
-		if ( DefaultWidget->DesignSizeMode == EDesignPreviewSizeMode::Custom || 
-			 DefaultWidget->DesignSizeMode == EDesignPreviewSizeMode::Desired )
-		{
-			return 1.0f;
-		}
-	}
-
-	return GetDefault<UUserInterfaceSettings>(UUserInterfaceSettings::StaticClass())->GetDPIScaleBasedOnSize(FIntPoint(PreviewWidth, PreviewHeight));
+	return FWidgetBlueprintEditorUtils::GetWidgetPreviewDPIScale(GetDefaultWidget(), FVector2D(PreviewWidth,PreviewHeight));
 }
 
 FSlateRect SDesignerView::ComputeAreaBounds() const
diff --git a/Engine/Source/Editor/UMGEditor/Private/Designer/SDesignerView.h b/Engine/Source/Editor/UMGEditor/Private/Designer/SDesignerView.h
--- a/Engine/Source/Editor/UMGEditor/Private/Designer/SDesignerView.h
+++ b/Engine/Source/Editor/UMGEditor/Private/Designer/SDesignerView.h
@@ -142,8 +142,6 @@ private:
 
 	const FSlateBrush* GetPreviewBackground() const;
 
-	void GetPreviewAreaAndSize(FVector2D& Area, FVector2D& Size) const;
-
 	/** Gets the DPI scale that would be applied given the current preview width and height */
 	float GetPreviewDPIScale() const;
 
diff --git a/Engine/Source/Editor/UMGEditor/Private/WidgetBlueprintEditorUtils.cpp b/Engine/Source/Editor/UMGEditor/Private/WidgetBlueprintEditorUtils.cpp
--- a/Engine/Source/Editor/UMGEditor/Private/WidgetBlueprintEditorUtils.cpp
+++ b/Engine/Source/Editor/UMGEditor/Private/WidgetBlueprintEditorUtils.cpp
@@ -40,6 +40,7 @@
 #include "Components/NamedSlotInterface.h"
 #include "K2Node_Variable.h"
 
+#include "Engine/UserInterfaceSettings.h"
 #include "Input/HittestGrid.h"
 #include "Interfaces/ISlateRHIRendererModule.h"
 #include "Interfaces/ISlate3DRenderer.h"
@@ -1929,4 +1930,94 @@ int32 FWidgetBlueprintEditorUtils::UpdateHittestGrid(FHittestGrid& HitTestGrid,
 	return MaxLayerId;
 }
 
+TTuple<FVector2D, FVector2D> FWidgetBlueprintEditorUtils::GetWidgetPreviewAreaAndSize(UUserWidget* UserWidget, FVector2D DesiredSize, FVector2D PreviewSize)
+{
+	FVector2D Size(PreviewSize.X, PreviewSize.Y);
+	FVector2D Area(PreviewSize.X, PreviewSize.Y);
+
+	if (UserWidget)
+	{
+		switch (UserWidget->DesignSizeMode)
+		{
+		case EDesignPreviewSizeMode::Custom:
+			Area = UserWidget->DesignTimeSize;
+			// If the custom size is 0 in some dimension, use the desired size instead.
+			if (Area.X == 0)
+			{
+				Area.X = DesiredSize.X;
+			}
+			if (Area.Y == 0)
+			{
+				Area.Y = DesiredSize.Y;
+			}
+			Size = Area;
+			break;
+		case EDesignPreviewSizeMode::CustomOnScreen:
+			Size = UserWidget->DesignTimeSize;
+
+			// If the custom size is 0 in some dimension, use the desired size instead.
+			if (Size.X == 0)
+			{
+				Size.X = DesiredSize.X;
+			}
+			if (Size.Y == 0)
+			{
+				Size.Y = DesiredSize.Y;
+			}
+			return TTuple<FVector2D, FVector2D>(Area, Size);
+		case EDesignPreviewSizeMode::Desired:
+			Area = DesiredSize;
+			// Fall through to DesiredOnScreen
+		case EDesignPreviewSizeMode::DesiredOnScreen:
+			Size = DesiredSize;
+			return TTuple<FVector2D, FVector2D>(Area, Size);
+		case EDesignPreviewSizeMode::FillScreen:
+			break;
+		}
+	}
+	return TTuple<FVector2D, FVector2D>(Area, Size);
+}
+
+float FWidgetBlueprintEditorUtils::GetWidgetPreviewDPIScale(UUserWidget* UserWidget, FVector2D PreviewSize)
+{
+	// If the user is using a custom size then we disable the DPI scaling logic.
+	if (UserWidget)
+	{
+		if (UserWidget->DesignSizeMode == EDesignPreviewSizeMode::Custom ||
+			UserWidget->DesignSizeMode == EDesignPreviewSizeMode::Desired)
+		{
+			return 1.0f;
+		}
+	}
+
+	return GetDefault<UUserInterfaceSettings>(UUserInterfaceSettings::StaticClass())->GetDPIScaleBasedOnSize(FIntPoint(PreviewSize.X, PreviewSize.Y));
+}
+
+
+FVector2D FWidgetBlueprintEditorUtils::GetWidgetPreviewUnScaledCustomSize(FVector2D DesiredSize, UUserWidget* UserWidget)
+{
+
+	checkf(DesiredSize.X > 0.f && DesiredSize.Y > 0.f, TEXT("The size should have been previously checked to be > 0."));
+
+	FVector2D FinalSize(0.f,0.f);
+	int32 PreviewWidth;
+	const TCHAR* ConfigSectionName = TEXT("UMGEditor.Designer");
+	GConfig->GetInt(ConfigSectionName, TEXT("PreviewWidth"), PreviewWidth, GEditorPerProjectIni);
+	int32 PreviewHeight;
+	GConfig->GetInt(ConfigSectionName, TEXT("PreviewHeight"), PreviewHeight, GEditorPerProjectIni);
+
+	FVector2D PreviewSize(PreviewWidth, PreviewHeight);
+
+	TTuple<FVector2D, FVector2D> AreaAndSize = GetWidgetPreviewAreaAndSize(UserWidget, DesiredSize, PreviewSize);
+
+	float DPIScale = GetWidgetPreviewDPIScale(UserWidget, PreviewSize);
+
+	if (ensure(DPIScale > 0.f))
+	{
+		FinalSize = AreaAndSize.Get<1>() / DPIScale;
+	}
+
+	return FinalSize;
+
+}
 #undef LOCTEXT_NAMESPACE
diff --git a/Engine/Source/Editor/UMGEditor/Private/WidgetBlueprintEditorUtils.h b/Engine/Source/Editor/UMGEditor/Private/WidgetBlueprintEditorUtils.h
--- a/Engine/Source/Editor/UMGEditor/Private/WidgetBlueprintEditorUtils.h
+++ b/Engine/Source/Editor/UMGEditor/Private/WidgetBlueprintEditorUtils.h
@@ -65,6 +65,12 @@ public:
 
 	static int32 UpdateHittestGrid(FHittestGrid& HitTestGrid, TSharedRef<SWindow> Window, float Scale, FVector2D DrawSize, float DeltaTime);
 
+	static TTuple<FVector2D, FVector2D> GetWidgetPreviewAreaAndSize(UUserWidget* UserWidget, FVector2D DesiredSize, FVector2D PreviewSize);
+
+	static float GetWidgetPreviewDPIScale(UUserWidget* UserWidget, FVector2D PreviewSize);
+
+	static FVector2D GetWidgetPreviewUnScaledCustomSize(FVector2D DesiredSize, UUserWidget* UserWidget);
+
 private:
 
 	static FString CopyWidgetsInternal(UWidgetBlueprint* BP, TSet<FWidgetReference> Widgets);
diff --git a/Engine/Source/Editor/UMGEditor/Private/WidgetBlueprintThumbnailRenderer.cpp b/Engine/Source/Editor/UMGEditor/Private/WidgetBlueprintThumbnailRenderer.cpp
--- a/Engine/Source/Editor/UMGEditor/Private/WidgetBlueprintThumbnailRenderer.cpp
+++ b/Engine/Source/Editor/UMGEditor/Private/WidgetBlueprintThumbnailRenderer.cpp
@@ -10,6 +10,7 @@
 #include "Slate/WidgetRenderer.h"
 #include "WidgetBlueprint.h"
 #include "Widgets/SVirtualWindow.h"
+#include "WidgetBlueprintEditorUtils.h"
 
 #define LOCTEXT_NAMESPACE "UWidgetBlueprintThumbnailRenderer"
 
@@ -42,49 +43,62 @@ void UWidgetBlueprintThumbnailRenderer::Draw(UObject* Object, int32 X, int32 Y,
 
 	TSharedPtr<SWidget> WindowContent;
 	UClass* ClassToGenerate = WidgetBlueprintToRender->GeneratedClass;
+	UUserWidget* WidgetInstance = nullptr;
 	if (bIsBlueprintValid && ClassToGenerate->IsChildOf(UWidget::StaticClass()))
 	{
-		if (UUserWidget* WidgetInstance = NewObject<UUserWidget>(GetTransientPackage(), ClassToGenerate))
+		WidgetInstance = NewObject<UUserWidget>(GetTransientPackage(), ClassToGenerate);
+		if (WidgetInstance)
 		{
 			WidgetInstance->Initialize();
 			WidgetInstance->SetDesignerFlags(EWidgetDesignFlags::Designing | EWidgetDesignFlags::ExecutePreConstruct);
 
+
 			WindowContent = WidgetInstance->TakeWidget();
 		}
 	}
 
-	if (WindowContent == nullptr)
+	if (WindowContent == nullptr || WidgetInstance == nullptr)
 	{
 		return;
 	}
 
+
 	//Create Window
-	FVector2D WidgetSize(Width, Height);
-	TSharedRef<SVirtualWindow> Window = SNew(SVirtualWindow).Size(WidgetSize);
+	FVector2D ThumbnailSize(Width, Height);
+	TSharedRef<SVirtualWindow> Window = SNew(SVirtualWindow).Size(ThumbnailSize);
 	TUniquePtr<FHittestGrid> HitTestGrid = MakeUnique<FHittestGrid>();
 	Window->SetContent(WindowContent.ToSharedRef());
-	Window->Resize(WidgetSize);
+	Window->Resize(ThumbnailSize);
 
 	// Store the desired size to maintain the aspect ratio later
-	FGeometry WindowGeometry = FGeometry::MakeRoot(WidgetSize, FSlateLayoutTransform(1.0f));
+	FGeometry WindowGeometry = FGeometry::MakeRoot(ThumbnailSize, FSlateLayoutTransform(1.0f));
 	Window->SlatePrepass(1.0f);
 	FVector2D DesiredSizeWindow = Window->GetDesiredSize();
 
-	
+
 	if (DesiredSizeWindow.X < SMALL_NUMBER || DesiredSizeWindow.Y < SMALL_NUMBER)
 	{
 		return;
 	}
 
+	FVector2D UnscaledSize = FWidgetBlueprintEditorUtils::GetWidgetPreviewUnScaledCustomSize(DesiredSizeWindow, WidgetInstance);
+	if (UnscaledSize.X < SMALL_NUMBER || UnscaledSize.Y < SMALL_NUMBER)
+	{
+		return;
+	}
+	TTuple<float, FVector2D> ScaleAndOffset = GetScaleAndOffset(UnscaledSize, ThumbnailSize);
+	float Scale = ScaleAndOffset.Get<0>();
+	FVector2D Offset = ScaleAndOffset.Get<1>();
+	FVector2D ScaledSize = UnscaledSize * Scale;
+
 	// Create Renderer Target and WidgetRenderer
 	bool bApplyGammaCorrection = false;
 	FWidgetRenderer* WidgetRenderer = new FWidgetRenderer(bApplyGammaCorrection);
-	WidgetRenderer->SetIsPrepassNeeded(false);
 	const bool bIsLinearSpace = !bApplyGammaCorrection;
 
 	if (!RenderTarget2D)
 	{
-		RenderTarget2D  = NewObject<UTextureRenderTarget2D>();
+		RenderTarget2D = NewObject<UTextureRenderTarget2D>();
 		RenderTarget2D->Filter = TF_Bilinear;
 		RenderTarget2D->ClearColor = FLinearColor::Transparent;
 		RenderTarget2D->SRGB = true;
@@ -92,17 +106,11 @@ void UWidgetBlueprintThumbnailRenderer::Draw(UObject* Object, int32 X, int32 Y,
 	}
 
 	const EPixelFormat RequestedFormat = FSlateApplication::Get().GetRenderer()->GetSlateRecommendedColorFormat();
-	RenderTarget2D->InitCustomFormat(DesiredSizeWindow.X, DesiredSizeWindow.Y, RequestedFormat, true);
-
-	WidgetRenderer->DrawWindow(RenderTarget2D, *HitTestGrid, Window, 1.0f, DesiredSizeWindow, 0.1f);
+	RenderTarget2D->InitCustomFormat(ScaledSize.X, ScaledSize.Y, RequestedFormat, true);
+	WidgetRenderer->DrawWindow(RenderTarget2D, *HitTestGrid, Window, Scale, ScaledSize, 0.1f);
 
 	// Draw the SWidget on canvas
-
-	TTuple<FVector2D, FVector2D> ScaledSizeAndOffset = GetScaledSizeAndOffset(DesiredSizeWindow.X, DesiredSizeWindow.Y, Width, Height);
-	FVector2D DrawOffset = ScaledSizeAndOffset.Get<1>();
-	WidgetSize = ScaledSizeAndOffset.Get<0>();
-
-	FCanvasTileItem CanvasTile(FVector2D(X + DrawOffset.X, Y + DrawOffset.Y), RenderTarget2D->GetResource(), WidgetSize, FLinearColor::White);
+	FCanvasTileItem CanvasTile(FVector2D(X + Offset.X, Y + Offset.Y), RenderTarget2D->GetResource(), ScaledSize, FLinearColor::White);
 	CanvasTile.BlendMode = SE_BLEND_Translucent;
 	FlushRenderingCommands();
 	CanvasTile.Draw(Canvas);
@@ -112,34 +120,34 @@ void UWidgetBlueprintThumbnailRenderer::Draw(UObject* Object, int32 X, int32 Y,
 		BeginCleanup(WidgetRenderer);
 		WidgetRenderer = nullptr;
 	}
+	
 
 }
 
-TTuple<FVector2D, FVector2D> UWidgetBlueprintThumbnailRenderer::GetScaledSizeAndOffset(float ImgWidth, float ImgHeight, float ThumbnailWidth, float ThumbnailHeight) const
+TTuple<float, FVector2D> UWidgetBlueprintThumbnailRenderer::GetScaleAndOffset(FVector2D WidgetSize, FVector2D ThumbnailSize) const
 {
 	// Scale the widget blueprint image to fit in the thumbnail
 
-	checkf(ImgWidth > 0.f && ImgHeight > 0.f, TEXT("The size should have been previously checked to be > 0."));
+	checkf(WidgetSize.X > 0.f && WidgetSize.Y > 0.f, TEXT("The size should have been previously checked to be > 0."));
 
-	FVector2D ScaledSize(0.0f,0.0f);
-	FVector2D Offset(0.0f, 0.0f);
-	if (ImgWidth > ImgHeight)
+	float Scale;
+	float XOffset = 0;
+	float YOffset = 0;
+	if (WidgetSize.X > WidgetSize.Y)
 	{
-		float RatioAdjust = ImgHeight / ImgWidth;
-		Offset.Y = ThumbnailHeight * (1 - RatioAdjust) / 2;
-		ScaledSize.Y = ThumbnailHeight * RatioAdjust;
-		ScaledSize.X = ThumbnailWidth;
+		Scale = ThumbnailSize.X / WidgetSize.X;
+		WidgetSize *= Scale;
+		YOffset = (ThumbnailSize.Y - WidgetSize.Y) / 2.f;
 	}
 	else
 	{
-		float RatioAdjust = ImgWidth / ImgHeight;
-		Offset.X = ThumbnailWidth * (1 - RatioAdjust) / 2;
-		ScaledSize.X = ThumbnailWidth * RatioAdjust;
-		ScaledSize.Y = ThumbnailHeight;
+		Scale = ThumbnailSize.Y / WidgetSize.Y;
+		WidgetSize *= Scale;
+		XOffset = (ThumbnailSize.X - WidgetSize.X) / 2.f;
 	}
-	
-	
-	return TTuple<FVector2D, FVector2D>(ScaledSize, Offset);
+
+	return TTuple<float, FVector2D>(Scale, FVector2D(XOffset, YOffset));
 }
 
+
 #undef LOCTEXT_NAMESPACE
\ No newline at end of file
diff --git a/Engine/Source/Editor/UMGEditor/Public/WidgetBlueprintThumbnailRenderer.h b/Engine/Source/Editor/UMGEditor/Public/WidgetBlueprintThumbnailRenderer.h
--- a/Engine/Source/Editor/UMGEditor/Public/WidgetBlueprintThumbnailRenderer.h
+++ b/Engine/Source/Editor/UMGEditor/Public/WidgetBlueprintThumbnailRenderer.h
@@ -14,13 +14,13 @@ UCLASS()
 class UMGEDITOR_API UWidgetBlueprintThumbnailRenderer : public UDefaultSizedThumbnailRenderer
 {
 	GENERATED_BODY()
-		
+
 	//~ Begin UThumbnailRenderer Object
 	bool CanVisualizeAsset(UObject* Object) override;
 	void Draw(UObject* Object, int32 X, int32 Y, uint32 Width, uint32 Height, FRenderTarget* RenderTarget, FCanvas* Canvas, bool bAdditionalViewFamily) override;
 
 private:
-	TTuple<FVector2D, FVector2D> GetScaledSizeAndOffset(float ImgWidth, float ImgHeight, float ThumbnailWidth, float ThumbnailHeight) const;
+	TTuple<float, FVector2D> GetScaleAndOffset(FVector2D DesiredSize, FVector2D WidgetSize) const;
 
 	UPROPERTY(Transient)
 	TObjectPtr<UTextureRenderTarget2D> RenderTarget2D;
