PR #8168: Allow Selection of base class for Widget Blueprint (Contributed by [REDACTED_PERSON])

diff --git a/Engine/Source/Editor/ClassViewer/Private/SClassViewer.cpp b/Engine/Source/Editor/ClassViewer/Private/SClassViewer.cpp
--- a/Engine/Source/Editor/ClassViewer/Private/SClassViewer.cpp
+++ b/Engine/Source/Editor/ClassViewer/Private/SClassViewer.cpp
@@ -2644,6 +2644,11 @@ void SClassViewer::Tick( const FGeometry& AllottedGeometry, const double InCurre
 			ExpandRootNodes();
 		}
 
+		if (InitOptions.bExpandAllNodes)
+		{
+			SetAllExpansionStates(true);
+		}
+
 		// Scroll the first item into view if applicable
 		const TArray<TSharedPtr<FClassViewerNode>> SelectedItems = GetSelectedItems();
 		if (SelectedItems.Num() > 0)
diff --git a/Engine/Source/Editor/ClassViewer/Public/ClassViewerModule.h b/Engine/Source/Editor/ClassViewer/Public/ClassViewerModule.h
--- a/Engine/Source/Editor/ClassViewer/Public/ClassViewerModule.h
+++ b/Engine/Source/Editor/ClassViewer/Public/ClassViewerModule.h
@@ -96,6 +96,9 @@ public:
 	/** If true, root nodes will be expanded by default. */
 	bool bExpandRootNodes;
 
+	/** If true, all nodes will be expanded by default. */
+	bool bExpandAllNodes;
+
 	/** true allows class dynamic loading on selection */
 	bool bEnableClassDynamicLoading;
 
@@ -142,6 +145,7 @@ public:
 		, bShowNoneOption(false)
 		, bShowObjectRootClass(false)
 		, bExpandRootNodes(true)
+		, bExpandAllNodes(false)
 		, bEnableClassDynamicLoading(true)
 		, NameTypeToDisplay(EClassViewerNameTypeToDisplay::ClassName)
 		, bAllowViewOptions(true)
diff --git a/Engine/Source/Editor/UMGEditor/Private/WidgetBlueprintFactory.cpp b/Engine/Source/Editor/UMGEditor/Private/WidgetBlueprintFactory.cpp
--- a/Engine/Source/Editor/UMGEditor/Private/WidgetBlueprintFactory.cpp
+++ b/Engine/Source/Editor/UMGEditor/Private/WidgetBlueprintFactory.cpp
@@ -59,6 +59,29 @@ UWidgetBlueprintFactory::UWidgetBlueprintFactory(const FObjectInitializer& Objec
 
 bool UWidgetBlueprintFactory::ConfigureProperties()
 {
+	{
+		FClassViewerModule& ClassViewerModule = FModuleManager::LoadModuleChecked<FClassViewerModule>("ClassViewer");
+
+		// Fill in options
+		FClassViewerInitializationOptions Options;
+		Options.DisplayMode = EClassViewerDisplayMode::Type::TreeView;
+		Options.Mode = EClassViewerMode::ClassPicker;
+		Options.bShowNoneOption = false;
+		Options.bExpandAllNodes = true;
+
+		TSharedPtr<FWidgetClassFilter> Filter = MakeShareable(new FWidgetClassFilter);
+		Options.ClassFilters.Add(Filter.ToSharedRef());
+		Options.ExtraPickerCommonClasses.Add(UUserWidget::StaticClass());
+
+		Filter->DisallowedClassFlags = CLASS_Abstract | CLASS_Deprecated | CLASS_NewerVersionExists;
+		Filter->AllowedChildrenOfClasses.Add(UUserWidget::StaticClass());
+
+		const FText TitleText = LOCTEXT("CreateWidgetBlueprint", "Pick Root Widget for New Widget Blueprint");
+
+		UClass* ChosenParentClass = nullptr;
+		SClassPickerDialog::PickClass(TitleText, Options, ChosenParentClass, UUserWidget::StaticClass());
+		ParentClass = ChosenParentClass ? ChosenParentClass : UUserWidget::StaticClass();
+	}
 	if (GetDefault<UUMGEditorProjectSettings>()->bUseWidgetTemplateSelector)
 	{
 		// Load the classviewer module to display a class picker
