UMG: Add window that shows the widget that needs to be bound to the UserWidget and their compile state.

diff --git a/Engine/Source/Editor/UMGEditor/Private/BlueprintModes/WidgetDesignerApplicationMode.cpp b/Engine/Source/Editor/UMGEditor/Private/BlueprintModes/WidgetDesignerApplicationMode.cpp
--- a/Engine/Source/Editor/UMGEditor/Private/BlueprintModes/WidgetDesignerApplicationMode.cpp
+++ b/Engine/Source/Editor/UMGEditor/Private/BlueprintModes/WidgetDesignerApplicationMode.cpp
@@ -12,6 +12,7 @@
 
 #include "TabFactory/PaletteTabSummoner.h"
 #include "TabFactory/HierarchyTabSummoner.h"
+#include "TabFactory/BindWidgetTabSummoner.h"
 #include "TabFactory/DesignerTabSummoner.h"
 #include "TabFactory/SequencerTabSummoner.h"
 #include "TabFactory/DetailsTabSummoner.h"
@@ -53,9 +54,21 @@ FWidgetDesignerApplicationMode::FWidgetDesignerApplicationMode(TSharedPtr<FWidge
 				)
 				->Split
 				(
-					FTabManager::NewStack()
-					->SetSizeCoefficient( 0.5f )
-					->AddTab( FHierarchyTabSummoner::TabID, ETabState::OpenedTab )
+					FTabManager::NewSplitter()
+					->SetSizeCoefficient(0.15f)
+					->SetOrientation(Orient_Horizontal)
+					->Split
+					(
+						FTabManager::NewStack()
+						->SetSizeCoefficient(0.5f)
+						->AddTab(FHierarchyTabSummoner::TabID, ETabState::OpenedTab)
+					)
+					->Split
+					(
+						FTabManager::NewStack()
+						->SetSizeCoefficient(0.5f)
+						->AddTab(FBindWidgetTabSummoner::TabID, ETabState::OpenedTab)
+					)
 				)
 			)
 			->Split
@@ -104,6 +117,7 @@ FWidgetDesignerApplicationMode::FWidgetDesignerApplicationMode(TSharedPtr<FWidge
 	TabFactories.RegisterFactory(MakeShareable(new FDetailsTabSummoner(InWidgetEditor)));
 	TabFactories.RegisterFactory(MakeShareable(new FDesignerTabSummoner(InWidgetEditor)));
 	TabFactories.RegisterFactory(MakeShareable(new FHierarchyTabSummoner(InWidgetEditor)));
+	TabFactories.RegisterFactory(MakeShareable(new FBindWidgetTabSummoner(InWidgetEditor)));
 	TabFactories.RegisterFactory(MakeShareable(new FPaletteTabSummoner(InWidgetEditor)));
 	TabFactories.RegisterFactory(MakeShareable(new FSequencerTabSummoner(InWidgetEditor)));
 	TabFactories.RegisterFactory(MakeShareable(new FAnimationTabSummoner(InWidgetEditor)));
diff --git a/Engine/Source/Editor/UMGEditor/Private/TabFactory/BindWidgetTabSummoner.cpp b/Engine/Source/Editor/UMGEditor/Private/TabFactory/BindWidgetTabSummoner.cpp
new file mode 100644
--- /dev/null
+++ b/Engine/Source/Editor/UMGEditor/Private/TabFactory/BindWidgetTabSummoner.cpp
@@ -0,0 +1,32 @@
+// Copyright Epic Games, Inc. All Rights Reserved.
+
+#include "TabFactory/BindWidgetTabSummoner.h"
+#include "EditorStyleSet.h"
+#include "Widgets/SBindWidgetView.h"
+
+#define LOCTEXT_NAMESPACE "UMG.BindWidget"
+
+const FName FBindWidgetTabSummoner::TabID(TEXT("BindWidget"));
+
+FBindWidgetTabSummoner::FBindWidgetTabSummoner(TSharedPtr<class FWidgetBlueprintEditor> InBlueprintEditor)
+		: FWorkflowTabFactory(TabID, InBlueprintEditor)
+		, BlueprintEditor(InBlueprintEditor)
+{
+	TabLabel = LOCTEXT("SlateBindWidgetTabLabel", "Bind Widgets");
+	TabIcon = FSlateIcon(FEditorStyle::GetStyleSetName(), "Kismet.Tabs.Palette");
+
+	bIsSingleton = true;
+
+	ViewMenuDescription = LOCTEXT("SlateHierarchy_ViewMenu_Desc", "Bind Widgets");
+	ViewMenuTooltip = LOCTEXT("SlateHierarchy_ViewMenu_ToolTip", "Show the list of widgets that needs to be binded.");
+}
+
+TSharedRef<SWidget> FBindWidgetTabSummoner::CreateTabBody(const FWorkflowTabSpawnInfo& Info) const
+{
+	TSharedPtr<FWidgetBlueprintEditor> BlueprintEditorPtr = StaticCastSharedPtr<FWidgetBlueprintEditor>(BlueprintEditor.Pin());
+
+	return SNew(SBindWidgetView, BlueprintEditorPtr)
+		.AddMetaData<FTagMetaData>(FTagMetaData(TEXT("BindWidget")));
+}
+
+#undef LOCTEXT_NAMESPACE 
diff --git a/Engine/Source/Editor/UMGEditor/Private/TabFactory/BindWidgetTabSummoner.h b/Engine/Source/Editor/UMGEditor/Private/TabFactory/BindWidgetTabSummoner.h
new file mode 100644
--- /dev/null
+++ b/Engine/Source/Editor/UMGEditor/Private/TabFactory/BindWidgetTabSummoner.h
@@ -0,0 +1,24 @@
+// Copyright Epic Games, Inc. All Rights Reserved.
+
+#pragma once
+
+#include "CoreMinimal.h"
+#include "Widgets/SWidget.h"
+#include "WorkflowOrientedApp/WorkflowTabFactory.h"
+#include "WidgetBlueprintEditor.h"
+
+class FWidgetBlueprintEditor;
+
+struct FBindWidgetTabSummoner : public FWorkflowTabFactory
+{
+public:
+	static const FName TabID;
+	
+public:
+	FBindWidgetTabSummoner(TSharedPtr<FWidgetBlueprintEditor> InBlueprintEditor);
+	
+	virtual TSharedRef<SWidget> CreateTabBody(const FWorkflowTabSpawnInfo& Info) const override;
+	
+protected:
+	TWeakPtr<FWidgetBlueprintEditor> BlueprintEditor;
+};
diff --git a/Engine/Source/Editor/UMGEditor/Private/Widgets/SBindWidgetView.cpp b/Engine/Source/Editor/UMGEditor/Private/Widgets/SBindWidgetView.cpp
new file mode 100644
--- /dev/null
+++ b/Engine/Source/Editor/UMGEditor/Private/Widgets/SBindWidgetView.cpp
@@ -0,0 +1,430 @@
+// Copyright Epic Games, Inc. All Rights Reserved.
+
+#include "Widgets/SBindWidgetView.h"
+#include "WidgetBlueprint.h"
+#include "WidgetBlueprintEditor.h"
+#include "WidgetBlueprintEditorUtils.h"
+
+#include "IDocumentation.h"
+#include "Editor.h"
+#include "EditorStyleSet.h"
+#include "Styling/SlateIconFinder.h"
+#include "Widgets/SNullWidget.h"
+#include "Widgets/Text/STextBlock.h"
+#include "Widgets/Views/STableRow.h"
+#include "Widgets/Views/STreeView.h"
+
+
+#define LOCTEXT_NAMESPACE "UMG.BindWidget"
+
+namespace UE
+{
+namespace UMG
+{
+static FName Column_PropertyName = "PropertyName";
+static FName Column_WidgetType = "WidgetType";
+static FName Column_Optional = "Optional";
+
+
+/**
+ * 
+ */
+struct FBindWidgetListEntry
+{
+	enum class EEntryType
+	{
+		Category,
+		Binding,
+	};
+	enum class EBindingType
+	{
+		Widget,		//UWidget
+		Animation,	//UWidgetAnimation
+	};
+
+	EEntryType EntryType = EEntryType::Category;
+	EBindingType BindingType = EBindingType::Widget;
+
+	FObjectProperty* Property = nullptr;
+	bool bIsOptional = false;
+	bool bIsBound = false;
+	bool bIsCorrectClass = false;
+
+	FText CategoryName;
+};
+
+typedef TSharedPtr<FBindWidgetListEntry> FBindWidgetListEntryPtr;
+
+/**
+ *
+ */
+class SBindWidgetListRow : public SMultiColumnTableRow<FBindWidgetListEntryPtr>
+{
+public:
+	SLATE_BEGIN_ARGS(SBindWidgetListRow) {}
+		SLATE_ARGUMENT(FBindWidgetListEntryPtr, Entry)
+	SLATE_END_ARGS()
+
+	void Construct(const FArguments& Args, const TSharedRef<STableViewBase>& OwnerTableView)
+	{
+		EntryPtr = Args._Entry;
+
+		SMultiColumnTableRow<FBindWidgetListEntryPtr>::Construct(
+			FSuperRowType::FArguments()
+			.Padding(1.0f),
+			OwnerTableView
+		);
+	}
+
+	virtual TSharedRef<SWidget> GenerateWidgetForColumn(const FName& ColumnName) override
+	{
+		if (EntryPtr->EntryType == FBindWidgetListEntry::EEntryType::Binding)
+		{
+			if (ColumnName == Column_PropertyName)
+			{
+				return SNew(STextBlock)
+					.Text(FText::FromName(EntryPtr->Property->GetFName()))
+					.ToolTipText(this, &SBindWidgetListRow::HandleGetTooltip);
+					
+			}
+			if (ColumnName == Column_WidgetType)
+			{
+				const FSlateBrush* Icon = FSlateIconFinder::FindIconBrushForClass(EntryPtr->Property->PropertyClass);
+				return SNew(SHorizontalBox)
+					.ToolTipText(this, &SBindWidgetListRow::HandleGetTooltip)
+					+ SHorizontalBox::Slot()
+					.AutoWidth()
+					.Padding(4.f, 2.f)
+					.VAlign(EVerticalAlignment::VAlign_Center)
+					[
+						SNew(SImage)
+						.ColorAndOpacity(FSlateColor::UseForeground())
+						.Image(Icon)
+					]
+					+ SHorizontalBox::Slot()
+					.VAlign(EVerticalAlignment::VAlign_Center)
+					[
+						SNew(STextBlock)
+						.Text(EntryPtr->Property->PropertyClass->GetDisplayNameText())
+					];
+			}
+			if (ColumnName == Column_Optional)
+			{
+				const FSlateBrush* Image = nullptr;
+				FText Tooltip = FText::GetEmpty();
+				if (EntryPtr->bIsBound && EntryPtr->bIsCorrectClass)
+				{
+					Image = FEditorStyle::GetBrush("Icons.Success");
+					Tooltip = LOCTEXT("BoundCorrectlyTooltip", "The widget is bound");
+				}
+				else if (EntryPtr->bIsBound) // is not of the correct class
+				{
+					Image = FEditorStyle::GetBrush("Icons.Error");
+					Tooltip = LOCTEXT("BoundWidgetWrongClassTooltip", "The bound widget is not of the correct type.");
+				}
+				else if (EntryPtr->bIsOptional) // is not of the correct class
+				{
+					Image = FEditorStyle::GetBrush("Icons.FilledCircle");
+					Tooltip = LOCTEXT("BoundOptionalTooltip", "The widget is not bound but it is optional.");
+				}
+				else
+				{
+					Image = FEditorStyle::GetBrush("Icons.Error");
+					Tooltip = LOCTEXT("BoundNoWidgetTooltip", "The widget is not bound.");
+				}
+
+				return SNew(SBox)
+					[
+						SNew(SImage)
+						.Image(Image)
+					]
+					.ToolTipText(Tooltip);
+			}
+		}
+
+		return SNullWidget::NullWidget;
+	}
+
+private:
+	FText HandleGetTooltip() const
+	{
+		return EntryPtr->Property->GetToolTipText();
+	}
+
+private:
+	FBindWidgetListEntryPtr EntryPtr;
+};
+
+
+/**
+ *
+ */
+class SBindWidgetView : public STreeView<FBindWidgetListEntryPtr>
+{
+private:
+	using Super = STreeView<FBindWidgetListEntryPtr>;
+
+public:
+	SLATE_BEGIN_ARGS(SBindWidgetView) {}
+	SLATE_END_ARGS()
+
+	void Construct(const FArguments& Args, TArray<FBindWidgetListEntryPtr> InSourceData)
+	{
+		AllSourceData = MoveTemp(InSourceData);
+		{
+			TSharedRef<FBindWidgetListEntry> Ref = MakeShared<FBindWidgetListEntry>();
+			Ref->CategoryName = LOCTEXT("Widget", "Widget");
+			Ref->EntryType = FBindWidgetListEntry::EEntryType::Category;
+			Ref->BindingType = FBindWidgetListEntry::EBindingType::Widget;
+
+			CategorySourceData.Add(Ref);
+		}
+		{
+			TSharedRef<FBindWidgetListEntry> Ref = MakeShared<FBindWidgetListEntry>();
+			Ref->CategoryName = LOCTEXT("Animation", "Animation");
+			Ref->EntryType = FBindWidgetListEntry::EEntryType::Category;
+			Ref->BindingType = FBindWidgetListEntry::EBindingType::Animation;
+
+			CategorySourceData.Add(Ref);
+		}
+
+		Super::Construct(Super::FArguments()
+			.SelectionMode(ESelectionMode::None)
+			.OnGenerateRow(this, &SBindWidgetView::HandleGenerateRow)
+			.OnGetChildren(this, &SBindWidgetView::HandleGetChildren)
+			.TreeItemsSource(&CategorySourceData)
+			.HeaderRow
+			(
+				SNew(SHeaderRow)
+				+ SHeaderRow::Column(Column_WidgetType)
+				.FillWidth(0.5f)
+				.DefaultLabel(LOCTEXT("TypeHeaderName", "Type"))
+				+ SHeaderRow::Column(Column_PropertyName)
+				.FillWidth(0.5f)
+				.DefaultLabel(LOCTEXT("PropertyNameHeaderName", "Property"))
+				+ SHeaderRow::Column(Column_Optional)
+				.FixedWidth(16.f)
+				.DefaultLabel(LOCTEXT("EmptyHeaderName", ""))
+			));
+	}
+
+	void SetSourceData(TArray<FBindWidgetListEntryPtr> InSourceData)
+	{
+		AllSourceData = MoveTemp(InSourceData);
+		RebuildList();
+	}
+
+private:
+	TSharedRef<ITableRow> HandleGenerateRow(FBindWidgetListEntryPtr Entry, const TSharedRef<STableViewBase>& OwnerTable) const
+	{
+		if (Entry->EntryType == FBindWidgetListEntry::EEntryType::Category)
+		{
+			return SNew(STableRow<FBindWidgetListEntryPtr>, OwnerTable)
+				.Style(FEditorStyle::Get(), "UMGEditor.PaletteHeader")
+				.Padding(5.0f)
+				.ShowSelection(false)
+				[
+					SNew(STextBlock)
+					.TransformPolicy(ETextTransformPolicy::ToUpper)
+					.Text(Entry->CategoryName)
+					.Font(FAppStyle::Get().GetFontStyle("SmallFontBold"))
+				];
+		}
+		else
+		{
+			return SNew(SBindWidgetListRow, OwnerTable)
+				.Entry(Entry);
+		}
+	}
+
+	void HandleGetChildren(FBindWidgetListEntryPtr Parent, TArray<FBindWidgetListEntryPtr>& FilteredChildren)
+	{
+		if (Parent->EntryType == FBindWidgetListEntry::EEntryType::Category)
+		{
+			for(const FBindWidgetListEntryPtr& Data : AllSourceData)
+			{
+				if (Data->EntryType == FBindWidgetListEntry::EEntryType::Binding && Data->BindingType == Parent->BindingType)
+				{
+					FilteredChildren.Add(Data);
+				}
+			}
+		}
+	}
+
+private:
+	TArray<FBindWidgetListEntryPtr> AllSourceData;
+	TArray<FBindWidgetListEntryPtr> CategorySourceData;
+};
+
+
+TArray<UWidget*> GetAllSourceWidgets(UWidgetBlueprint* WidgetBlueprint)
+{
+	TArray<UWidget*> Widgets;
+	UWidgetBlueprint* WidgetBPToScan = WidgetBlueprint;
+	while (WidgetBPToScan != nullptr)
+	{
+		Widgets = WidgetBPToScan->GetAllSourceWidgets();
+		if (Widgets.Num() != 0)
+		{
+			break;
+		}
+		// Get the parent WidgetBlueprint
+		WidgetBPToScan = WidgetBPToScan->ParentClass && WidgetBPToScan->ParentClass->ClassGeneratedBy ? Cast<UWidgetBlueprint>(WidgetBPToScan->ParentClass->ClassGeneratedBy) : nullptr;
+	}
+	return Widgets;
+}
+
+
+TArray<UWidgetAnimation*> GetAllSourceWidgetAnimations(UWidgetBlueprint* WidgetBlueprint)
+{
+	return WidgetBlueprint->Animations;
+}
+
+
+TArray<FBindWidgetListEntryPtr> BuildSourceDataList(TSharedPtr<FWidgetBlueprintEditor> InBlueprintEditor)
+{
+	TArray<FBindWidgetListEntryPtr> Result;
+
+	UWidgetBlueprint* WidgetBlueprint = CastChecked<UWidgetBlueprint>(InBlueprintEditor->GetBlueprintObj());
+	TArray<UWidget*> AllWidgets = GetAllSourceWidgets(WidgetBlueprint);
+	TArray<UWidgetAnimation*> AllWidgetAnimations = GetAllSourceWidgetAnimations(WidgetBlueprint);
+
+	for (FObjectProperty* WidgetProperty : TFieldRange<FObjectProperty>(WidgetBlueprint->ParentClass, EFieldIterationFlags::IncludeSuper))
+	{
+		if (WidgetProperty->PropertyClass->IsChildOf(UWidget::StaticClass()))
+		{
+			bool bIsOptional = false;
+			if (FWidgetBlueprintEditorUtils::IsBindWidgetProperty(WidgetProperty, bIsOptional))
+			{
+				TSharedRef<FBindWidgetListEntry> Ref = MakeShared<FBindWidgetListEntry>();
+				Ref->Property = WidgetProperty;
+				Ref->bIsOptional = bIsOptional;
+				Ref->EntryType = FBindWidgetListEntry::EEntryType::Binding;
+				Ref->BindingType = FBindWidgetListEntry::EBindingType::Widget;
+
+				int32 FoundIndex = AllWidgets.IndexOfByPredicate([WidgetProperty](UWidget* Widget){ return Widget->GetFName() == WidgetProperty->GetFName();});
+				Ref->bIsBound = FoundIndex != INDEX_NONE;
+				if (Ref->bIsBound)
+				{
+					Ref->bIsCorrectClass = AllWidgets[FoundIndex]->IsA(WidgetProperty->PropertyClass);
+				}
+
+				Result.Add(Ref);
+			}
+		}
+		else if (WidgetProperty->PropertyClass->IsChildOf(UWidgetAnimation::StaticClass()))
+		{
+			bool bIsOptional = false;
+			if (FWidgetBlueprintEditorUtils::IsBindWidgetProperty(WidgetProperty, bIsOptional))
+			{
+				TSharedRef<FBindWidgetListEntry> Ref = MakeShared<FBindWidgetListEntry>();
+				Ref->Property = WidgetProperty;
+				Ref->bIsOptional = bIsOptional;
+				Ref->EntryType = FBindWidgetListEntry::EEntryType::Binding;
+				Ref->BindingType = FBindWidgetListEntry::EBindingType::Animation;
+
+				int32 FoundIndex = AllWidgetAnimations.IndexOfByPredicate([WidgetProperty](UWidgetAnimation* Widget) { return Widget->GetFName() == WidgetProperty->GetFName(); });
+				Ref->bIsBound = FoundIndex != INDEX_NONE;
+				if (Ref->bIsBound)
+				{
+					Ref->bIsCorrectClass = AllWidgetAnimations[FoundIndex]->IsA(WidgetProperty->PropertyClass);
+				}
+
+				Result.Add(Ref);
+			}
+		}
+	}
+
+	return Result;
+}
+
+
+UWidgetBlueprint* GetWidgetBlueprint(TWeakPtr<FWidgetBlueprintEditor> InBlueprintEditor)
+{
+	if (TSharedPtr<FWidgetBlueprintEditor> Pin = InBlueprintEditor.Pin())
+	{
+		return CastChecked<UWidgetBlueprint>(Pin->GetBlueprintObj());
+	}
+	return nullptr;
+}
+
+} // namespace UMG
+} // namespace UE
+
+
+void SBindWidgetView::Construct(const FArguments& InArgs, TSharedPtr<FWidgetBlueprintEditor> InBlueprintEditor)
+{
+	BlueprintEditor = InBlueprintEditor;
+	bRefreshRequested = false;
+
+	// register for any objects replaced
+	FCoreUObjectDelegates::OnObjectsReplaced.AddRaw(this, &SBindWidgetView::HandleObjectsReplaced);
+
+	UWidgetBlueprint* Blueprint = CastChecked<UWidgetBlueprint>(InBlueprintEditor->GetBlueprintObj());
+	Blueprint->OnChanged().AddRaw(this, &SBindWidgetView::HandleBlueprintChanged);
+	Blueprint->OnCompiled().AddRaw(this, &SBindWidgetView::HandleBlueprintChanged);
+
+	ChildSlot
+	[
+		SNew(SBorder)
+		.BorderImage(FEditorStyle::GetBrush("ToolPanel.GroupBorder"))
+		[
+			SNew(SBorder)
+			.Padding(0)
+			.BorderImage(FEditorStyle::GetBrush("NoBrush"))
+			[
+				SAssignNew(ListView, UE::UMG::SBindWidgetView, UE::UMG::BuildSourceDataList(InBlueprintEditor))
+			]
+		]
+	];
+}
+
+
+SBindWidgetView::~SBindWidgetView()
+{
+	if (UWidgetBlueprint* Blueprint = UE::UMG::GetWidgetBlueprint(BlueprintEditor))
+	{
+		Blueprint->OnChanged().RemoveAll(this);
+		Blueprint->OnCompiled().RemoveAll(this);
+	}
+
+	FCoreUObjectDelegates::OnObjectsReplaced.RemoveAll(this);
+}
+
+
+void SBindWidgetView::Tick(const FGeometry& AllottedGeometry, const double InCurrentTime, const float InDeltaTime)
+{
+	if (bRefreshRequested)
+	{
+		TSharedPtr<UE::UMG::SBindWidgetView> ListViewPin = ListView.Pin();
+		TSharedPtr<FWidgetBlueprintEditor> BlueprintEditorPin = BlueprintEditor.Pin();
+		if (ListViewPin && BlueprintEditorPin)
+		{
+			ListViewPin->SetSourceData(UE::UMG::BuildSourceDataList(BlueprintEditorPin));
+		}
+		bRefreshRequested = false;
+	}
+}
+
+
+void SBindWidgetView::HandleObjectsReplaced(const TMap<UObject*, UObject*>& ReplacementMap)
+{
+	if (!bRefreshRequested)
+	{
+		for (const auto& Entry : ReplacementMap)
+		{
+			if (Entry.Key->IsA<UVisual>())
+			{
+				bRefreshRequested = true;
+			}
+		}
+	}
+}
+
+
+void SBindWidgetView::HandleBlueprintChanged(UBlueprint* InBlueprint)
+{
+	bRefreshRequested = true;
+}
+
+#undef LOCTEXT_NAMESPACE
diff --git a/Engine/Source/Editor/UMGEditor/Private/Widgets/SBindWidgetView.h b/Engine/Source/Editor/UMGEditor/Private/Widgets/SBindWidgetView.h
new file mode 100644
--- /dev/null
+++ b/Engine/Source/Editor/UMGEditor/Private/Widgets/SBindWidgetView.h
@@ -0,0 +1,43 @@
+// Copyright Epic Games, Inc. All Rights Reserved.
+
+#pragma once
+
+#include "CoreMinimal.h"
+#include "Widgets/DeclarativeSyntaxSupport.h"
+#include "Widgets/SCompoundWidget.h"
+
+class FWidgetBlueprintEditor;
+class UBlueprint;
+class UObject;
+
+namespace UE
+{
+namespace UMG
+{
+	class SBindWidgetView;
+}
+}
+
+/**
+ * 
+ */
+class SBindWidgetView : public SCompoundWidget
+{
+public:
+	SLATE_BEGIN_ARGS(SBindWidgetView){}
+	SLATE_END_ARGS()
+
+	void Construct(const FArguments& InArgs, TSharedPtr<FWidgetBlueprintEditor> InBlueprintEditor);
+	virtual ~SBindWidgetView();
+
+	virtual void Tick(const FGeometry& AllottedGeometry, const double InCurrentTime, const float InDeltaTime) override;
+
+private:
+	void HandleBlueprintChanged(UBlueprint* InBlueprint);
+	void HandleObjectsReplaced(const TMap<UObject*, UObject*>& ReplacementMap);
+
+private:
+	TWeakPtr<FWidgetBlueprintEditor> BlueprintEditor;
+	TWeakPtr<UE::UMG::SBindWidgetView> ListView;
+	bool bRefreshRequested;
+};
