Add missing include for NonUnity build. Introduced with CL17023249

diff --git a/Engine/Source/Editor/UMGEditor/Private/Widgets/SBindWidgetView.cpp b/Engine/Source/Editor/UMGEditor/Private/Widgets/SBindWidgetView.cpp
--- a/Engine/Source/Editor/UMGEditor/Private/Widgets/SBindWidgetView.cpp
+++ b/Engine/Source/Editor/UMGEditor/Private/Widgets/SBindWidgetView.cpp
@@ -9,7 +9,10 @@
 #include "Editor.h"
 #include "EditorStyleSet.h"
 #include "Styling/SlateIconFinder.h"
+#include "Widgets/SBoxPanel.h"
 #include "Widgets/SNullWidget.h"
+#include "Widgets/Images/SImage.h"
+#include "Widgets/Layout/SBox.h"
 #include "Widgets/Text/STextBlock.h"
 #include "Widgets/Views/STableRow.h"
 #include "Widgets/Views/STreeView.h"
