- Custom thumbnail for Widget blueprints #RB patrick.boutot #jira ue-118803

diff --git a/Engine/Source/Editor/UMGEditor/Private/Customizations/WidgetThumbnailCustomization.cpp b/Engine/Source/Editor/UMGEditor/Private/Customizations/WidgetThumbnailCustomization.cpp
new file mode 100644
--- /dev/null
+++ b/Engine/Source/Editor/UMGEditor/Private/Customizations/WidgetThumbnailCustomization.cpp
@@ -0,0 +1,104 @@
+// Copyright Epic Games, Inc. All Rights Reserved.
+
+#include "WidgetThumbnailCustomization.h"
+
+#include "DesktopPlatformModule.h"
+#include "DetailCategoryBuilder.h"
+#include "DetailLayoutBuilder.h"
+#include "DetailWidgetRow.h"
+#include "ImageUtils.h"
+#include "ImageWrapperHelper.h"
+#include "PropertyCustomizationHelpers.h"
+#include "Widgets/Input/SButton.h"
+#include "Widgets/Layout/SBox.h"
+#include "Widgets/Text/STextBlock.h"
+#include "WidgetBlueprintEditorUtils.h"
+
+#define LOCTEXT_NAMESPACE "FWidgetThumbnailCustomization"
+
+TSharedRef<IDetailCustomization> FWidgetThumbnailCustomization::MakeInstance()
+{
+	return MakeShared<FWidgetThumbnailCustomization>();
+}
+
+void FWidgetThumbnailCustomization::CustomizeDetails(IDetailLayoutBuilder& DetailBuilder)
+{
+	TArray<TWeakObjectPtr<UObject>> ObjectsBeingCustomized;
+	DetailBuilder.GetObjectsBeingCustomized(ObjectsBeingCustomized);
+	if (ObjectsBeingCustomized.Num() > 1)
+	{
+		return;
+	}
+	WidgetBlueprint = Cast<UWidgetBlueprint>(ObjectsBeingCustomized[0].Get());
+	if (!WidgetBlueprint.IsValid())
+	{
+		return;
+	}
+
+	IDetailCategoryBuilder& ThumbnailCategory = DetailBuilder.EditCategory("ThumbnailSettings", FText::GetEmpty(), ECategoryPriority::Default);
+	TSharedRef< IPropertyHandle > ThumbnailImage = DetailBuilder.GetProperty("ThumbnailImage");
+	TSharedRef< IPropertyHandle > ThumbnailSizeMode = DetailBuilder.GetProperty("ThumbnailSizeMode");
+	TSharedRef< IPropertyHandle > ThumbnailCustomSize = DetailBuilder.GetProperty("ThumbnailCustomSize");
+	ThumbnailCategory.AddProperty(ThumbnailSizeMode).IsEnabled(TAttribute<bool>(this, &FWidgetThumbnailCustomization::IsThumbnailAutomatic));
+	ThumbnailCategory.AddProperty(ThumbnailCustomSize).IsEnabled(TAttribute<bool>(this, &FWidgetThumbnailCustomization::IsThumbnailAutomaticAndCustom));
+	ThumbnailCategory.AddProperty(ThumbnailImage);
+
+	ThumbnailCategory.AddCustomRow(LOCTEXT("Browse_Thumbnail", "Custom Thumbnail"))
+		[
+			SNew(SHorizontalBox)
+
+			+ SHorizontalBox::Slot()
+			.AutoWidth()
+			[
+				SNew(SButton)
+				.OnClicked(this, &FWidgetThumbnailCustomization::LoadThumbnailFromFile)
+				.ToolTipText(LOCTEXT("Button_Import_Tooltip", "Browse computer for an image to set as thumbnail."))
+				.ContentPadding(2.f)
+				.Text(LOCTEXT("ImportImage", "Import Image"))
+			]
+		];
+}
+FReply FWidgetThumbnailCustomization::LoadThumbnailFromFile()
+{
+	IDesktopPlatform* DesktopPlatform = FDesktopPlatformModule::Get();
+
+	if (DesktopPlatform)
+	{
+		TArray<FString> OpenFilenames;
+		TSharedPtr<SWindow> ParentWindow = FGlobalTabmanager::Get()->GetRootWindow();
+		bool bOpened = DesktopPlatform->OpenFileDialog(
+			(ParentWindow.IsValid()) ? ParentWindow->GetNativeWindow()->GetOSWindowHandle() : nullptr,
+			LOCTEXT("LoadSnapshotDialogTitle", "Load an Image to set as thumbnail for this widget blueprint.").ToString(),
+			FPaths::GameAgnosticSavedDir(),
+			TEXT(""),
+			ImageWrapperHelper::GetImageFilesFilterString(false).GetData(),
+			EFileDialogFlags::None,
+			OpenFilenames
+		);
+
+		if (OpenFilenames.Num() > 0)
+		{
+			// load the new thumbnail
+			UTexture2D* ThumbnailTexture = FImageUtils::ImportFileAsTexture2D(*OpenFilenames[0]);
+			FWidgetBlueprintEditorUtils::SetTextureAsAssetThumbnail(WidgetBlueprint.Get(), ThumbnailTexture);
+		}
+	}
+
+	return FReply::Handled();
+}
+
+void FWidgetThumbnailCustomization::ClearThumbnail()
+{
+	WidgetBlueprint->ThumbnailImage = nullptr;
+}
+
+bool FWidgetThumbnailCustomization::IsThumbnailAutomatic() const
+{
+	return WidgetBlueprint->ThumbnailImage == nullptr;
+}
+
+bool FWidgetThumbnailCustomization::IsThumbnailAutomaticAndCustom() const
+{
+	return WidgetBlueprint->ThumbnailSizeMode == EThumbnailPreviewSizeMode::Custom && IsThumbnailAutomatic();
+}
+#undef LOCTEXT_NAMESPACE
diff --git a/Engine/Source/Editor/UMGEditor/Private/Customizations/WidgetThumbnailCustomization.h b/Engine/Source/Editor/UMGEditor/Private/Customizations/WidgetThumbnailCustomization.h
new file mode 100644
--- /dev/null
+++ b/Engine/Source/Editor/UMGEditor/Private/Customizations/WidgetThumbnailCustomization.h
@@ -0,0 +1,25 @@
+// Copyright Epic Games, Inc. All Rights Reserved.
+
+#pragma once
+
+#include "CoreMinimal.h"
+#include "IDetailCustomization.h"
+#include "WidgetBlueprint.h"
+
+/** UI customization for WidgetBlueprint thumbnail */
+class FWidgetThumbnailCustomization : public IDetailCustomization
+{
+public:
+	static TSharedRef<IDetailCustomization> MakeInstance();
+	//~ Begin IDetailCustomization Interface.
+	virtual void CustomizeDetails(IDetailLayoutBuilder& DetailBuilder) override;
+	//~ End IDetailCustomization Interface.
+
+private:
+	FReply LoadThumbnailFromFile();
+	void ClearThumbnail();
+	bool IsThumbnailAutomatic() const;
+	bool IsThumbnailAutomaticAndCustom() const;
+
+	TWeakObjectPtr<UWidgetBlueprint> WidgetBlueprint;
+};
diff --git a/Engine/Source/Editor/UMGEditor/Private/Designer/SDesignerView.cpp b/Engine/Source/Editor/UMGEditor/Private/Designer/SDesignerView.cpp
--- a/Engine/Source/Editor/UMGEditor/Private/Designer/SDesignerView.cpp
+++ b/Engine/Source/Editor/UMGEditor/Private/Designer/SDesignerView.cpp
@@ -1065,7 +1065,7 @@ ETransformMode::Type SDesignerView::GetTransformMode() const
 
 FOptionalSize SDesignerView::GetPreviewAreaWidth() const
 {
-	TTuple<FVector2D, FVector2D> AreaAndSize = FWidgetBlueprintEditorUtils::GetWidgetPreviewAreaAndSize(GetDefaultWidget(), CachedPreviewDesiredSize, FVector2D(PreviewWidth, PreviewHeight));
+	TTuple<FVector2D, FVector2D> AreaAndSize = FWidgetBlueprintEditorUtils::GetWidgetPreviewAreaAndSize(GetDefaultWidget(), CachedPreviewDesiredSize, FVector2D(PreviewWidth, PreviewHeight), GetDefaultWidget()->DesignSizeMode, TOptional<FVector2D>());
 	FVector2D Area = AreaAndSize.Get<0>();
 
 	return Area.X;
@@ -1073,7 +1073,7 @@ FOptionalSize SDesignerView::GetPreviewAreaWidth() const
 
 FOptionalSize SDesignerView::GetPreviewAreaHeight() const
 {
-	TTuple<FVector2D, FVector2D> AreaAndSize = FWidgetBlueprintEditorUtils::GetWidgetPreviewAreaAndSize(GetDefaultWidget(), CachedPreviewDesiredSize, FVector2D(PreviewWidth, PreviewHeight));
+	TTuple<FVector2D, FVector2D> AreaAndSize = FWidgetBlueprintEditorUtils::GetWidgetPreviewAreaAndSize(GetDefaultWidget(), CachedPreviewDesiredSize, FVector2D(PreviewWidth, PreviewHeight), GetDefaultWidget()->DesignSizeMode, TOptional<FVector2D>());
 	FVector2D Area = AreaAndSize.Get<0>();;
 
 	return Area.Y;
@@ -1081,7 +1081,7 @@ FOptionalSize SDesignerView::GetPreviewAreaHeight() const
 
 FOptionalSize SDesignerView::GetPreviewSizeWidth() const
 {
-	TTuple<FVector2D, FVector2D> AreaAndSize = FWidgetBlueprintEditorUtils::GetWidgetPreviewAreaAndSize(GetDefaultWidget(), CachedPreviewDesiredSize, FVector2D(PreviewWidth, PreviewHeight));
+	TTuple<FVector2D, FVector2D> AreaAndSize = FWidgetBlueprintEditorUtils::GetWidgetPreviewAreaAndSize(GetDefaultWidget(), CachedPreviewDesiredSize, FVector2D(PreviewWidth, PreviewHeight), GetDefaultWidget()->DesignSizeMode, TOptional<FVector2D>());
 	FVector2D Size = AreaAndSize.Get<1>();
 
 	return Size.X;
@@ -1089,7 +1089,7 @@ FOptionalSize SDesignerView::GetPreviewSizeWidth() const
 
 FOptionalSize SDesignerView::GetPreviewSizeHeight() const
 {
-	TTuple<FVector2D, FVector2D> AreaAndSize = FWidgetBlueprintEditorUtils::GetWidgetPreviewAreaAndSize(GetDefaultWidget(), CachedPreviewDesiredSize, FVector2D(PreviewWidth, PreviewHeight));
+	TTuple<FVector2D, FVector2D> AreaAndSize = FWidgetBlueprintEditorUtils::GetWidgetPreviewAreaAndSize(GetDefaultWidget(), CachedPreviewDesiredSize, FVector2D(PreviewWidth, PreviewHeight), GetDefaultWidget()->DesignSizeMode, TOptional<FVector2D>());
 	FVector2D Size = AreaAndSize.Get<1>();
 
 	return Size.Y;
diff --git a/Engine/Source/Editor/UMGEditor/Private/UMGEditorActions.cpp b/Engine/Source/Editor/UMGEditor/Private/UMGEditorActions.cpp
--- a/Engine/Source/Editor/UMGEditor/Private/UMGEditorActions.cpp
+++ b/Engine/Source/Editor/UMGEditor/Private/UMGEditorActions.cpp
@@ -7,6 +7,10 @@
 void FUMGEditorCommands::RegisterCommands()
 {
 	UI_COMMAND(CreateNativeBaseClass, "Create Native Base Class", "Create a native base class for this widget, using the current parent as the parent of the native class.", EUserInterfaceActionType::Button, FInputChord() );
+	UI_COMMAND(ExportAsPNG, "Export as Image", "Export the current widget blueprint as .png format.", EUserInterfaceActionType::Button, FInputChord());
+	UI_COMMAND(SetImageAsThumbnail, "Capture and Set as Thumbnail", "Captures the current state of the widget blueprint and sets it as the thumbnail for this asset.", EUserInterfaceActionType::Button, FInputChord());
+	UI_COMMAND(ClearCustomThumbnail, "Clear Thumbnail", "Removes the image used as thumbnail and enables automatic thumbnail generation.", EUserInterfaceActionType::Button, FInputChord());
+
 }
 
 #undef LOCTEXT_NAMESPACE
\ No newline at end of file
diff --git a/Engine/Source/Editor/UMGEditor/Private/UMGEditorActions.h b/Engine/Source/Editor/UMGEditor/Private/UMGEditorActions.h
--- a/Engine/Source/Editor/UMGEditor/Private/UMGEditorActions.h
+++ b/Engine/Source/Editor/UMGEditor/Private/UMGEditorActions.h
@@ -27,6 +27,9 @@ public:
 	 * UMG Editor Commands
 	 */
 	TSharedPtr< FUICommandInfo > CreateNativeBaseClass;
+	TSharedPtr< FUICommandInfo > ExportAsPNG;
+	TSharedPtr< FUICommandInfo > SetImageAsThumbnail;
+	TSharedPtr< FUICommandInfo > ClearCustomThumbnail;
 
 	/**
 	 * Initialize commands
diff --git a/Engine/Source/Editor/UMGEditor/Private/UMGEditorModule.cpp b/Engine/Source/Editor/UMGEditor/Private/UMGEditorModule.cpp
--- a/Engine/Source/Editor/UMGEditor/Private/UMGEditorModule.cpp
+++ b/Engine/Source/Editor/UMGEditor/Private/UMGEditorModule.cpp
@@ -39,6 +39,7 @@
 #include "DynamicEntryBoxDetails.h"
 #include "ListViewBaseDetails.h"
 #include "WidgetBlueprintThumbnailRenderer.h"
+#include "WidgetThumbnailCustomization.h"
 
 #define LOCTEXT_NAMESPACE "UMG"
 
@@ -105,7 +106,7 @@ public:
 		PropertyModule.RegisterCustomClassLayout(TEXT("DynamicEntryBoxBase"), FOnGetDetailCustomizationInstance::CreateStatic(&FDynamicEntryBoxBaseDetails::MakeInstance));
 		PropertyModule.RegisterCustomClassLayout(TEXT("DynamicEntryBox"), FOnGetDetailCustomizationInstance::CreateStatic(&FDynamicEntryBoxDetails::MakeInstance));
 		PropertyModule.RegisterCustomClassLayout(TEXT("ListViewBase"), FOnGetDetailCustomizationInstance::CreateStatic(&FListViewBaseDetails::MakeInstance));
-
+		PropertyModule.RegisterCustomClassLayout(TEXT("WidgetBlueprint"), FOnGetDetailCustomizationInstance::CreateStatic(&FWidgetThumbnailCustomization::MakeInstance));
 	
 		CVarThumbnailRenderEnable->AsVariable()->SetOnChangedCallback(FConsoleVariableDelegate::CreateStatic(&FUMGEditorModule::ThumbnailRenderingEnabled));
 	}
@@ -150,6 +151,12 @@ public:
 
 		UnregisterSettings();
 
+		FPropertyEditorModule* PropertyModule = FModuleManager::GetModulePtr<FPropertyEditorModule>("PropertyEditor");
+		if (PropertyModule != nullptr)
+		{
+			PropertyModule->UnregisterCustomClassLayout(TEXT("WidgetBlueprint"));
+		}
+
 		//// Unregister the setting
 		//ISettingsModule* SettingsModule = FModuleManager::GetModulePtr<ISettingsModule>("Settings");
 
diff --git a/Engine/Source/Editor/UMGEditor/Private/WidgetBlueprintEditor.cpp b/Engine/Source/Editor/UMGEditor/Private/WidgetBlueprintEditor.cpp
--- a/Engine/Source/Editor/UMGEditor/Private/WidgetBlueprintEditor.cpp
+++ b/Engine/Source/Editor/UMGEditor/Private/WidgetBlueprintEditor.cpp
@@ -58,6 +58,15 @@
 
 #include "SPaletteViewModel.h"
 
+#include "DesktopPlatformModule.h"
+#include "Engine/TextureRenderTarget2D.h"
+#include "IDesktopPlatform.h"
+#include "IImageWrapper.h"
+#include "IImageWrapperModule.h"
+#include "ImageUtils.h"
+#include "Serialization/BufferArchive.h"
+#include "Widgets/SVirtualWindow.h"
+
 #define LOCTEXT_NAMESPACE "UMG"
 
 FWidgetBlueprintEditor::FWidgetBlueprintEditor()
@@ -179,16 +188,34 @@ TSharedPtr<FExtender> FWidgetBlueprintEditor::CreateMenuExtender()
 		GetToolkitCommands(),
 		FMenuExtensionDelegate::CreateSP(this, &FWidgetBlueprintEditor::FillFileMenu));
 	
+	MenuExtender->AddMenuExtension(
+		"AssetEditorActions",
+		EExtensionHook::After,
+		GetToolkitCommands(),
+		FMenuExtensionDelegate::CreateSP(this, &FWidgetBlueprintEditor::FillAssetMenu));
+
 	return MenuExtender;
 }
 
 void FWidgetBlueprintEditor::FillFileMenu(FMenuBuilder& MenuBuilder)
 {
+	MenuBuilder.BeginSection(TEXT("Import/Export"), LOCTEXT("Import/Export", "Import/Export"));
+	MenuBuilder.AddMenuEntry(FUMGEditorCommands::Get().ExportAsPNG);
+	MenuBuilder.EndSection();
+
 	MenuBuilder.BeginSection(TEXT("WidgetBlueprint"), LOCTEXT("WidgetBlueprint", "Widget Blueprint"));
 	MenuBuilder.AddMenuEntry(FUMGEditorCommands::Get().CreateNativeBaseClass);
 	MenuBuilder.EndSection();
 }
 
+void FWidgetBlueprintEditor::FillAssetMenu(FMenuBuilder& MenuBuilder)
+{
+	MenuBuilder.BeginSection(TEXT("Thumbnail"), LOCTEXT("Thumbnail", "Thumbnail"));
+	MenuBuilder.AddMenuEntry(FUMGEditorCommands::Get().SetImageAsThumbnail);
+	MenuBuilder.AddMenuEntry(FUMGEditorCommands::Get().ClearCustomThumbnail);
+	MenuBuilder.EndSection();
+}
+
 void FWidgetBlueprintEditor::BindToolkitCommands()
 {
 	FUMGEditorCommands::Register();
@@ -199,6 +226,122 @@ void FWidgetBlueprintEditor::BindToolkitCommands()
 			FCanExecuteAction::CreateSP(this, &FWidgetBlueprintEditor::IsParentClassNative)
 		)
 	);
+
+	GetToolkitCommands()->MapAction(FUMGEditorCommands::Get().ExportAsPNG,
+		FUIAction(
+			FExecuteAction::CreateSP(this, &FWidgetBlueprintEditor::TakeSnapshot),
+			FCanExecuteAction::CreateSP(this, &FWidgetBlueprintEditor::IsPreviewWidgetInitialized)
+		)
+	);
+
+	GetToolkitCommands()->MapAction(FUMGEditorCommands::Get().SetImageAsThumbnail,
+		FUIAction(
+			FExecuteAction::CreateSP(this, &FWidgetBlueprintEditor::CaptureThumbnail)
+		)
+	);
+
+	GetToolkitCommands()->MapAction(FUMGEditorCommands::Get().ClearCustomThumbnail,
+		FUIAction(
+			FExecuteAction::CreateSP(this, &FWidgetBlueprintEditor::ClearThumbnail),
+			FCanExecuteAction::CreateSP(this, &FWidgetBlueprintEditor::IsImageUsedAsThumbnail)
+		)
+	);
+}
+
+void FWidgetBlueprintEditor::TakeSnapshot()
+{
+	IDesktopPlatform* DesktopPlatform = FDesktopPlatformModule::Get();
+	if (DesktopPlatform)
+	{
+		TSharedPtr<SWindow> ParentWindow = FGlobalTabmanager::Get()->GetRootWindow();
+		const void* ParentWindowWindowHandle = FSlateApplication::Get().FindBestParentWindowHandleForDialogs(ParentWindow);
+		TArray<FString> SaveFilenames;
+		const bool bOpened = DesktopPlatform->SaveFileDialog(
+			ParentWindowWindowHandle,
+			LOCTEXT("ExportWidgetBlueprintDialogTitle", "Save Widget Blueprint Screenshot").ToString(),
+			FPaths::GameAgnosticSavedDir(),
+			TEXT(""),
+			TEXT("PNG (*.png)|*.png"),
+			EFileDialogFlags::None,
+			SaveFilenames
+		);
+		if (SaveFilenames.Num() > 0)
+		{
+			TUniquePtr<FArchive> Ar(IFileManager::Get().CreateFileWriter(*SaveFilenames[0]));
+			if (Ar)
+			{
+				TSharedPtr<SWidget> WindowContent;
+				UUserWidget* PreviewWidget = GetPreview();
+
+				UTextureRenderTarget2D* RenderTarget2D = NewObject<UTextureRenderTarget2D>();
+				TOptional<FWidgetBlueprintEditorUtils::FWidgetThumbnailProperties> ScaleAndOffset = FWidgetBlueprintEditorUtils::DrawSWidgetInRenderTarget(PreviewWidget, RenderTarget2D);
+
+				if (!ScaleAndOffset.IsSet())
+				{
+					FMessageLog("Blueprint").Warning(LOCTEXT("ExportWidgetBlueprint_ImageSourceFailedToCreate", "ExportWidgetBlueprint: Failed to create image source."));
+					return;
+				}
+
+				FBufferArchive Buffer;
+				bool bSuccess = FImageUtils::ExportRenderTarget2DAsPNG(RenderTarget2D, Buffer);
+				if (bSuccess)
+				{
+					Ar->Serialize(const_cast<uint8*>(Buffer.GetData()), Buffer.Num());
+				}
+			}
+		}
+	}
+
+
+}
+
+void FWidgetBlueprintEditor::CaptureThumbnail()
+{
+	TSharedPtr<SWidget> WindowContent;
+	UUserWidget* PreviewWidget = GetPreview();
+
+	if (!PreviewWidget)
+	{
+		return;
+	}
+
+	UTextureRenderTarget2D* RenderTarget2D = NewObject<UTextureRenderTarget2D>();
+	RenderTarget2D->Filter = TF_Bilinear;
+	RenderTarget2D->ClearColor = FLinearColor::Transparent;
+	RenderTarget2D->SRGB = true;
+	RenderTarget2D->TargetGamma = 1;
+
+	TOptional<FWidgetBlueprintEditorUtils::FWidgetThumbnailProperties> ScaleAndOffset = FWidgetBlueprintEditorUtils::DrawSWidgetInRenderTargetForThumbnail(PreviewWidget, RenderTarget2D, FVector2D(256.f, 256.f), TOptional<FVector2D>());
+
+	if (!ScaleAndOffset.IsSet())
+	{
+		return;
+	}
+
+	TArray64<uint8> RawData;
+	FImageUtils::GetRawData(RenderTarget2D, RawData);
+	IImageWrapperModule& ImageWrapperModule = FModuleManager::Get().LoadModuleChecked<IImageWrapperModule>(TEXT("ImageWrapper"));
+	TSharedPtr<IImageWrapper> PNGImageWrapper = ImageWrapperModule.CreateImageWrapper(EImageFormat::PNG);
+	PNGImageWrapper->SetRaw(RawData.GetData(), RawData.GetAllocatedSize(), RenderTarget2D->SizeX, RenderTarget2D->SizeY, ERGBFormat::BGRA, 8);
+	const TArray64<uint8> PNGData = PNGImageWrapper->GetCompressed(100);
+
+	UTexture2D* ThumbnailTexture = FImageUtils::ImportBufferAsTexture2D(PNGData);
+	FWidgetBlueprintEditorUtils::SetTextureAsAssetThumbnail(GetWidgetBlueprintObj(), ThumbnailTexture);
+}
+
+void FWidgetBlueprintEditor::ClearThumbnail() 
+{
+	GetWidgetBlueprintObj()->ThumbnailImage = nullptr;
+}
+
+bool FWidgetBlueprintEditor::IsImageUsedAsThumbnail()
+{
+	return GetWidgetBlueprintObj()->ThumbnailImage != nullptr;
+}
+
+bool FWidgetBlueprintEditor::IsPreviewWidgetInitialized()
+{
+	return GetPreview() != nullptr;
 }
 
 void FWidgetBlueprintEditor::OpenCreateNativeBaseClassDialog()
diff --git a/Engine/Source/Editor/UMGEditor/Private/WidgetBlueprintEditorUtils.cpp b/Engine/Source/Editor/UMGEditor/Private/WidgetBlueprintEditorUtils.cpp
--- a/Engine/Source/Editor/UMGEditor/Private/WidgetBlueprintEditorUtils.cpp
+++ b/Engine/Source/Editor/UMGEditor/Private/WidgetBlueprintEditorUtils.cpp
@@ -41,11 +41,14 @@
 #include "Components/NamedSlotInterface.h"
 #include "K2Node_Variable.h"
 
+#include "Engine/TextureRenderTarget2D.h"
 #include "Engine/UserInterfaceSettings.h"
 #include "Input/HittestGrid.h"
 #include "Interfaces/ISlateRHIRendererModule.h"
 #include "Interfaces/ISlate3DRenderer.h"
 #include "Rendering/SlateDrawBuffer.h"
+#include "Slate/WidgetRenderer.h"
+#include "Widgets/SVirtualWindow.h"
 
 #define LOCTEXT_NAMESPACE "UMG"
 
@@ -1931,17 +1934,18 @@ int32 FWidgetBlueprintEditorUtils::UpdateHittestGrid(FHittestGrid& HitTestGrid,
 	return MaxLayerId;
 }
 
-TTuple<FVector2D, FVector2D> FWidgetBlueprintEditorUtils::GetWidgetPreviewAreaAndSize(UUserWidget* UserWidget, FVector2D DesiredSize, FVector2D PreviewSize)
+TTuple<FVector2D, FVector2D> FWidgetBlueprintEditorUtils::GetWidgetPreviewAreaAndSize(UUserWidget* UserWidget, FVector2D DesiredSize, FVector2D PreviewSize, EDesignPreviewSizeMode SizeMode, TOptional<FVector2D> ThumbnailCustomSize)
 {
 	FVector2D Size(PreviewSize.X, PreviewSize.Y);
 	FVector2D Area(PreviewSize.X, PreviewSize.Y);
 
 	if (UserWidget)
 	{
-		switch (UserWidget->DesignSizeMode)
+
+		switch (SizeMode)
 		{
 		case EDesignPreviewSizeMode::Custom:
-			Area = UserWidget->DesignTimeSize;
+			Area = ThumbnailCustomSize.IsSet()? ThumbnailCustomSize.GetValue() : UserWidget->DesignTimeSize;
 			// If the custom size is 0 in some dimension, use the desired size instead.
 			if (Area.X == 0)
 			{
@@ -1954,7 +1958,7 @@ TTuple<FVector2D, FVector2D> FWidgetBlueprintEditorUtils::GetWidgetPreviewAreaAn
 			Size = Area;
 			break;
 		case EDesignPreviewSizeMode::CustomOnScreen:
-			Size = UserWidget->DesignTimeSize;
+			Size = ThumbnailCustomSize.IsSet() ? ThumbnailCustomSize.GetValue() : UserWidget->DesignTimeSize;
 
 			// If the custom size is 0 in some dimension, use the desired size instead.
 			if (Size.X == 0)
@@ -1995,7 +1999,7 @@ float FWidgetBlueprintEditorUtils::GetWidgetPreviewDPIScale(UUserWidget* UserWid
 }
 
 
-FVector2D FWidgetBlueprintEditorUtils::GetWidgetPreviewUnScaledCustomSize(FVector2D DesiredSize, UUserWidget* UserWidget)
+FVector2D FWidgetBlueprintEditorUtils::GetWidgetPreviewUnScaledCustomSize(FVector2D DesiredSize, UUserWidget* UserWidget, TOptional<FVector2D> ThumbnailCustomSize, EThumbnailPreviewSizeMode ThumbnailSizeMode)
 {
 
 	checkf(DesiredSize.X > 0.f && DesiredSize.Y > 0.f, TEXT("The size should have been previously checked to be > 0."));
@@ -2009,9 +2013,18 @@ FVector2D FWidgetBlueprintEditorUtils::GetWidgetPreviewUnScaledCustomSize(FVecto
 
 	FVector2D PreviewSize(PreviewWidth, PreviewHeight);
 
-	TTuple<FVector2D, FVector2D> AreaAndSize = GetWidgetPreviewAreaAndSize(UserWidget, DesiredSize, PreviewSize);
+	TTuple<FVector2D, FVector2D> AreaAndSize = GetWidgetPreviewAreaAndSize(UserWidget, DesiredSize, PreviewSize, ConvertThumbnailSizeModeToDesignerSizeMode(ThumbnailSizeMode, UserWidget), ThumbnailCustomSize.IsSet() ? ThumbnailCustomSize.GetValue() : TOptional<FVector2D>());
 
-	float DPIScale = GetWidgetPreviewDPIScale(UserWidget, PreviewSize);
+	float DPIScale;
+	if (ThumbnailCustomSize.IsSet())
+	{
+		DPIScale = 1.0f;
+	}
+	else 
+	{
+		DPIScale = GetWidgetPreviewDPIScale(UserWidget, PreviewSize);
+
+	}
 
 	if (ensure(DPIScale > 0.f))
 	{
@@ -2021,4 +2034,155 @@ FVector2D FWidgetBlueprintEditorUtils::GetWidgetPreviewUnScaledCustomSize(FVecto
 	return FinalSize;
 
 }
+
+EDesignPreviewSizeMode FWidgetBlueprintEditorUtils::ConvertThumbnailSizeModeToDesignerSizeMode(EThumbnailPreviewSizeMode ThumbnailSizeMode, UUserWidget* WidgetInstance)
+{
+	switch (ThumbnailSizeMode)
+	{
+	case EThumbnailPreviewSizeMode::MatchDesignerMode:
+		return WidgetInstance->DesignSizeMode;
+	case EThumbnailPreviewSizeMode::FillScreen:
+		return EDesignPreviewSizeMode::FillScreen;
+	case EThumbnailPreviewSizeMode::Custom:
+		return EDesignPreviewSizeMode::Custom;
+	case EThumbnailPreviewSizeMode::Desired:
+		return EDesignPreviewSizeMode::Desired;
+	default:
+		return EDesignPreviewSizeMode::Desired;
+	}
+}
+
+TOptional<FWidgetBlueprintEditorUtils::FWidgetThumbnailProperties> FWidgetBlueprintEditorUtils::DrawSWidgetInRenderTarget(UUserWidget* WidgetInstance, UTextureRenderTarget2D* RenderTarget2D)
+{
+	return DrawSWidgetInRenderTargetInternal(WidgetInstance, RenderTarget2D, FVector2D(256.f,256.f), false, TOptional<FVector2D>(), EThumbnailPreviewSizeMode::MatchDesignerMode);
+}
+
+TOptional<FWidgetBlueprintEditorUtils::FWidgetThumbnailProperties> FWidgetBlueprintEditorUtils::DrawSWidgetInRenderTargetForThumbnail(UUserWidget* WidgetInstance, UTextureRenderTarget2D* RenderTarget2D, FVector2D ThumbnailSize, TOptional<FVector2D> ThumbnailCustomSize, EThumbnailPreviewSizeMode ThumbnailSizeMode)
+{
+	return DrawSWidgetInRenderTargetInternal(WidgetInstance, RenderTarget2D, ThumbnailSize, true, ThumbnailCustomSize, ThumbnailSizeMode);
+}
+
+TOptional<FWidgetBlueprintEditorUtils::FWidgetThumbnailProperties>  FWidgetBlueprintEditorUtils::DrawSWidgetInRenderTargetInternal(UUserWidget* WidgetInstance, UTextureRenderTarget2D* RenderTarget2D, FVector2D ThumbnailSize, bool bIsForThumbnail, TOptional<FVector2D> ThumbnailCustomSize,EThumbnailPreviewSizeMode ThumbnailSizeMode)
+{
+	//Create Window
+	FVector2D Offset(0.f,0.f);
+	FVector2D ScaledSize(0.f, 0.f);
+	TSharedPtr<SWidget> WindowContent = WidgetInstance->TakeWidget();
+
+	if (!WindowContent)
+	{
+		return TOptional<FWidgetBlueprintEditorUtils::FWidgetThumbnailProperties>();
+	}
+
+	TSharedRef<SVirtualWindow> Window = SNew(SVirtualWindow).Size(ThumbnailSize);
+	TUniquePtr<FHittestGrid> HitTestGrid = MakeUnique<FHittestGrid>();
+	Window->SetContent(WindowContent.ToSharedRef());
+	Window->Resize(ThumbnailSize);
+
+	// Store the desired size to maintain the aspect ratio later
+	FGeometry WindowGeometry = FGeometry::MakeRoot(ThumbnailSize, FSlateLayoutTransform(1.0f));
+	Window->SlatePrepass(1.0f);
+	FVector2D DesiredSizeWindow = Window->GetDesiredSize();
+
+
+	if (DesiredSizeWindow.X < SMALL_NUMBER || DesiredSizeWindow.Y < SMALL_NUMBER)
+	{
+		return TOptional<FWidgetBlueprintEditorUtils::FWidgetThumbnailProperties>();
+	}
+
+	FVector2D UnscaledSize = FWidgetBlueprintEditorUtils::GetWidgetPreviewUnScaledCustomSize(DesiredSizeWindow, WidgetInstance, ThumbnailCustomSize,ThumbnailSizeMode);
+	if (UnscaledSize.X < SMALL_NUMBER || UnscaledSize.Y < SMALL_NUMBER)
+	{
+		return TOptional<FWidgetBlueprintEditorUtils::FWidgetThumbnailProperties>();
+	}
+
+	float Scale;
+	// Create Renderer Target and WidgetRenderer
+	bool bApplyGammaCorrection = false;
+	FWidgetRenderer* WidgetRenderer = new FWidgetRenderer(bApplyGammaCorrection);
+	const bool bIsLinearSpace = !bApplyGammaCorrection;
+
+	// Change some configuration if it is for thumbnail creation
+	if (bIsForThumbnail) 
+	{
+		TTuple<float, FVector2D> ScaleAndOffset = FWidgetBlueprintEditorUtils::GetThumbnailImageScaleAndOffset(UnscaledSize, ThumbnailSize);
+		Scale = ScaleAndOffset.Get<0>();
+		Offset = ScaleAndOffset.Get<1>();
+	}
+	else
+	{
+		RenderTarget2D->Filter = TF_Bilinear;
+		RenderTarget2D->ClearColor = FLinearColor::Transparent;
+		RenderTarget2D->SRGB = true;
+		RenderTarget2D->RenderTargetFormat = RTF_RGBA8;
+		Scale = 1.f;
+		Offset = FVector2D(0.f, 0.f);
+		WidgetRenderer->SetIsPrepassNeeded(false);
+	}
+
+	ScaledSize = UnscaledSize * Scale;
+
+	const EPixelFormat RequestedFormat = FSlateApplication::Get().GetRenderer()->GetSlateRecommendedColorFormat();
+	RenderTarget2D->InitCustomFormat(ScaledSize.X, ScaledSize.Y, RequestedFormat, false);
+	WidgetRenderer->DrawWindow(RenderTarget2D, *HitTestGrid, Window, Scale, ScaledSize, 0.1f);
+
+	if (WidgetRenderer)
+	{
+		BeginCleanup(WidgetRenderer);
+		WidgetRenderer = nullptr;
+	}
+
+	return TOptional<FWidgetBlueprintEditorUtils::FWidgetThumbnailProperties>(FWidgetBlueprintEditorUtils::FWidgetThumbnailProperties{ ScaledSize,Offset });
+}
+
+TTuple<float, FVector2D> FWidgetBlueprintEditorUtils::GetThumbnailImageScaleAndOffset(FVector2D WidgetSize, FVector2D ThumbnailSize)
+{
+	// Scale the widget blueprint image to fit in the thumbnail
+
+	checkf(WidgetSize.X > 0.f && WidgetSize.Y > 0.f, TEXT("The size should have been previously checked to be > 0."));
+
+	float Scale;
+	float XOffset = 0;
+	float YOffset = 0;
+	if (WidgetSize.X > WidgetSize.Y)
+	{
+		Scale = ThumbnailSize.X / WidgetSize.X;
+		WidgetSize *= Scale;
+		YOffset = (ThumbnailSize.Y - WidgetSize.Y) / 2.f;
+	}
+	else
+	{
+		Scale = ThumbnailSize.Y / WidgetSize.Y;
+		WidgetSize *= Scale;
+		XOffset = (ThumbnailSize.X - WidgetSize.X) / 2.f;
+	}
+
+	return TTuple<float, FVector2D>(Scale, FVector2D(XOffset, YOffset));
+}
+
+void FWidgetBlueprintEditorUtils::SetTextureAsAssetThumbnail(UWidgetBlueprint* WidgetBlueprint, UTexture2D* ThumbnailTexture)
+{
+	const TCHAR* ThumbnailName = TEXT("Thumbnail");
+	UTexture2D* ExistingThumbnail = FindObject<UTexture2D>(WidgetBlueprint, ThumbnailName, false);
+	if (ExistingThumbnail)
+	{
+		ExistingThumbnail->Rename(nullptr, GetTransientPackage());
+	}
+	if (!ThumbnailTexture)
+	{
+		WidgetBlueprint->ThumbnailImage = nullptr;
+		return;
+	}
+	FVector2D TextureSize(ThumbnailTexture->GetSizeX(), ThumbnailTexture->GetSizeY());
+	if (TextureSize.X < SMALL_NUMBER || TextureSize.Y < SMALL_NUMBER)
+	{
+		WidgetBlueprint->ThumbnailImage = nullptr;
+	}
+	else
+	{
+		ThumbnailTexture->Rename(ThumbnailName, WidgetBlueprint, REN_NonTransactional | REN_DontCreateRedirectors);
+		WidgetBlueprint->ThumbnailImage = ThumbnailTexture;
+	}
+}
+
 #undef LOCTEXT_NAMESPACE
diff --git a/Engine/Source/Editor/UMGEditor/Private/WidgetBlueprintEditorUtils.h b/Engine/Source/Editor/UMGEditor/Private/WidgetBlueprintEditorUtils.h
--- a/Engine/Source/Editor/UMGEditor/Private/WidgetBlueprintEditorUtils.h
+++ b/Engine/Source/Editor/UMGEditor/Private/WidgetBlueprintEditorUtils.h
@@ -3,8 +3,9 @@
 #pragma once
 
 #include "CoreMinimal.h"
-#include "WidgetReference.h"
 #include "WidgetBlueprintEditor.h"
+#include "WidgetBlueprint.h"
+#include "WidgetReference.h"
 
 class FHittestGrid;
 class FMenuBuilder;
@@ -18,7 +19,14 @@ class SWindow;
 
 class FWidgetBlueprintEditorUtils
 {
+
 public:
+	struct FWidgetThumbnailProperties
+	{
+		FVector2D ScaledSize;
+		FVector2D Offset;
+	};
+
 	static bool VerifyWidgetRename(TSharedRef<class FWidgetBlueprintEditor> BlueprintEditor, FWidgetReference Widget, const FText& NewName, FText& OutErrorMessage);
 
 	static bool RenameWidget(TSharedRef<class FWidgetBlueprintEditor> BlueprintEditor, const FName& OldObjectName, const FString& NewDisplayName);
@@ -65,11 +73,21 @@ public:
 
 	static int32 UpdateHittestGrid(FHittestGrid& HitTestGrid, TSharedRef<SWindow> Window, float Scale, FVector2D DrawSize, float DeltaTime);
 
-	static TTuple<FVector2D, FVector2D> GetWidgetPreviewAreaAndSize(UUserWidget* UserWidget, FVector2D DesiredSize, FVector2D PreviewSize);
+	static TTuple<FVector2D, FVector2D> GetWidgetPreviewAreaAndSize(UUserWidget* UserWidget, FVector2D DesiredSize, FVector2D PreviewSize, EDesignPreviewSizeMode ThumbnailSizeMode, TOptional<FVector2D> ThumbnailCustomSize);
 
 	static float GetWidgetPreviewDPIScale(UUserWidget* UserWidget, FVector2D PreviewSize);
 
-	static FVector2D GetWidgetPreviewUnScaledCustomSize(FVector2D DesiredSize, UUserWidget* UserWidget);
+	static EDesignPreviewSizeMode ConvertThumbnailSizeModeToDesignerSizeMode(EThumbnailPreviewSizeMode ThumbnailSizeMode, UUserWidget* WidgetInstance);
+
+	static FVector2D GetWidgetPreviewUnScaledCustomSize(FVector2D DesiredSize, UUserWidget* UserWidget, TOptional<FVector2D> ThumbnailCustomSize, EThumbnailPreviewSizeMode ThumbnailSizeMode = EThumbnailPreviewSizeMode::MatchDesignerMode);
+
+	static TTuple<float, FVector2D> GetThumbnailImageScaleAndOffset(FVector2D WidgetSize, FVector2D ThumbnailSize);
+
+	static void SetTextureAsAssetThumbnail(UWidgetBlueprint* WidgetBlueprint, UTexture2D* ThumbnailTexture);
+
+	static TOptional<FWidgetThumbnailProperties> DrawSWidgetInRenderTargetForThumbnail(UUserWidget* WidgetInstance, UTextureRenderTarget2D* RenderTarget2D, FVector2D ThumbnailSize, TOptional<FVector2D> ThumbnailCustomSize, EThumbnailPreviewSizeMode ThumbnailSizeMode = EThumbnailPreviewSizeMode::MatchDesignerMode);
+
+	static TOptional<FWidgetThumbnailProperties> DrawSWidgetInRenderTarget(UUserWidget* WidgetInstance, UTextureRenderTarget2D* RenderTarget2D);
 
 private:
 
@@ -106,4 +124,6 @@ private:
 	static bool ShouldContinueDeleteOperation(UWidgetBlueprint* BP, const TArray<FText>& WidgetNames);
 
 	static bool ShouldContinueReplaceOperation(UWidgetBlueprint* BP, const TArray<FText>& WidgetNames);	
+
+	static TOptional<FWidgetThumbnailProperties> DrawSWidgetInRenderTargetInternal(UUserWidget* WidgetInstance, UTextureRenderTarget2D* RenderTarget2D, FVector2D ThumbnailSize, bool bIsForThumbnail, TOptional<FVector2D> ThumbnailCustomSize, EThumbnailPreviewSizeMode ThumbnailSizeMode);
 };
diff --git a/Engine/Source/Editor/UMGEditor/Private/WidgetBlueprintThumbnailRenderer.cpp b/Engine/Source/Editor/UMGEditor/Private/WidgetBlueprintThumbnailRenderer.cpp
--- a/Engine/Source/Editor/UMGEditor/Private/WidgetBlueprintThumbnailRenderer.cpp
+++ b/Engine/Source/Editor/UMGEditor/Private/WidgetBlueprintThumbnailRenderer.cpp
@@ -30,123 +30,87 @@ void UWidgetBlueprintThumbnailRenderer::Draw(UObject* Object, int32 X, int32 Y,
 		0.0f, 0.0f, SizeOfUV, SizeOfUV,
 		GrayBackgroundColor);
 
-
 	// Convert the object from UObject to SWidget
 	UWidgetBlueprint* WidgetBlueprintToRender = Cast<UWidgetBlueprint>(Object);
 
-	const bool bIsBlueprintValid = IsValid(WidgetBlueprintToRender)
-		&& IsValid(WidgetBlueprintToRender->GeneratedClass)
-		&& WidgetBlueprintToRender->bHasBeenRegenerated
-		//&& Blueprint->IsUpToDate() - This condition makes the thumbnail blank whenever the BP is dirty. It seems too strict.
-		&& !WidgetBlueprintToRender->bBeingCompiled
-		&& !WidgetBlueprintToRender->HasAnyFlags(RF_Transient);
-
-	TSharedPtr<SWidget> WindowContent;
-	UClass* ClassToGenerate = WidgetBlueprintToRender->GeneratedClass;
-	UUserWidget* WidgetInstance = nullptr;
-	if (bIsBlueprintValid && ClassToGenerate->IsChildOf(UWidget::StaticClass()))
+	// check if an image is used instead of auto generating the thumbnail
+	if (WidgetBlueprintToRender->ThumbnailImage)
 	{
-		WidgetInstance = NewObject<UUserWidget>(GetTransientPackage(), ClassToGenerate);
-		if (WidgetInstance)
-		{
-			WidgetInstance->Initialize();
-			WidgetInstance->SetDesignerFlags(EWidgetDesignFlags::Designing | EWidgetDesignFlags::ExecutePreConstruct);
-
-
-			WindowContent = WidgetInstance->TakeWidget();
+		FVector2D TextureSize(WidgetBlueprintToRender->ThumbnailImage->GetSizeX(), WidgetBlueprintToRender->ThumbnailImage->GetSizeY());
+		if (TextureSize.X > SMALL_NUMBER && TextureSize.Y > SMALL_NUMBER) {
+			TTuple<float, FVector2D> ScaleAndOffset = FWidgetBlueprintEditorUtils::GetThumbnailImageScaleAndOffset(TextureSize, FVector2D(Width, Height));
+			float Scale = ScaleAndOffset.Get<0>();
+			FVector2D Offset = ScaleAndOffset.Get<1>();
+			FVector2D ThumbnailImageOffset = Offset;
+			FVector2D ThumbnailImageScaledSize = Scale * TextureSize;
+
+			FCanvasTileItem CanvasTile(ThumbnailImageOffset, WidgetBlueprintToRender->ThumbnailImage->GetResource(), ThumbnailImageScaledSize, FLinearColor::White);
+			CanvasTile.BlendMode = SE_BLEND_Translucent;
+			CanvasTile.Draw(Canvas);
 		}
 	}
-
-	if (WindowContent == nullptr || WidgetInstance == nullptr)
-	{
-		return;
-	}
-
-
-	//Create Window
-	FVector2D ThumbnailSize(Width, Height);
-	TSharedRef<SVirtualWindow> Window = SNew(SVirtualWindow).Size(ThumbnailSize);
-	TUniquePtr<FHittestGrid> HitTestGrid = MakeUnique<FHittestGrid>();
-	Window->SetContent(WindowContent.ToSharedRef());
-	Window->Resize(ThumbnailSize);
-
-	// Store the desired size to maintain the aspect ratio later
-	FGeometry WindowGeometry = FGeometry::MakeRoot(ThumbnailSize, FSlateLayoutTransform(1.0f));
-	Window->SlatePrepass(1.0f);
-	FVector2D DesiredSizeWindow = Window->GetDesiredSize();
-
-
-	if (DesiredSizeWindow.X < SMALL_NUMBER || DesiredSizeWindow.Y < SMALL_NUMBER)
-	{
-		return;
-	}
-
-	FVector2D UnscaledSize = FWidgetBlueprintEditorUtils::GetWidgetPreviewUnScaledCustomSize(DesiredSizeWindow, WidgetInstance);
-	if (UnscaledSize.X < SMALL_NUMBER || UnscaledSize.Y < SMALL_NUMBER)
-	{
-		return;
-	}
-	TTuple<float, FVector2D> ScaleAndOffset = GetScaleAndOffset(UnscaledSize, ThumbnailSize);
-	float Scale = ScaleAndOffset.Get<0>();
-	FVector2D Offset = ScaleAndOffset.Get<1>();
-	FVector2D ScaledSize = UnscaledSize * Scale;
-
-	// Create Renderer Target and WidgetRenderer
-	bool bApplyGammaCorrection = false;
-	FWidgetRenderer* WidgetRenderer = new FWidgetRenderer(bApplyGammaCorrection);
-	const bool bIsLinearSpace = !bApplyGammaCorrection;
-
-	if (!RenderTarget2D)
+	else
 	{
-		RenderTarget2D = NewObject<UTextureRenderTarget2D>();
-		RenderTarget2D->Filter = TF_Bilinear;
-		RenderTarget2D->ClearColor = FLinearColor::Transparent;
-		RenderTarget2D->SRGB = true;
-		RenderTarget2D->TargetGamma = 1;
-	}
-
-	const EPixelFormat RequestedFormat = FSlateApplication::Get().GetRenderer()->GetSlateRecommendedColorFormat();
-	RenderTarget2D->InitCustomFormat(ScaledSize.X, ScaledSize.Y, RequestedFormat, true);
-	WidgetRenderer->DrawWindow(RenderTarget2D, *HitTestGrid, Window, Scale, ScaledSize, 0.1f);
-
-	// Draw the SWidget on canvas
-	FCanvasTileItem CanvasTile(FVector2D(X + Offset.X, Y + Offset.Y), RenderTarget2D->GetResource(), ScaledSize, FLinearColor::White);
-	CanvasTile.BlendMode = SE_BLEND_Translucent;
-	FlushRenderingCommands();
-	CanvasTile.Draw(Canvas);
+		const bool bIsBlueprintValid = IsValid(WidgetBlueprintToRender)
+			&& IsValid(WidgetBlueprintToRender->GeneratedClass)
+			&& WidgetBlueprintToRender->bHasBeenRegenerated
+			//&& Blueprint->IsUpToDate() - This condition makes the thumbnail blank whenever the BP is dirty. It seems too strict.
+			&& !WidgetBlueprintToRender->bBeingCompiled
+			&& !WidgetBlueprintToRender->HasAnyFlags(RF_Transient);
+
+		TSharedPtr<SWidget> WindowContent;
+		UClass* ClassToGenerate = WidgetBlueprintToRender->GeneratedClass;
+		UUserWidget* WidgetInstance = nullptr;
+		if (bIsBlueprintValid && ClassToGenerate->IsChildOf(UWidget::StaticClass()))
+		{
+			WidgetInstance = NewObject<UUserWidget>(GetTransientPackage(), ClassToGenerate);
+			if (WidgetInstance)
+			{
+				WidgetInstance->Initialize();
+				WidgetInstance->SetDesignerFlags(EWidgetDesignFlags::Designing | EWidgetDesignFlags::ExecutePreConstruct);
+			}
+		}
 
-	if (WidgetRenderer)
-	{
-		BeginCleanup(WidgetRenderer);
-		WidgetRenderer = nullptr;
-	}
-	
+		if (WidgetInstance == nullptr)
+		{
+			return;
+		}
+		FVector2D ThumbnailSize(Width, Height);
 
-}
+		if (!RenderTarget2D)
+		{
+			RenderTarget2D = NewObject<UTextureRenderTarget2D>();
+			RenderTarget2D->Filter = TF_Bilinear;
+			RenderTarget2D->ClearColor = FLinearColor::Transparent;
+			RenderTarget2D->SRGB = true;
+			RenderTarget2D->TargetGamma = 1;
+		}
 
-TTuple<float, FVector2D> UWidgetBlueprintThumbnailRenderer::GetScaleAndOffset(FVector2D WidgetSize, FVector2D ThumbnailSize) const
-{
-	// Scale the widget blueprint image to fit in the thumbnail
+		TOptional<FWidgetBlueprintEditorUtils::FWidgetThumbnailProperties> ScaleAndOffset;
 
-	checkf(WidgetSize.X > 0.f && WidgetSize.Y > 0.f, TEXT("The size should have been previously checked to be > 0."));
+		if (WidgetBlueprintToRender->ThumbnailSizeMode == EThumbnailPreviewSizeMode::Custom)
+		{
+			ScaleAndOffset = FWidgetBlueprintEditorUtils::DrawSWidgetInRenderTargetForThumbnail(WidgetInstance, RenderTarget2D, ThumbnailSize, WidgetBlueprintToRender->ThumbnailCustomSize, WidgetBlueprintToRender->ThumbnailSizeMode);
+		}
+		else
+		{
+			ScaleAndOffset = FWidgetBlueprintEditorUtils::DrawSWidgetInRenderTargetForThumbnail(WidgetInstance, RenderTarget2D, ThumbnailSize, TOptional<FVector2D>(), WidgetBlueprintToRender->ThumbnailSizeMode);
 
-	float Scale;
-	float XOffset = 0;
-	float YOffset = 0;
-	if (WidgetSize.X > WidgetSize.Y)
-	{
-		Scale = ThumbnailSize.X / WidgetSize.X;
-		WidgetSize *= Scale;
-		YOffset = (ThumbnailSize.Y - WidgetSize.Y) / 2.f;
-	}
-	else
-	{
-		Scale = ThumbnailSize.Y / WidgetSize.Y;
-		WidgetSize *= Scale;
-		XOffset = (ThumbnailSize.X - WidgetSize.X) / 2.f;
+		}
+		if (!ScaleAndOffset.IsSet())
+		{
+			return;
+		}
+		FVector2D Offset = ScaleAndOffset.GetValue().Offset;
+		FVector2D ScaledSize = ScaleAndOffset.GetValue().ScaledSize;
+
+		// Draw the SWidget on canvas
+		FCanvasTileItem CanvasTile(FVector2D(X + Offset.X, Y + Offset.Y), RenderTarget2D->GetResource(), ScaledSize, FLinearColor::White);
+		CanvasTile.BlendMode = SE_BLEND_Translucent;
+		FlushRenderingCommands();
+		CanvasTile.Draw(Canvas);
 	}
 
-	return TTuple<float, FVector2D>(Scale, FVector2D(XOffset, YOffset));
 }
 
 
diff --git a/Engine/Source/Editor/UMGEditor/Public/WidgetBlueprint.h b/Engine/Source/Editor/UMGEditor/Public/WidgetBlueprint.h
--- a/Engine/Source/Editor/UMGEditor/Public/WidgetBlueprint.h
+++ b/Engine/Source/Editor/UMGEditor/Public/WidgetBlueprint.h
@@ -201,6 +201,18 @@ enum class EWidgetSupportsDynamicCreation : uint8
 	No,
 };
 
+
+UENUM()
+enum class EThumbnailPreviewSizeMode : uint8
+{
+	MatchDesignerMode,
+	FillScreen,
+	Custom,
+	Desired
+};
+
+
+
 /**
  * This represents the tickability of a widget computed at compile time
  * It is designed as a hint so the runtime can determine if ticking needs to be enabled
@@ -308,7 +320,6 @@ public:
 	bool ArePropertyBindingsAllowed() const;
 
 protected:
-#if WITH_EDITOR
 	virtual void LoadModulesRequiredForCompilation() override;
 
 private:
@@ -338,5 +349,14 @@ public:
 	 */
 	UPROPERTY(AssetRegistrySearchable)
 	int32 PropertyBindings;
-#endif
+
+	UPROPERTY(EditDefaultsOnly, Category = ThumbnailSettings)
+	EThumbnailPreviewSizeMode ThumbnailSizeMode;
+
+	UPROPERTY(EditDefaultsOnly, Category = ThumbnailSettings)
+	FVector2D ThumbnailCustomSize;
+
+	UPROPERTY(EditDefaultsOnly, Category = ThumbnailSettings)
+	TObjectPtr<UTexture2D> ThumbnailImage;
+
 };
diff --git a/Engine/Source/Editor/UMGEditor/Public/WidgetBlueprintEditor.h b/Engine/Source/Editor/UMGEditor/Public/WidgetBlueprintEditor.h
--- a/Engine/Source/Editor/UMGEditor/Public/WidgetBlueprintEditor.h
+++ b/Engine/Source/Editor/UMGEditor/Public/WidgetBlueprintEditor.h
@@ -207,8 +207,15 @@ protected:
 	virtual void InitalizeExtenders() override;
 	TSharedPtr<FExtender> CreateMenuExtender();
 	void FillFileMenu(FMenuBuilder& MenuBuilder);
+	void FillAssetMenu(FMenuBuilder& MenuBuilder);
 	void BindToolkitCommands();
 
+	void TakeSnapshot();
+	void CaptureThumbnail();
+	void ClearThumbnail();
+	bool IsImageUsedAsThumbnail();
+	bool IsPreviewWidgetInitialized();
+
 	void OpenCreateNativeBaseClassDialog();
 	void OnCreateNativeBaseClassSuccessfully(const FString& InClassName, const FString& InClassPath, const FString& InModuleName);
 
diff --git a/Engine/Source/Editor/UMGEditor/Public/WidgetBlueprintThumbnailRenderer.h b/Engine/Source/Editor/UMGEditor/Public/WidgetBlueprintThumbnailRenderer.h
--- a/Engine/Source/Editor/UMGEditor/Public/WidgetBlueprintThumbnailRenderer.h
+++ b/Engine/Source/Editor/UMGEditor/Public/WidgetBlueprintThumbnailRenderer.h
@@ -20,7 +20,6 @@ class UMGEDITOR_API UWidgetBlueprintThumbnailRenderer : public UDefaultSizedThum
 	void Draw(UObject* Object, int32 X, int32 Y, uint32 Width, uint32 Height, FRenderTarget* RenderTarget, FCanvas* Canvas, bool bAdditionalViewFamily) override;
 
 private:
-	TTuple<float, FVector2D> GetScaleAndOffset(FVector2D DesiredSize, FVector2D WidgetSize) const;
 
 	UPROPERTY(Transient)
 	TObjectPtr<UTextureRenderTarget2D> RenderTarget2D;
diff --git a/Engine/Source/Editor/UMGEditor/UMGEditor.Build.cs b/Engine/Source/Editor/UMGEditor/UMGEditor.Build.cs
--- a/Engine/Source/Editor/UMGEditor/UMGEditor.Build.cs
+++ b/Engine/Source/Editor/UMGEditor/UMGEditor.Build.cs
@@ -72,6 +72,7 @@ public class UMGEditor : ModuleRules
 				"ToolMenus",
 				"SlateReflector",
 				"DeveloperSettings",
+				"ImageWrapper"
 			}
 			);
 	}
diff --git a/Engine/Source/Runtime/Engine/Private/ImageUtils.cpp b/Engine/Source/Runtime/Engine/Private/ImageUtils.cpp
--- a/Engine/Source/Runtime/Engine/Private/ImageUtils.cpp
+++ b/Engine/Source/Runtime/Engine/Private/ImageUtils.cpp
@@ -29,7 +29,7 @@ DEFINE_LOG_CATEGORY_STATIC(LogImageUtils, Log, All);
  * @param RawData - an array to be filled with pixel data.
  * @return true if RawData has been successfully filled.
  */
-static bool GetRawData(UTextureRenderTarget2D* TexRT, TArray64<uint8>& RawData)
+bool FImageUtils::GetRawData(UTextureRenderTarget2D* TexRT, TArray64<uint8>& RawData)
 {
 	FRenderTarget* RenderTarget = TexRT->GameThread_GetRenderTargetResource();
 	EPixelFormat Format = TexRT->GetFormat();
@@ -517,7 +517,7 @@ public:
 		Format = TexRT->GetFormat();
 
 		TArray64<uint8> RawData;
-		bool bReadSuccess = GetRawData(TexRT, RawData);
+		bool bReadSuccess = FImageUtils::GetRawData(TexRT, RawData);
 		if (bReadSuccess)
 		{
 			WriteHDRImage(RawData, Ar);
diff --git a/Engine/Source/Runtime/Engine/Public/ImageUtils.h b/Engine/Source/Runtime/Engine/Public/ImageUtils.h
--- a/Engine/Source/Runtime/Engine/Public/ImageUtils.h
+++ b/Engine/Source/Runtime/Engine/Public/ImageUtils.h
@@ -235,4 +235,6 @@ public:
 	*/
 	ENGINE_API static bool ExportTextureCubeAsHDR(UTextureCube* TexRT, FArchive& Ar);
 
+	ENGINE_API static bool GetRawData(UTextureRenderTarget2D* TexRT, TArray64<uint8>& RawData);
+
 };
diff --git a/Engine/Source/Runtime/ImageWrapper/Private/ImageWrapperHelper.cpp b/Engine/Source/Runtime/ImageWrapper/Private/ImageWrapperHelper.cpp
new file mode 100644
--- /dev/null
+++ b/Engine/Source/Runtime/ImageWrapper/Private/ImageWrapperHelper.cpp
@@ -0,0 +1,107 @@
+// Copyright Epic Games, Inc. All Rights Reserved.
+
+#include "ImageWrapperHelper.h"
+#include "Misc/Paths.h"
+
+FStringView ImageWrapperHelper::GetFormatExtension(EImageFormat InImageFormat, bool bIncludeDot)
+{
+	const TCHAR* StringExtension = nullptr;
+	switch (InImageFormat)
+	{
+	case EImageFormat::Invalid:
+		return StringExtension;
+	case EImageFormat::PNG:
+		StringExtension = TEXT(".png");
+		break;
+	case EImageFormat::JPEG:
+		StringExtension = TEXT(".jpeg");
+		break;
+	case EImageFormat::GrayscaleJPEG:
+		StringExtension = TEXT(".jpg");
+		break;
+	case EImageFormat::BMP:
+		StringExtension = TEXT(".bmp");
+		break;
+	case EImageFormat::ICO:
+		StringExtension = TEXT(".ico");
+		break;
+	case EImageFormat::EXR:
+		StringExtension = TEXT(".exr");
+		break;
+	case EImageFormat::ICNS:
+		StringExtension = TEXT(".icns");
+		break;
+	case EImageFormat::TGA:
+		StringExtension = TEXT(".tga");
+		break;
+	case EImageFormat::HDR:
+		StringExtension = TEXT(".hdr");
+		break;
+	default:
+		return StringExtension;
+	}
+	if (!bIncludeDot) 
+	{
+		++StringExtension;
+	}
+	return StringExtension;
+}
+
+EImageFormat ImageWrapperHelper::GetImageFormat(FStringView StringExtention)
+{
+	int32 Length = StringExtention.Len();
+	if (Length == 0)
+	{
+		return EImageFormat::Invalid;
+	}
+
+	if (StringExtention.GetData()[0] == TEXT('.')) 
+	{
+		StringExtention = StringExtention.SubStr(1, Length);
+	}
+
+	if (StringExtention.Equals(TEXT("png"), ESearchCase::IgnoreCase))
+	{
+		return EImageFormat::PNG;
+	}
+	if (StringExtention.Equals(TEXT("jpeg"), ESearchCase::IgnoreCase))
+	{
+		return EImageFormat::JPEG;
+	}
+	if (StringExtention.Equals(TEXT("jpg"), ESearchCase::IgnoreCase))
+	{
+		return EImageFormat::GrayscaleJPEG;
+	}
+	if (StringExtention.Equals(TEXT("bmp"), ESearchCase::IgnoreCase))
+	{
+		return EImageFormat::BMP;
+	}
+	if (StringExtention.Equals(TEXT("ico"), ESearchCase::IgnoreCase))
+	{
+		return EImageFormat::ICO;
+	}
+	if (StringExtention.Equals(TEXT("exr"), ESearchCase::IgnoreCase))
+	{
+		return EImageFormat::EXR;
+	}
+	if (StringExtention.Equals(TEXT("icns"), ESearchCase::IgnoreCase))
+	{
+		return EImageFormat::ICNS;
+	}
+	if (StringExtention.Equals(TEXT("tga"), ESearchCase::IgnoreCase))
+	{
+		return EImageFormat::TGA;
+	}
+	if (StringExtention.Equals(TEXT("hdr"), ESearchCase::IgnoreCase))
+	{
+		return EImageFormat::HDR;
+	}
+	return EImageFormat::Invalid;
+}
+
+const FStringView ImageWrapperHelper::GetImageFilesFilterString(bool bIncludeAllFiles)
+{
+	return (bIncludeAllFiles)
+		? TEXT("Image files (*.jpg; *.png; *.bmp; *.ico; *.exr; *.icns; *.jpeg; *.tga; *.hdr)|*.jpg; *.png; *.bmp; *.ico; *.exr; *.icns; *.jpeg; *.tga; *.hdr|All files (*.*)|*.*")
+		: TEXT("Image files (*.jpg; *.png; *.bmp; *.ico; *.exr; *.icns; *.jpeg; *.tga; *.hdr)|*.jpg; *.png; *.bmp; *.ico; *.exr; *.icns; *.jpeg; *.tga; *.hdr");
+}
\ No newline at end of file
diff --git a/Engine/Source/Runtime/ImageWrapper/Public/ImageWrapperHelper.h b/Engine/Source/Runtime/ImageWrapper/Public/ImageWrapperHelper.h
new file mode 100644
--- /dev/null
+++ b/Engine/Source/Runtime/ImageWrapper/Public/ImageWrapperHelper.h
@@ -0,0 +1,13 @@
+// Copyright Epic Games, Inc. All Rights Reserved.
+
+#pragma once
+#include "CoreMinimal.h"
+#include "IImageWrapper.h"
+
+class IMAGEWRAPPER_API ImageWrapperHelper
+{
+public:
+	 static FStringView GetFormatExtension(EImageFormat InImageFormat, bool bIncludeDot=false);
+	 static EImageFormat GetImageFormat(FStringView StringExtention);
+	 static const FStringView GetImageFilesFilterString(bool bIncludeAllFiles);
+};
