Correction to Replace Widget with child done on a NamedSlot to properly mark the parent as modify and have the change to it in the transaction.

diff --git a/Engine/Source/Editor/UMGEditor/Private/WidgetBlueprintEditorUtils.cpp b/Engine/Source/Editor/UMGEditor/Private/WidgetBlueprintEditorUtils.cpp
--- a/Engine/Source/Editor/UMGEditor/Private/WidgetBlueprintEditorUtils.cpp
+++ b/Engine/Source/Editor/UMGEditor/Private/WidgetBlueprintEditorUtils.cpp
@@ -688,6 +688,13 @@ bool FWidgetBlueprintEditorUtils::ReplaceNamedSlotHostContent(UWidget* WidgetTem
 		{
 			if (SlotContent == WidgetTemplate)
 			{
+				NamedSlotHost.GetObject()->Modify();
+				if (UPanelWidget* NamedSlot = WidgetTemplate->GetParent())
+				{
+					// Make sure we also mark the named slot as modified to properly track changes in it.
+					NamedSlot->Modify();
+				}
+
 				if (NewContentWidget)
 				{
 					NewContentWidget->Modify();
@@ -697,7 +704,6 @@ bool FWidgetBlueprintEditorUtils::ReplaceNamedSlotHostContent(UWidget* WidgetTem
 						NewContentWidget->RemoveFromParent();
 					}
 				}
-				NamedSlotHost.GetObject()->Modify();
 				NamedSlotHost->SetContentForSlot(SlotName, NewContentWidget);
 				return true;
 			}
