Fixed issue with UMG details view causing warnings on BP compilation due to thumbnails + lazy object ptr.
Replaced lazy ptr by soft object ptr which also supersedes the previously existing soft object path.

diff --git a/Engine/Source/Editor/UMGEditor/Private/Components/DetailsView.cpp b/Engine/Source/Editor/UMGEditor/Private/Components/DetailsView.cpp
--- a/Engine/Source/Editor/UMGEditor/Private/Components/DetailsView.cpp
+++ b/Engine/Source/Editor/UMGEditor/Private/Components/DetailsView.cpp
@@ -42,8 +42,8 @@ void UDetailsView::BuildContentWidget()
 		UObject* ViewedObject = GetObject();
 		if (ViewedObject == nullptr)
 		{
-			bool bIsLazyObjectNull = LazyObject.IsNull();
-			if (bIsLazyObjectNull)
+			bool bIsObjectNull = Object.IsNull();
+			if (bIsObjectNull)
 			{
 				MissingWidgetText = FPropertyViewHelper::UndefinedObjectText;
 			}
@@ -143,7 +143,6 @@ void UDetailsView::PostEditChangeProperty(FPropertyChangedEvent& PropertyChanged
 			|| PropertyChangedEvent.GetPropertyName() == GET_MEMBER_NAME_CHECKED(UDetailsView, bForceHiddenPropertyVisibility)
 			|| PropertyChangedEvent.GetPropertyName() == GET_MEMBER_NAME_CHECKED(UDetailsView, ColumnWidth))
 		{
-			SoftObjectPath = LazyObject.Get();
 			BuildContentWidget();
 		}
 	}
diff --git a/Engine/Source/Editor/UMGEditor/Private/Components/PropertyViewBase.cpp b/Engine/Source/Editor/UMGEditor/Private/Components/PropertyViewBase.cpp
--- a/Engine/Source/Editor/UMGEditor/Private/Components/PropertyViewBase.cpp
+++ b/Engine/Source/Editor/UMGEditor/Private/Components/PropertyViewBase.cpp
@@ -59,16 +59,15 @@ TSharedRef<SWidget> UPropertyViewBase::RebuildWidget()
 
 UObject* UPropertyViewBase::GetObject() const
 {
-	return LazyObject.Get();
+	return Object.Get();
 }
 
 
 void UPropertyViewBase::SetObject(UObject* InObject)
 {
-	if (LazyObject.Get() != InObject)
+	if (Object.Get() != InObject)
 	{
-		LazyObject = InObject;
-		SoftObjectPath = InObject;
+		Object = InObject;
 		OnObjectChanged();
 	}
 }
@@ -82,7 +81,7 @@ void UPropertyViewBase::OnPropertyChangedBroadcast(FName PropertyName)
 
 void UPropertyViewBase::InternalOnAssetLoaded(UObject* AssetLoaded)
 {
-	if (SoftObjectPath.GetAssetPathName() == FSoftObjectPath(AssetLoaded).GetAssetPathName())
+	if(Object.ToSoftObjectPath().GetAssetPathName() == FSoftObjectPath(AssetLoaded).GetAssetPathName())
 	{
 		BuildContentWidget();
 	}
@@ -105,9 +104,14 @@ void UPropertyViewBase::PostLoad()
 {
 	Super::PostLoad();
 
-	if (!LazyObject.IsValid() && SoftObjectPath.IsAsset() && bAutoLoadAsset && !HasAnyFlags(RF_BeginDestroyed))
+	if (Object.IsNull() && !SoftObjectPath_DEPRECATED.IsNull())
 	{
-		LazyObject = SoftObjectPath.TryLoad();
+		Object = SoftObjectPath_DEPRECATED;
+	}
+
+	if(!Object.IsValid() && Object.ToSoftObjectPath().IsAsset() && bAutoLoadAsset && !HasAnyFlags(RF_BeginDestroyed))
+	{
+		Object.LoadSynchronous();
 		BuildContentWidget();
 	}
 }
@@ -115,9 +119,8 @@ void UPropertyViewBase::PostLoad()
 
 void UPropertyViewBase::PostEditChangeProperty(FPropertyChangedEvent& PropertyChangedEvent)
 {
-	if (PropertyChangedEvent.GetPropertyName() == GET_MEMBER_NAME_CHECKED(UPropertyViewBase, LazyObject))
+	if (PropertyChangedEvent.GetPropertyName() == GET_MEMBER_NAME_CHECKED(UPropertyViewBase, Object))
 	{
-		SoftObjectPath = LazyObject.Get();
 		OnObjectChanged();
 	}
 
diff --git a/Engine/Source/Editor/UMGEditor/Private/Components/SinglePropertyView.cpp b/Engine/Source/Editor/UMGEditor/Private/Components/SinglePropertyView.cpp
--- a/Engine/Source/Editor/UMGEditor/Private/Components/SinglePropertyView.cpp
+++ b/Engine/Source/Editor/UMGEditor/Private/Components/SinglePropertyView.cpp
@@ -42,8 +42,8 @@ void USinglePropertyView::BuildContentWidget()
 		UObject* ViewedObject = GetObject();
 		if (ViewedObject == nullptr)
 		{
-			bool bIsLazyObjectNull = LazyObject.IsNull();
-			if (bIsLazyObjectNull)
+			bool bIsObjectNull = Object.IsNull();
+			if (bIsObjectNull)
 			{
 				MissingWidgetText = FPropertyViewHelper::UndefinedObjectText;
 			}
diff --git a/Engine/Source/Editor/UMGEditor/Public/Components/PropertyViewBase.h b/Engine/Source/Editor/UMGEditor/Public/Components/PropertyViewBase.h
--- a/Engine/Source/Editor/UMGEditor/Public/Components/PropertyViewBase.h
+++ b/Engine/Source/Editor/UMGEditor/Public/Components/PropertyViewBase.h
@@ -5,8 +5,8 @@
 #include "CoreMinimal.h"
 #include "Components/Widget.h"
 #include "Widgets/SWidget.h"
-#include "UObject/LazyObjectPtr.h"
 #include "UObject/SoftObjectPath.h"
+#include "UObject/SoftObjectPtr.h"
 #include "PropertyViewBase.generated.h"
 
 class SBorder;
@@ -26,10 +26,10 @@ class UMGEDITOR_API UPropertyViewBase : public UWidget
 protected:
 	/** The object to view. */
 	UPROPERTY(meta = (DisplayName="Object"))
-	TLazyObjectPtr<UObject> LazyObject;
+	TSoftObjectPtr<UObject> Object;
 
 	UPROPERTY()
-	FSoftObjectPath SoftObjectPath;
+	FSoftObjectPath SoftObjectPath_DEPRECATED;
 
 	/** Load the object (if it's an asset) when the widget is created. */
 	UPROPERTY(EditAnywhere, AdvancedDisplay, Category = "View")
