Fix zero scale size bug for thumbnails in Linux #RB patrick.boutot #jira ue-121018

diff --git a/Engine/Source/Editor/UMGEditor/Private/WidgetBlueprintThumbnailRenderer.cpp b/Engine/Source/Editor/UMGEditor/Private/WidgetBlueprintThumbnailRenderer.cpp
--- a/Engine/Source/Editor/UMGEditor/Private/WidgetBlueprintThumbnailRenderer.cpp
+++ b/Engine/Source/Editor/UMGEditor/Private/WidgetBlueprintThumbnailRenderer.cpp
@@ -22,6 +22,18 @@ bool UWidgetBlueprintThumbnailRenderer::CanVisualizeAsset(UObject* Object)
 
 void UWidgetBlueprintThumbnailRenderer::Draw(UObject* Object, int32 X, int32 Y, uint32 Width, uint32 Height, FRenderTarget* RenderTarget, FCanvas* Canvas, bool bAdditionalViewFamily)
 {
+#if !UE_SERVER
+
+	if (Width < SMALL_NUMBER || Height < SMALL_NUMBER)
+	{
+		return;
+	}
+
+	if (!FApp::CanEverRender())
+	{
+		return;
+	}
+
 	// Create a plain gray background for the thumbnail
 	const int32 SizeOfUV = 1;
 	FLinearColor GrayBackgroundColor(FVector4(.03f, .03f, .03f, 1.f));
@@ -110,7 +122,7 @@ void UWidgetBlueprintThumbnailRenderer::Draw(UObject* Object, int32 X, int32 Y,
 		FlushRenderingCommands();
 		CanvasTile.Draw(Canvas);
 	}
-
+#endif
 }
 
 
