Disable property value input when it has bindings set in UMG details panel

diff --git a/Engine/Source/Editor/Persona/Private/AnimGraphNodeDetails.cpp b/Engine/Source/Editor/Persona/Private/AnimGraphNodeDetails.cpp
--- a/Engine/Source/Editor/Persona/Private/AnimGraphNodeDetails.cpp
+++ b/Engine/Source/Editor/Persona/Private/AnimGraphNodeDetails.cpp
@@ -1196,7 +1196,7 @@ bool FAnimGraphNodeBindingExtension::IsPropertyExtendable(const UClass* InObject
 	return false;
 }
 
-TSharedRef<SWidget> FAnimGraphNodeBindingExtension::GenerateExtensionWidget(const IDetailLayoutBuilder& InDetailBuilder, const UClass* InObjectClass, TSharedPtr<IPropertyHandle> InPropertyHandle)
+void FAnimGraphNodeBindingExtension::ExtendWidgetRow(FDetailWidgetRow& InWidgetRow, const IDetailLayoutBuilder& InDetailBuilder, const UClass* InObjectClass, TSharedPtr<IPropertyHandle> InPropertyHandle)
 {
 	int32 OptionalPinIndex;
 	UAnimGraphNode_Base* AnimGraphNode;
@@ -1234,7 +1234,10 @@ TSharedRef<SWidget> FAnimGraphNodeBindingExtension::GenerateExtensionWidget(cons
 	UAnimGraphNode_Base::FAnimPropertyBindingWidgetArgs Args(AnimGraphNodes, AnimNodeProperty, PropertyName, OptionalPinIndex);
 	Args.bOnGraphNode = false;
 	
-	return UAnimGraphNode_Base::MakePropertyBindingWidget(Args);
+	InWidgetRow.ExtensionContent()
+	[
+		UAnimGraphNode_Base::MakePropertyBindingWidget(Args)
+	];
 }
 
 #undef LOCTEXT_NAMESPACE
diff --git a/Engine/Source/Editor/Persona/Private/AnimGraphNodeDetails.h b/Engine/Source/Editor/Persona/Private/AnimGraphNodeDetails.h
--- a/Engine/Source/Editor/Persona/Private/AnimGraphNodeDetails.h
+++ b/Engine/Source/Editor/Persona/Private/AnimGraphNodeDetails.h
@@ -29,7 +29,7 @@ class FAnimGraphNodeBindingExtension : public IDetailPropertyExtensionHandler
 public:
 	// IDetailPropertyExtensionHandler interface
 	virtual bool IsPropertyExtendable(const UClass* InObjectClass, const IPropertyHandle& PropertyHandle) const override;
-	virtual TSharedRef<SWidget> GenerateExtensionWidget(const IDetailLayoutBuilder& InDetailBuilder, const UClass* InObjectClass, TSharedPtr<IPropertyHandle> PropertyHandle) override;
+	virtual void ExtendWidgetRow(FDetailWidgetRow& InWidgetRow, const IDetailLayoutBuilder& InDetailBuilder, const UClass* InObjectClass, TSharedPtr<IPropertyHandle> PropertyHandle) override;
 
 private:
 	// Helper function
diff --git a/Engine/Source/Editor/PropertyEditor/Private/SDetailSingleItemRow.cpp b/Engine/Source/Editor/PropertyEditor/Private/SDetailSingleItemRow.cpp
--- a/Engine/Source/Editor/PropertyEditor/Private/SDetailSingleItemRow.cpp
+++ b/Engine/Source/Editor/PropertyEditor/Private/SDetailSingleItemRow.cpp
@@ -264,6 +264,9 @@ void SDetailSingleItemRow::Construct( const FArguments& InArgs, FDetailLayoutCus
 		{
 			WidgetRow = Customization->GetWidgetRow();
 
+			// Populate the extension content in the WidgetRow if there's an extension handler.
+			PopulateExtensionWidget();
+
 			TSharedPtr<SWidget> NameWidget = WidgetRow.NameWidget.Widget;
 
 			TSharedPtr<SWidget> ValueWidget =
@@ -274,6 +277,8 @@ void SDetailSingleItemRow::Construct( const FArguments& InArgs, FDetailLayoutCus
 					WidgetRow.ValueWidget.Widget
 				];
 
+			TSharedPtr<SWidget> ExtensionWidget = WidgetRow.ExtensionWidget.Widget;
+
 			TAttribute<bool> IsEnabledAttribute;
 			if (WidgetRow.IsEnabledAttr.IsSet() || WidgetRow.IsEnabledAttr.IsBound())
 			{
@@ -290,8 +295,22 @@ void SDetailSingleItemRow::Construct( const FArguments& InArgs, FDetailLayoutCus
 				IsEnabledAttribute = InOwnerTreeNode->IsPropertyEditingEnabled();
 			}
 
+			TAttribute<bool> IsValueWidgetEnabledAttribute;
+			if (WidgetRow.IsValueEnabledAttr.IsSet() || WidgetRow.IsValueEnabledAttr.IsBound()) {
+				TAttribute<bool> ValueEnabledAttr = WidgetRow.IsValueEnabledAttr;
+				IsValueWidgetEnabledAttribute = TAttribute<bool>::Create(
+					[IsEnabledAttribute, ValueEnabledAttr]()
+					{
+						return IsEnabledAttribute.Get() && ValueEnabledAttr.Get();
+					});
+			}
+			else {
+				IsValueWidgetEnabledAttribute = IsEnabledAttribute;
+			}
+
 			NameWidget->SetEnabled(IsEnabledAttribute);
-			ValueWidget->SetEnabled(IsEnabledAttribute);
+			ValueWidget->SetEnabled(IsValueWidgetEnabledAttribute);
+			ExtensionWidget->SetEnabled(IsEnabledAttribute);
 
 			TSharedRef<SSplitter> Splitter = SNew(SSplitter)
 					.Style(FEditorStyle::Get(), "DetailsView.Splitter")
@@ -396,9 +415,6 @@ void SDetailSingleItemRow::Construct( const FArguments& InArgs, FDetailLayoutCus
 						NameColumnBox
 					];
 
-				TSharedRef<SWidget> ExtensionWidget = CreateExtensionWidget();
-				ExtensionWidget->SetEnabled(IsEnabledAttribute);
-
 				// create Value column:
 				// | Name | Value | Right |
 				Splitter->AddSlot()
@@ -416,12 +432,12 @@ void SDetailSingleItemRow::Construct( const FArguments& InArgs, FDetailLayoutCus
 						]
 						// extension widget
 						+ SHorizontalBox::Slot()
-						.HAlign(HAlign_Right)
-						.VAlign(VAlign_Center)
+						.HAlign(WidgetRow.ExtensionWidget.HorizontalAlignment)
+						.VAlign(WidgetRow.ExtensionWidget.VerticalAlignment)
 						.Padding(5,0,0,0)
 						.AutoWidth()
 						[
-							ExtensionWidget
+							ExtensionWidget.ToSharedRef()
 						]
 					];
 			}
@@ -951,10 +967,8 @@ bool SDetailSingleItemRow::IsBlacklistChecked() const
 	return false;
 }
 
-TSharedRef<SWidget> SDetailSingleItemRow::CreateExtensionWidget() const
+void SDetailSingleItemRow::PopulateExtensionWidget()
 {
-	TSharedPtr<SWidget> ExtensionWidget = SNullWidget::NullWidget;
-
 	TSharedPtr<FDetailTreeNode> OwnerTreeNodePinned = OwnerTreeNode.Pin();
 	if (OwnerTreeNodePinned.IsValid())
 	{ 
@@ -967,12 +981,10 @@ TSharedRef<SWidget> SDetailSingleItemRow::CreateExtensionWidget() const
 			if (Handle->IsValidHandle() && ObjectClass && ExtensionHandler->IsPropertyExtendable(ObjectClass, *Handle))
 			{
 				FDetailLayoutBuilderImpl& DetailLayout = OwnerTreeNodePinned->GetParentCategory()->GetParentLayoutImpl();
-				ExtensionWidget = ExtensionHandler->GenerateExtensionWidget(DetailLayout, ObjectClass, Handle);
+				ExtensionHandler->ExtendWidgetRow(WidgetRow, DetailLayout, ObjectClass, Handle);
 			}
 		}
 	}
-
-	return ExtensionWidget.ToSharedRef();
 }
 
 bool SDetailSingleItemRow::CanFavorite() const
diff --git a/Engine/Source/Editor/PropertyEditor/Private/SDetailSingleItemRow.h b/Engine/Source/Editor/PropertyEditor/Private/SDetailSingleItemRow.h
--- a/Engine/Source/Editor/PropertyEditor/Private/SDetailSingleItemRow.h
+++ b/Engine/Source/Editor/PropertyEditor/Private/SDetailSingleItemRow.h
@@ -75,7 +75,7 @@ private:
 	FSlateColor GetInnerBackgroundColor() const;
 
 	void CreateGlobalExtensionWidgets(TArray<FPropertyRowExtensionButton>& ExtensionButtons) const;
-	TSharedRef<SWidget> CreateExtensionWidget() const;
+	void PopulateExtensionWidget();
 	void OnResetToDefaultClicked() const;
 	bool IsResetToDefaultEnabled() const;
 
diff --git a/Engine/Source/Editor/PropertyEditor/Public/DetailWidgetRow.h b/Engine/Source/Editor/PropertyEditor/Public/DetailWidgetRow.h
--- a/Engine/Source/Editor/PropertyEditor/Public/DetailWidgetRow.h
+++ b/Engine/Source/Editor/PropertyEditor/Public/DetailWidgetRow.h
@@ -109,6 +109,7 @@ public:
 	FDetailWidgetRow()
 		: NameWidget( *this, 0.0f, 0.0f, HAlign_Left, VAlign_Center )
 		, ValueWidget( *this, DefaultValueMinWidth, DefaultValueMaxWidth, HAlign_Left, VAlign_Fill )
+		, ExtensionWidget( *this, 0.0f, 0.0f, HAlign_Right, VAlign_Center)
 		, WholeRowWidget( *this, 0.0f, 0.0f, HAlign_Fill, VAlign_Fill )
 		, VisibilityAttr( EVisibility::Visible )
 		, IsEnabledAttr( true )
@@ -122,6 +123,7 @@ public:
 	{
 		NameWidget = FDetailWidgetDecl(*this, Other.NameWidget);
 		ValueWidget = FDetailWidgetDecl(*this, Other.ValueWidget);
+		ExtensionWidget = FDetailWidgetDecl(*this, Other.ExtensionWidget);
 		WholeRowWidget = FDetailWidgetDecl(*this, Other.WholeRowWidget);
 		VisibilityAttr = Other.VisibilityAttr;
 		IsEnabledAttr = Other.IsEnabledAttr;
@@ -175,6 +177,14 @@ public:
 		return ValueWidget;
 	}
 
+	/**
+     * Assigns content to the extension (right) slot
+	 */
+	FDetailWidgetDecl& ExtensionContent()
+	{
+		return ExtensionWidget;
+	}
+
 	/**
 	 * Sets a string which should be used to filter the content when a user searches
 	 */
@@ -202,6 +212,15 @@ public:
 		return *this;
 	}
 
+    /**
+     * Sets the enabled state of the value widget only
+     */
+	FDetailWidgetRow& IsValueEnabled( const TAttribute<bool>& InIsEnabled )
+	{
+		IsValueEnabledAttr = InIsEnabled;
+		return *this;
+	}
+
 	/**
 	 * Sets a custom copy action to take when copying the data from this row
 	 */
@@ -239,6 +258,11 @@ public:
 		return ValueWidget.Widget->GetType() != InvalidDetailWidgetName;
 	}
 
+	bool HasExtensionContent() const
+	{
+		return ExtensionWidget.Widget->GetType() != InvalidDetailWidgetName;
+	}
+
 	/**
 	 * @return true if the row has columns, false if it spans the entire row
 	 */
@@ -308,12 +332,16 @@ public:
 	FDetailWidgetDecl NameWidget;
 	/** Value column content */
 	FDetailWidgetDecl ValueWidget;
+	/** Extension (right) column content */
+	FDetailWidgetDecl ExtensionWidget;
 	/** Whole row content */
 	FDetailWidgetDecl WholeRowWidget;
 	/** Visibility of the row */
 	TAttribute<EVisibility> VisibilityAttr;
 	/** IsEnabled of the row */
 	TAttribute<bool> IsEnabledAttr;
+    /** IsEnabled of the value widget only */
+	TAttribute<bool> IsValueEnabledAttr;
 	/** String to filter with */
 	FText FilterTextString;
 	/** Action for coping data on this row */
diff --git a/Engine/Source/Editor/PropertyEditor/Public/IDetailPropertyExtensionHandler.h b/Engine/Source/Editor/PropertyEditor/Public/IDetailPropertyExtensionHandler.h
--- a/Engine/Source/Editor/PropertyEditor/Public/IDetailPropertyExtensionHandler.h
+++ b/Engine/Source/Editor/PropertyEditor/Public/IDetailPropertyExtensionHandler.h
@@ -3,6 +3,7 @@
 #pragma once
 
 #include "CoreMinimal.h"
+#include "DetailWidgetRow.h"
 #include "Widgets/SWidget.h"
 
 class IPropertyHandle;
@@ -15,12 +16,13 @@ public:
 
 	virtual bool IsPropertyExtendable(const UClass* InObjectClass, const class IPropertyHandle& PropertyHandle) const = 0;
 
-	UE_DEPRECATED(4.24, "Please use the overload that takes a IDetailLayoutBuilder")
+	UE_DEPRECATED(4.24, "Please use ExtendWidgetRow")
 	virtual TSharedRef<SWidget> GenerateExtensionWidget(const UClass* InObjectClass, TSharedPtr<IPropertyHandle> PropertyHandle) 
 	{ 
 		return SNullWidget::NullWidget;
 	}
 
+	UE_DEPRECATED(5.0, "Please use ExtendWidgetRow")
 	virtual TSharedRef<SWidget> GenerateExtensionWidget(const IDetailLayoutBuilder& InDetailBuilder, const UClass* InObjectClass, TSharedPtr<IPropertyHandle> PropertyHandle)
 	{
 	PRAGMA_DISABLE_DEPRECATION_WARNINGS
@@ -28,4 +30,24 @@ public:
 		return GenerateExtensionWidget(InObjectClass, PropertyHandle);
 	PRAGMA_ENABLE_DEPRECATION_WARNINGS
 	}
+
+	/**
+	 * Gives the extension handler a chance to add extension widgets to the widget row.
+	 * Typically, an extension handler can do so by using InWidgetRow.ExtensionContent(),
+	 * but it can also modify other aspects of the widget row as well.
+	 */
+	virtual void ExtendWidgetRow(FDetailWidgetRow& InWidgetRow, const IDetailLayoutBuilder& InDetailBuilder, const UClass* InObjectClass, TSharedPtr<IPropertyHandle> PropertyHandle)
+	{
+	PRAGMA_DISABLE_DEPRECATION_WARNINGS
+		// Call old deprecated path for back-compat
+		TSharedRef<SWidget> ExtensionWidget = GenerateExtensionWidget(InDetailBuilder, InObjectClass, PropertyHandle);
+	PRAGMA_ENABLE_DEPRECATION_WARNINGS
+
+		if (ExtensionWidget != SNullWidget::NullWidget)
+		{
+			InWidgetRow.ExtensionContent()[
+				ExtensionWidget
+			];
+		}
+	}
 };
diff --git a/Engine/Source/Editor/SequenceRecorder/Private/SClassPropertyRecorderSettings.cpp b/Engine/Source/Editor/SequenceRecorder/Private/SClassPropertyRecorderSettings.cpp
--- a/Engine/Source/Editor/SequenceRecorder/Private/SClassPropertyRecorderSettings.cpp
+++ b/Engine/Source/Editor/SequenceRecorder/Private/SClassPropertyRecorderSettings.cpp
@@ -146,7 +146,7 @@ bool SClassPropertyRecorderSettings::IsPropertyExtendable(const UClass* InObject
 	return PropertyHandle.GetProperty()->HasAnyPropertyFlags(CPF_Interp) && FMovieSceneMultiPropertyRecorder::CanPropertyBeRecorded(*PropertyHandle.GetProperty());
 }
 
-TSharedRef<SWidget> SClassPropertyRecorderSettings::GenerateExtensionWidget(const IDetailLayoutBuilder& InDetailLayoutBuilder, const UClass* InObjectClass, TSharedPtr<IPropertyHandle> PropertyHandle)
+void SClassPropertyRecorderSettings::ExtendWidgetRow(FDetailWidgetRow& InWidgetRow, const IDetailLayoutBuilder& InDetailLayoutBuilder, const UClass* InObjectClass, TSharedPtr<IPropertyHandle> PropertyHandle)
 {
 	ECheckBoxState InitialState = ECheckBoxState::Unchecked;
 	FName PropertyName = *PropertyHandle->GeneratePathToProperty();
@@ -157,9 +157,12 @@ TSharedRef<SWidget> SClassPropertyRecorderSettings::GenerateExtensionWidget(cons
 		InitialState = PropertiesToRecord->Contains(PropertyName) ? ECheckBoxState::Checked : ECheckBoxState::Unchecked;
 	}
 
-	return SNew(SCheckBox)
-		.OnCheckStateChanged(this, &SClassPropertyRecorderSettings::HandlePropertyCheckStateChanged, PropertyHandle)
-		.IsChecked(InitialState);
+	InWidgetRow.ExtensionContent()
+	[
+		SNew(SCheckBox)
+			.OnCheckStateChanged(this, &SClassPropertyRecorderSettings::HandlePropertyCheckStateChanged, PropertyHandle)
+			.IsChecked(InitialState)
+	];
 }
 
 void SClassPropertyRecorderSettings::HandlePropertyCheckStateChanged(ECheckBoxState InState, TSharedPtr<IPropertyHandle> PropertyHandle)
diff --git a/Engine/Source/Editor/SequenceRecorder/Private/SClassPropertyRecorderSettings.h b/Engine/Source/Editor/SequenceRecorder/Private/SClassPropertyRecorderSettings.h
--- a/Engine/Source/Editor/SequenceRecorder/Private/SClassPropertyRecorderSettings.h
+++ b/Engine/Source/Editor/SequenceRecorder/Private/SClassPropertyRecorderSettings.h
@@ -41,7 +41,7 @@ private:
 
 	/** IDetailPropertyExtensionHandler interface */
 	virtual bool IsPropertyExtendable(const UClass* InObjectClass, const class IPropertyHandle& PropertyHandle) const override;
-	virtual TSharedRef<SWidget> GenerateExtensionWidget(const IDetailLayoutBuilder& InDetailLayoutBuilder, const UClass* InObjectClass, TSharedPtr<IPropertyHandle> PropertyHandle) override;
+	virtual void ExtendWidgetRow(FDetailWidgetRow& InWidgetRow, const IDetailLayoutBuilder& InDetailLayoutBuilder, const UClass* InObjectClass, TSharedPtr<IPropertyHandle> PropertyHandle) override;
 
 private:
 	/** Handle to the property we are editing */
diff --git a/Engine/Source/Editor/UMGEditor/Private/Customizations/UMGDetailCustomizations.cpp b/Engine/Source/Editor/UMGEditor/Private/Customizations/UMGDetailCustomizations.cpp
--- a/Engine/Source/Editor/UMGEditor/Private/Customizations/UMGDetailCustomizations.cpp
+++ b/Engine/Source/Editor/UMGEditor/Private/Customizations/UMGDetailCustomizations.cpp
@@ -512,6 +512,48 @@ TSharedRef<SWidget> FBlueprintWidgetCustomization::MakePropertyBindingWidget(TWe
 	}
 }
 
+bool FBlueprintWidgetCustomization::HasPropertyBindings(TWeakPtr<FWidgetBlueprintEditor> InEditor, const TSharedRef<IPropertyHandle>& InPropertyHandle)
+{
+	TArray<UObject*> Objects;
+	InPropertyHandle->GetOuterObjects(Objects);
+
+	UWidgetBlueprint* ThisBlueprint = InEditor.Pin()->GetWidgetBlueprintObj();
+
+	// In the UI, we treat a child property of a struct/array as bound if the parent is bound
+	// E.g., we'll disable the child value widgets as well.
+	// So find the parent property and check if it's bound below.
+	FName ParentPropertyName;
+	for (TSharedPtr<const IPropertyHandle> CurrentPropertyHandle = InPropertyHandle;
+		CurrentPropertyHandle && CurrentPropertyHandle->GetProperty();
+		CurrentPropertyHandle = CurrentPropertyHandle->GetParentHandle())
+	{
+		ParentPropertyName = CurrentPropertyHandle->GetProperty()->GetFName();
+	}
+
+	//TODO UMG O(N) Isn't good for this, needs to be map, but map isn't serialized, need cached runtime map for fast lookups.
+
+	for (const UObject* Object : Objects)
+	{
+		// Ignore null outer objects
+		if (Object == nullptr)
+		{
+			continue;
+		}
+
+		for (const FDelegateEditorBinding& Binding : ThisBlueprint->Bindings)
+		{
+			if (Binding.ObjectName == Object->GetName() && Binding.PropertyName == ParentPropertyName)
+			{
+				return true;
+			}
+		}
+
+		break;
+	}
+
+	return false;
+}
+
 void FBlueprintWidgetCustomization::CreateEventCustomization( IDetailLayoutBuilder& DetailLayout, FDelegateProperty* Property, UWidget* Widget )
 {
 	TSharedRef<IPropertyHandle> DelegatePropertyHandle = DetailLayout.GetProperty(Property->GetFName(), Property->GetOwnerChecked<UClass>());
@@ -809,6 +851,13 @@ void FBlueprintWidgetCustomization::CustomizeAccessibilityProperty(IDetailLayout
 	// Make sure the old AccessibleText properties are hidden so we don't get duplicate widgets
 	DetailLayout.HideProperty(AccessibleTextPropertyHandle);
 
+	TSharedRef<SWidget> ValueWidget = AccessibleTextPropertyHandle->CreatePropertyValueWidget();
+	TWeakPtr<FWidgetBlueprintEditor> ThisEditor = Editor;
+	ValueWidget->SetEnabled(TAttribute<bool>::Create(
+		[ThisEditor, AccessibleTextPropertyHandle]() {
+			return !HasPropertyBindings(ThisEditor, AccessibleTextPropertyHandle);
+		}));
+
 	TSharedRef<SWidget> BindingWidget = MakePropertyBindingWidget(Editor, AccessibleTextDelegateProperty, AccessibleTextPropertyHandle, false);
 	TSharedRef<SHorizontalBox> CustomTextLayout = SNew(SHorizontalBox)
 	.Visibility(TAttribute<EVisibility>::Create([AccessibleBehaviorPropertyHandle]() -> EVisibility
@@ -820,7 +869,7 @@ void FBlueprintWidgetCustomization::CustomizeAccessibilityProperty(IDetailLayout
 	+ SHorizontalBox::Slot()
 	.Padding(FMargin(4.0f, 0.0f))
 	[
-		AccessibleTextPropertyHandle->CreatePropertyValueWidget()
+		ValueWidget
 	]
 	+SHorizontalBox::Slot()
 	.AutoWidth()
diff --git a/Engine/Source/Editor/UMGEditor/Private/Customizations/UMGDetailCustomizations.h b/Engine/Source/Editor/UMGEditor/Private/Customizations/UMGDetailCustomizations.h
--- a/Engine/Source/Editor/UMGEditor/Private/Customizations/UMGDetailCustomizations.h
+++ b/Engine/Source/Editor/UMGEditor/Private/Customizations/UMGDetailCustomizations.h
@@ -34,6 +34,9 @@ public:
 	/** Make a property binding widget */
 	static TSharedRef<SWidget> MakePropertyBindingWidget(TWeakPtr<FWidgetBlueprintEditor> InEditor, FDelegateProperty* InProperty, TSharedRef<IPropertyHandle> InDelegatePropertyHandle, bool bInGeneratePureBindings);
 
+	/** Whether the property (or its parent property, if this is an array/struct element) currently has bindings */
+	static bool HasPropertyBindings(TWeakPtr<FWidgetBlueprintEditor> InEditor, const TSharedRef<IPropertyHandle>& InPropertyHandle);
+
 private:
 	void PerformBindingCustomization(IDetailLayoutBuilder& DetailLayout);
 
diff --git a/Engine/Source/Editor/UMGEditor/Private/Details/DetailWidgetExtensionHandler.cpp b/Engine/Source/Editor/UMGEditor/Private/Details/DetailWidgetExtensionHandler.cpp
--- a/Engine/Source/Editor/UMGEditor/Private/Details/DetailWidgetExtensionHandler.cpp
+++ b/Engine/Source/Editor/UMGEditor/Private/Details/DetailWidgetExtensionHandler.cpp
@@ -20,6 +20,17 @@ bool FDetailWidgetExtensionHandler::IsPropertyExtendable(const UClass* InObjectC
 		TArray<UObject*> Objects;
 		InPropertyHandle.GetOuterObjects(Objects);
 
+		// Make the extension handler run for child struct/array properties if the parent is bindable.
+		// This is so we can later disable any child rows of a bound property in ExtendWidgetRow().
+		// Check to see if there's a delegate for the parent property.
+		TSharedPtr<const IPropertyHandle> ParentPropertyHandle;
+		for (TSharedPtr<const IPropertyHandle> CurrentPropertyHandle = InPropertyHandle.AsShared();
+			CurrentPropertyHandle && CurrentPropertyHandle->GetProperty();
+			CurrentPropertyHandle = CurrentPropertyHandle->GetParentHandle())
+		{
+			ParentPropertyHandle = CurrentPropertyHandle;
+		}
+
 		// We don't allow bindings on the CDO.
 		if (Objects[0] != nullptr && Objects[0]->HasAnyFlags(RF_ClassDefaultObject) )
 		{
@@ -32,10 +43,10 @@ bool FDetailWidgetExtensionHandler::IsPropertyExtendable(const UClass* InObjectC
 			return false;
 		}
 
-		FProperty* Property = InPropertyHandle.GetProperty();
-		FString DelegateName = Property->GetName() + "Delegate";
+		FProperty* ParentProperty = ParentPropertyHandle->GetProperty();
+		FString DelegateName = ParentProperty->GetName() + "Delegate";
 
-		if ( UClass* ContainerClass = Property->GetOwner<UClass>() )
+		if ( UClass* ContainerClass = ParentProperty->GetOwner<UClass>() )
 		{
 			FDelegateProperty* DelegateProperty = FindFProperty<FDelegateProperty>(ContainerClass, FName(*DelegateName));
 			if ( DelegateProperty )
@@ -48,19 +59,46 @@ bool FDetailWidgetExtensionHandler::IsPropertyExtendable(const UClass* InObjectC
 	return false;
 }
 
-TSharedRef<SWidget> FDetailWidgetExtensionHandler::GenerateExtensionWidget(const IDetailLayoutBuilder& InDetailBuilder, const UClass* InObjectClass, TSharedPtr<IPropertyHandle> InPropertyHandle)
+void FDetailWidgetExtensionHandler::ExtendWidgetRow(
+	FDetailWidgetRow& InWidgetRow,
+	const IDetailLayoutBuilder& InDetailBuilder,
+	const UClass* InObjectClass,
+	TSharedPtr<IPropertyHandle> InPropertyHandle)
 {
+	// Always make the value widget disabled when the property is bindable.
+	const TWeakPtr<FWidgetBlueprintEditor> BlueprintEditorPtr = BlueprintEditor;
+	const TSharedRef<IPropertyHandle> PropertyHandleRef = InPropertyHandle.ToSharedRef();
+	InWidgetRow.IsValueEnabled(TAttribute<bool>::Create(
+		[BlueprintEditorPtr, PropertyHandleRef]() {
+			// NOTE: HasPropertyBindings is potentially expensive (see its implementation) and this will
+			// get called for every bindable property's row, on every frame.
+			// However, this is not currently a big concern because we only extend the widget row
+			// in cases where there is one OuterObject (see IsPropertyExtendable above).
+			// But the performance might degrade if we extend the widget row with the property binding
+			// widget for a multiple-selection.
+			return !FBlueprintWidgetCustomization::HasPropertyBindings(BlueprintEditorPtr, PropertyHandleRef);
+		}));
+
 	FProperty* Property = InPropertyHandle->GetProperty();
 	FString DelegateName = Property->GetName() + "Delegate";
-	
-	FDelegateProperty* DelegateProperty = FindFieldChecked<FDelegateProperty>(Property->GetOwnerChecked<UClass>(), FName(*DelegateName));
+
+	// Don't show the property binding widget for child properties of a bindable property,
+	// i.e. if the property's owner isn't the class itself.
+	// This check is necessary because IsPropertyExtendable() accepts all bindable properties
+	// and their descendants, primarily so we can disable the descendants' value widgets above.
+	const UClass* OwnerClass = Property->GetOwner<UClass>();
+	if (!OwnerClass) {
+		return;
+	}
+
+	FDelegateProperty* DelegateProperty = FindFieldChecked<FDelegateProperty>(OwnerClass, FName(*DelegateName));
 
 	const bool bIsEditable = Property->HasAnyPropertyFlags(CPF_Edit | CPF_EditConst);
 	const bool bDoSignaturesMatch = DelegateProperty->SignatureFunction->GetReturnProperty()->SameType(Property);
 
 	if ( !ensure(bIsEditable && bDoSignaturesMatch) )
 	{
-		return SNullWidget::NullWidget;
+		return;
 	}
 
 	UWidgetBlueprint* WidgetBlueprint = BlueprintEditor.Pin()->GetWidgetBlueprintObj();
@@ -70,9 +108,12 @@ TSharedRef<SWidget> FDetailWidgetExtensionHandler::GenerateExtensionWidget(const
 		// Even if they don't want them on, we need to show them so they can remove them if they had any.
 		if (BlueprintEditor.Pin()->GetWidgetBlueprintObj()->Bindings.Num() == 0)
 		{
-			return SNullWidget::NullWidget;
+			return;
 		}
 	}
 
-	return FBlueprintWidgetCustomization::MakePropertyBindingWidget(BlueprintEditor, DelegateProperty, InPropertyHandle.ToSharedRef(), true);
+	InWidgetRow.ExtensionContent()
+	[
+		FBlueprintWidgetCustomization::MakePropertyBindingWidget(BlueprintEditor, DelegateProperty, InPropertyHandle.ToSharedRef(), true)
+	];
 }
diff --git a/Engine/Source/Editor/UMGEditor/Private/Details/DetailWidgetExtensionHandler.h b/Engine/Source/Editor/UMGEditor/Private/Details/DetailWidgetExtensionHandler.h
--- a/Engine/Source/Editor/UMGEditor/Private/Details/DetailWidgetExtensionHandler.h
+++ b/Engine/Source/Editor/UMGEditor/Private/Details/DetailWidgetExtensionHandler.h
@@ -16,7 +16,11 @@ public:
 
 	virtual bool IsPropertyExtendable(const UClass* InObjectClass, const class IPropertyHandle& PropertyHandle) const override;
 
-	virtual TSharedRef<SWidget> GenerateExtensionWidget(const IDetailLayoutBuilder& InDetailBuilder, const UClass* InObjectClass, TSharedPtr<IPropertyHandle> PropertyHandle) override;
+	virtual void ExtendWidgetRow(
+		FDetailWidgetRow& InWidgetRow,
+		const IDetailLayoutBuilder& InDetailBuilder,
+		const UClass* InObjectClass,
+		TSharedPtr<IPropertyHandle> PropertyHandle) override;
 
 private:
 	TWeakPtr<class FWidgetBlueprintEditor> BlueprintEditor;
