Float to Int conversion check for widget thumbnail size #RB patrick.boutot #jira ue-121018

diff --git a/Engine/Source/Editor/UMGEditor/Private/WidgetBlueprintEditorUtils.cpp b/Engine/Source/Editor/UMGEditor/Private/WidgetBlueprintEditorUtils.cpp
--- a/Engine/Source/Editor/UMGEditor/Private/WidgetBlueprintEditorUtils.cpp
+++ b/Engine/Source/Editor/UMGEditor/Private/WidgetBlueprintEditorUtils.cpp
@@ -2103,11 +2103,6 @@ TOptional<FWidgetBlueprintEditorUtils::FWidgetThumbnailProperties>  FWidgetBluep
 	}
 
 	float Scale;
-	// Create Renderer Target and WidgetRenderer
-	bool bApplyGammaCorrection = false;
-	FWidgetRenderer* WidgetRenderer = new FWidgetRenderer(bApplyGammaCorrection);
-	const bool bIsLinearSpace = !bApplyGammaCorrection;
-
 	// Change some configuration if it is for thumbnail creation
 	if (bIsForThumbnail) 
 	{
@@ -2123,13 +2118,24 @@ TOptional<FWidgetBlueprintEditorUtils::FWidgetThumbnailProperties>  FWidgetBluep
 		RenderTarget2D->RenderTargetFormat = RTF_RGBA8;
 		Scale = 1.f;
 		Offset = FVector2D(0.f, 0.f);
-		WidgetRenderer->SetIsPrepassNeeded(false);
 	}
 
 	ScaledSize = UnscaledSize * Scale;
-
+	if (ScaledSize.X < 1.f || ScaledSize.Y < 1.f)
+	{
+		return TOptional<FWidgetBlueprintEditorUtils::FWidgetThumbnailProperties>();
+	}
 	const EPixelFormat RequestedFormat = FSlateApplication::Get().GetRenderer()->GetSlateRecommendedColorFormat();
 	RenderTarget2D->InitCustomFormat(ScaledSize.X, ScaledSize.Y, RequestedFormat, false);
+
+	// Create Renderer Target and WidgetRenderer
+	bool bApplyGammaCorrection = false;
+	FWidgetRenderer* WidgetRenderer = new FWidgetRenderer(bApplyGammaCorrection);
+	const bool bIsLinearSpace = !bApplyGammaCorrection;
+	if (!bIsForThumbnail)
+	{
+		WidgetRenderer->SetIsPrepassNeeded(false);
+	}
 	WidgetRenderer->DrawWindow(RenderTarget2D, *HitTestGrid, Window, Scale, ScaledSize, 0.1f);
 
 	if (WidgetRenderer)
