UMG Details View will now show everything until there's a selected category or property to show (fixed bug introduced in CL 15797104)
Properties to show in UMG Details View now check up to their root property

diff --git a/Engine/Source/Editor/UMGEditor/Private/Components/DetailsView.cpp b/Engine/Source/Editor/UMGEditor/Private/Components/DetailsView.cpp
--- a/Engine/Source/Editor/UMGEditor/Private/Components/DetailsView.cpp
+++ b/Engine/Source/Editor/UMGEditor/Private/Components/DetailsView.cpp
@@ -158,9 +158,14 @@ void UDetailsView::ToggleWhitelistedProperties()
 	}
 }
 
+bool UDetailsView::IsRowVisibilityFiltered() const
+{
+	return bShowOnlyWhitelisted && (PropertiesToShow.Num() > 0 || CategoriesToShow.Num() > 0);
+}
+
 bool UDetailsView::GetIsPropertyVisible(const FPropertyAndParent& PropertyAndParent) const
 {
-    if (!bShowOnlyWhitelisted)
+	if (!IsRowVisibilityFiltered())
 	{
 		return true;
 	}
@@ -168,8 +173,16 @@ bool UDetailsView::GetIsPropertyVisible(const FPropertyAndParent& PropertyAndPar
 	{
 		return true;
 	}
-	// get the topmost parent's category name if the property has one
+
 	const FProperty* Property = PropertyAndParent.ParentProperties.Num() > 0 ? PropertyAndParent.ParentProperties.Last() : &PropertyAndParent.Property;
+
+	// Check against the parent property name
+	if (PropertiesToShow.Contains(Property->GetFName()))
+	{
+		return true;
+	}
+
+	// get the topmost parent's category name if the property has one
 	const FString CategoryString = FObjectEditorUtils::GetCategoryFName(Property).ToString();
 	FString SubcategoryString = CategoryString;
 	
@@ -192,7 +205,7 @@ bool UDetailsView::GetIsPropertyVisible(const FPropertyAndParent& PropertyAndPar
 
 bool UDetailsView::GetIsRowVisible(FName InRowName, FName InParentName) const
 {
-    if (!bShowOnlyWhitelisted)
+	if (!IsRowVisibilityFiltered())
 	{
 		return true;
 	}
