Removed comparison of int with SMALL_NUMBER #jira ue-121699 #RB patrick.boutot

diff --git a/Engine/Source/Editor/UMGEditor/Private/WidgetBlueprintThumbnailRenderer.cpp b/Engine/Source/Editor/UMGEditor/Private/WidgetBlueprintThumbnailRenderer.cpp
--- a/Engine/Source/Editor/UMGEditor/Private/WidgetBlueprintThumbnailRenderer.cpp
+++ b/Engine/Source/Editor/UMGEditor/Private/WidgetBlueprintThumbnailRenderer.cpp
@@ -24,7 +24,7 @@ void UWidgetBlueprintThumbnailRenderer::Draw(UObject* Object, int32 X, int32 Y,
 {
 #if !UE_SERVER
 
-	if (Width < SMALL_NUMBER || Height < SMALL_NUMBER)
+	if (Width < 1 || Height < 1)
 	{
 		return;
 	}
