UMG: Use a pool to reduce the number of thumbnail generated by the WidgetThumbnailRenderer.

diff --git a/Engine/Source/Editor/UMGEditor/Private/WidgetBlueprintThumbnailRenderer.cpp b/Engine/Source/Editor/UMGEditor/Private/WidgetBlueprintThumbnailRenderer.cpp
--- a/Engine/Source/Editor/UMGEditor/Private/WidgetBlueprintThumbnailRenderer.cpp
+++ b/Engine/Source/Editor/UMGEditor/Private/WidgetBlueprintThumbnailRenderer.cpp
@@ -7,6 +7,7 @@
 #include "CanvasTypes.h"
 #include "Engine/TextureRenderTarget2D.h"
 #include "Input/HittestGrid.h"
+#include "Kismet2/KismetEditorUtilities.h"
 #include "Slate/WidgetRenderer.h"
 #include "WidgetBlueprint.h"
 #include "Widgets/SVirtualWindow.h"
@@ -14,12 +15,104 @@
 
 #define LOCTEXT_NAMESPACE "UWidgetBlueprintThumbnailRenderer"
 
+class FWidgetBlueprintThumbnailPool
+{
+public:
+	struct FInstance
+	{
+		TWeakObjectPtr<UClass> WidgetClass;
+		TWeakObjectPtr<UTextureRenderTarget2D> RenderTarget2D;
+		TWeakObjectPtr<UUserWidget> Widget;
+	};
+
+public:
+	static constexpr int32 MaxNumInstance = 50;
+
+	FWidgetBlueprintThumbnailPool()
+	{
+		InstancedThumbnails.Reserve(MaxNumInstance);
+	}
+
+	~FWidgetBlueprintThumbnailPool()
+	{
+		Clear();
+	}
+
+	FInstance* FindThumbnail(const UClass* InClass) const
+	{
+		check(InClass);
+		const FName ClassName = InClass->GetFName();
+
+		return InstancedThumbnails.FindRef(ClassName);
+	}
+
+	FInstance& EnsureThumbnail(const UClass* InClass)
+	{
+		check(InClass);
+		const FName ClassName = InClass->GetFName();
+
+		FInstance* ExistingThumbnail = InstancedThumbnails.FindRef(ClassName);
+		if (!ExistingThumbnail)
+		{
+			if (InstancedThumbnails.Num() >= MaxNumInstance)
+			{
+				Clear();
+				CollectGarbage(GARBAGE_COLLECTION_KEEPFLAGS);
+			}
+
+			ExistingThumbnail = new FInstance;
+			InstancedThumbnails.Add(ClassName, ExistingThumbnail);
+		}
+
+		return *ExistingThumbnail;
+	}
+
+	void RemoveThumbnail(const UClass* InClass)
+	{
+		check(InClass);
+		InstancedThumbnails.Remove(InClass->GetFName());
+	}
+
+	void Clear()
+	{
+		for(auto& Instance : InstancedThumbnails)
+		{
+			delete Instance.Value;
+		}
+		InstancedThumbnails.Reset();
+	}
+
+private:
+	TMap<FName, FInstance*> InstancedThumbnails;
+};
+
+
+void UWidgetBlueprintThumbnailRenderer::FWidgetBlueprintThumbnailPoolDeleter::operator()(FWidgetBlueprintThumbnailPool* Pointer)
+{
+	delete Pointer;
+}
+
+
+UWidgetBlueprintThumbnailRenderer::UWidgetBlueprintThumbnailRenderer()
+	: ThumbnailPool(new FWidgetBlueprintThumbnailPool)
+{
+	FKismetEditorUtilities::OnBlueprintUnloaded.AddUObject(this, &UWidgetBlueprintThumbnailRenderer::OnBlueprintUnloaded);
+}
+
+
+UWidgetBlueprintThumbnailRenderer::~UWidgetBlueprintThumbnailRenderer()
+{
+	FKismetEditorUtilities::OnBlueprintUnloaded.RemoveAll(this);
+}
+
+
 bool UWidgetBlueprintThumbnailRenderer::CanVisualizeAsset(UObject* Object)
 {
 	UWidgetBlueprint* Blueprint = Cast<UWidgetBlueprint>(Object);
 	return (Blueprint && Blueprint->GeneratedClass && Blueprint->GeneratedClass->IsChildOf(UWidget::StaticClass()));
 }
 
+
 void UWidgetBlueprintThumbnailRenderer::Draw(UObject* Object, int32 X, int32 Y, uint32 Width, uint32 Height, FRenderTarget* RenderTarget, FCanvas* Canvas, bool bAdditionalViewFamily)
 {
 #if !UE_SERVER
@@ -34,6 +127,19 @@ void UWidgetBlueprintThumbnailRenderer::Draw(UObject* Object, int32 X, int32 Y,
 		return;
 	}
 
+	UWidgetBlueprint* WidgetBlueprintToRender = Cast<UWidgetBlueprint>(Object);
+	const bool bIsBlueprintValid = IsValid(WidgetBlueprintToRender)
+		&& IsValid(WidgetBlueprintToRender->GeneratedClass)
+		&& WidgetBlueprintToRender->bHasBeenRegenerated
+		&& !WidgetBlueprintToRender->bBeingCompiled
+		&& !WidgetBlueprintToRender->HasAnyFlags(RF_Transient)
+		&& !WidgetBlueprintToRender->GeneratedClass->HasAnyClassFlags(CLASS_Deprecated | CLASS_Abstract)
+		&& WidgetBlueprintToRender->GeneratedClass->IsChildOf(UWidget::StaticClass());
+	if (!bIsBlueprintValid)
+	{
+		return;
+	}
+
 	// Create a plain gray background for the thumbnail
 	const int32 SizeOfUV = 1;
 	FLinearColor GrayBackgroundColor(FVector4(.03f, .03f, .03f, 1.f));
@@ -42,14 +148,12 @@ void UWidgetBlueprintThumbnailRenderer::Draw(UObject* Object, int32 X, int32 Y,
 		0.0f, 0.0f, SizeOfUV, SizeOfUV,
 		GrayBackgroundColor);
 
-	// Convert the object from UObject to SWidget
-	UWidgetBlueprint* WidgetBlueprintToRender = Cast<UWidgetBlueprint>(Object);
-
 	// check if an image is used instead of auto generating the thumbnail
 	if (WidgetBlueprintToRender->ThumbnailImage)
 	{
 		FVector2D TextureSize(WidgetBlueprintToRender->ThumbnailImage->GetSizeX(), WidgetBlueprintToRender->ThumbnailImage->GetSizeY());
-		if (TextureSize.X > SMALL_NUMBER && TextureSize.Y > SMALL_NUMBER) {
+		if (TextureSize.X > SMALL_NUMBER && TextureSize.Y > SMALL_NUMBER)
+		{
 			TTuple<float, FVector2D> ScaleAndOffset = FWidgetBlueprintEditorUtils::GetThumbnailImageScaleAndOffset(TextureSize, FVector2D(Width, Height));
 			float Scale = ScaleAndOffset.Get<0>();
 			FVector2D Offset = ScaleAndOffset.Get<1>();
@@ -63,41 +167,37 @@ void UWidgetBlueprintThumbnailRenderer::Draw(UObject* Object, int32 X, int32 Y,
 	}
 	else
 	{
-		const bool bIsBlueprintValid = IsValid(WidgetBlueprintToRender)
-			&& IsValid(WidgetBlueprintToRender->GeneratedClass)
-			&& WidgetBlueprintToRender->bHasBeenRegenerated
-			//&& Blueprint->IsUpToDate() - This condition makes the thumbnail blank whenever the BP is dirty. It seems too strict.
-			&& !WidgetBlueprintToRender->bBeingCompiled
-			&& !WidgetBlueprintToRender->HasAnyFlags(RF_Transient);
-
-		TSharedPtr<SWidget> WindowContent;
-		UClass* ClassToGenerate = WidgetBlueprintToRender->GeneratedClass;
 		UUserWidget* WidgetInstance = nullptr;
-		if (bIsBlueprintValid && ClassToGenerate->IsChildOf(UWidget::StaticClass()))
+		UTextureRenderTarget2D* RenderTarget2D = nullptr;
+
 		{
-			WidgetInstance = NewObject<UUserWidget>(GetTransientPackage(), ClassToGenerate);
-			if (WidgetInstance)
+			UClass* ClassToGenerate = WidgetBlueprintToRender->GeneratedClass;
+			FWidgetBlueprintThumbnailPool::FInstance& ThumbnailInstance = ThumbnailPool->EnsureThumbnail(ClassToGenerate);
+
+			WidgetInstance = ThumbnailInstance.Widget.Get();
+			if (!WidgetInstance)
 			{
-				WidgetInstance->Initialize();
-				WidgetInstance->SetDesignerFlags(EWidgetDesignFlags::Designing | EWidgetDesignFlags::ExecutePreConstruct);
+				ThumbnailInstance.Widget = NewObject<UUserWidget>(GetTransientPackage(), ClassToGenerate);
+				ThumbnailInstance.Widget->Initialize();
+				ThumbnailInstance.Widget->SetDesignerFlags(EWidgetDesignFlags::Designing | EWidgetDesignFlags::ExecutePreConstruct);
+
+				WidgetInstance = ThumbnailInstance.Widget.Get();
 			}
-		}
 
-		if (WidgetInstance == nullptr)
-		{
-			return;
-		}
-		FVector2D ThumbnailSize(Width, Height);
+			RenderTarget2D = ThumbnailInstance.RenderTarget2D.Get();
+			if (!RenderTarget2D)
+			{
+				ThumbnailInstance.RenderTarget2D = NewObject<UTextureRenderTarget2D>(GetTransientPackage());
+				ThumbnailInstance.RenderTarget2D->Filter = TF_Bilinear;
+				ThumbnailInstance.RenderTarget2D->ClearColor = FLinearColor::Transparent;
+				ThumbnailInstance.RenderTarget2D->SRGB = true;
+				ThumbnailInstance.RenderTarget2D->TargetGamma = 1;
 
-		if (!RenderTarget2D)
-		{
-			RenderTarget2D = NewObject<UTextureRenderTarget2D>();
-			RenderTarget2D->Filter = TF_Bilinear;
-			RenderTarget2D->ClearColor = FLinearColor::Transparent;
-			RenderTarget2D->SRGB = true;
-			RenderTarget2D->TargetGamma = 1;
+				RenderTarget2D = ThumbnailInstance.RenderTarget2D.Get();
+			}
 		}
 
+		FVector2D ThumbnailSize(Width, Height);
 		TOptional<FWidgetBlueprintEditorUtils::FWidgetThumbnailProperties> ScaleAndOffset;
 
 		if (WidgetBlueprintToRender->ThumbnailSizeMode == EThumbnailPreviewSizeMode::Custom)
@@ -126,4 +226,12 @@ void UWidgetBlueprintThumbnailRenderer::Draw(UObject* Object, int32 X, int32 Y,
 }
 
 
+void UWidgetBlueprintThumbnailRenderer::OnBlueprintUnloaded(UBlueprint* Blueprint)
+{
+	if (Blueprint && Blueprint->GeneratedClass)
+	{
+		ThumbnailPool->RemoveThumbnail(Blueprint->GeneratedClass);
+	}
+}
+
 #undef LOCTEXT_NAMESPACE
\ No newline at end of file
diff --git a/Engine/Source/Editor/UMGEditor/Public/WidgetBlueprintThumbnailRenderer.h b/Engine/Source/Editor/UMGEditor/Public/WidgetBlueprintThumbnailRenderer.h
--- a/Engine/Source/Editor/UMGEditor/Public/WidgetBlueprintThumbnailRenderer.h
+++ b/Engine/Source/Editor/UMGEditor/Public/WidgetBlueprintThumbnailRenderer.h
@@ -9,19 +9,29 @@
 class FCanvas;
 class FRenderTarget;
 class UTextureRenderTarget2D;
+class FWidgetBlueprintThumbnailPool;
 
 UCLASS()
 class UMGEDITOR_API UWidgetBlueprintThumbnailRenderer : public UDefaultSizedThumbnailRenderer
 {
 	GENERATED_BODY()
 
+	UWidgetBlueprintThumbnailRenderer();
+	virtual ~UWidgetBlueprintThumbnailRenderer();
+
 	//~ Begin UThumbnailRenderer Object
 	bool CanVisualizeAsset(UObject* Object) override;
 	void Draw(UObject* Object, int32 X, int32 Y, uint32 Width, uint32 Height, FRenderTarget* RenderTarget, FCanvas* Canvas, bool bAdditionalViewFamily) override;
+	//~ End UThumbnailRenderer Object
 
 private:
+	void OnBlueprintUnloaded(UBlueprint* Blueprint);
 
-	UPROPERTY(Transient)
-	TObjectPtr<UTextureRenderTarget2D> RenderTarget2D;
+private:
+	struct FWidgetBlueprintThumbnailPoolDeleter
+	{
+		void operator()(FWidgetBlueprintThumbnailPool* Pointer);
+	};
 
+	TUniquePtr<FWidgetBlueprintThumbnailPool, FWidgetBlueprintThumbnailPoolDeleter> ThumbnailPool;
 };
