Enable widget library by default.
---
Add size & type display options to SAssetView constructor / interface.
Implement search & centralized list / column / size control on widget library.

diff --git a/Engine/Source/Editor/ContentBrowser/Private/AssetViewWidgets.cpp b/Engine/Source/Editor/ContentBrowser/Private/AssetViewWidgets.cpp
--- a/Engine/Source/Editor/ContentBrowser/Private/AssetViewWidgets.cpp
+++ b/Engine/Source/Editor/ContentBrowser/Private/AssetViewWidgets.cpp
@@ -1639,6 +1639,7 @@ void SAssetTileItem::Construct( const FArguments& InArgs )
 		.IsSelected(InArgs._IsSelected)
 		);
 
+	bShowType = InArgs._ShowType;
 	AssetThumbnail = InArgs._AssetThumbnail;
 	ItemWidth = InArgs._ItemWidth;
 	ThumbnailPadding = IsFolder() ? InArgs._ThumbnailPadding + 5.0f : InArgs._ThumbnailPadding;
@@ -1865,7 +1866,7 @@ FOptionalSize SAssetTileItem::GetThumbnailBoxSize() const
 
 EVisibility SAssetTileItem::GetAssetClassLabelVisibility() const
 {
-	if(!IsFolder())
+	if(!IsFolder() && bShowType)
 	{
 		FOptionalSize ThumbSize = GetThumbnailBoxSize();
 		if (ThumbSize.IsSet())
@@ -1950,8 +1951,8 @@ void SAssetTileItem::InitializeAssetNameHeights()
 			const static FName SmallFontName("ContentBrowser.AssetTileViewNameFontSmall");
 			FSlateFontInfo Font = FEditorStyle::GetFontStyle(SmallFontName);
 			TSharedRef<FSlateFontMeasure> FontMeasureService = FSlateApplication::Get().GetRenderer()->GetFontMeasureService();
-			float MaxHeight = FontMeasureService->GetMaxCharacterHeight(Font);
-			AssetNameHeights[(int32)EThumbnailSize::Small] = MaxHeight * 2;
+			SmallFontHeight = FontMeasureService->GetMaxCharacterHeight(Font);
+			AssetNameHeights[(int32)EThumbnailSize::Small] = SmallFontHeight * 2 ;
 		}
 
 
@@ -1959,10 +1960,10 @@ void SAssetTileItem::InitializeAssetNameHeights()
 			const static FName SmallFontName("ContentBrowser.AssetTileViewNameFont");
 			FSlateFontInfo Font = FEditorStyle::GetFontStyle(SmallFontName);
 			TSharedRef<FSlateFontMeasure> FontMeasureService = FSlateApplication::Get().GetRenderer()->GetFontMeasureService();
-			float MaxHeight = FontMeasureService->GetMaxCharacterHeight(Font);
-			AssetNameHeights[(int32)EThumbnailSize::Medium] = MaxHeight * 3;
-			AssetNameHeights[(int32)EThumbnailSize::Large] = MaxHeight * 4;
-			AssetNameHeights[(int32)EThumbnailSize::Huge] = MaxHeight * 5;
+			RegularFontHeight = FontMeasureService->GetMaxCharacterHeight(Font);
+			AssetNameHeights[(int32)EThumbnailSize::Medium] = RegularFontHeight * 3;
+			AssetNameHeights[(int32)EThumbnailSize::Large] = RegularFontHeight * 4;
+			AssetNameHeights[(int32)EThumbnailSize::Huge] = RegularFontHeight * 5;
 		}
 
 		bInitializedHeights = true;
@@ -1971,6 +1972,8 @@ void SAssetTileItem::InitializeAssetNameHeights()
 }
 
 float SAssetTileItem::AssetNameHeights[(int32)EThumbnailSize::MAX];
+float SAssetTileItem::RegularFontHeight(0);
+float SAssetTileItem::SmallFontHeight(0);
 
 ///////////////////////////////
 // SAssetColumnItem
diff --git a/Engine/Source/Editor/ContentBrowser/Private/SAssetView.cpp b/Engine/Source/Editor/ContentBrowser/Private/SAssetView.cpp
--- a/Engine/Source/Editor/ContentBrowser/Private/SAssetView.cpp
+++ b/Engine/Source/Editor/ContentBrowser/Private/SAssetView.cpp
@@ -79,7 +79,7 @@ namespace
 }
 
 
-static FText ThumbnailSizeToDisplayName(EThumbnailSize InSize)
+FText SAssetView::ThumbnailSizeToDisplayName(EThumbnailSize InSize)
 {
 	switch (InSize)
 	{
@@ -206,7 +206,7 @@ void SAssetView::Construct( const FArguments& InArgs )
 	// Listen for when view settings are changed
 	UContentBrowserSettings::OnSettingChanged().AddSP(this, &SAssetView::HandleSettingChanged);
 
-	ThumbnailSize = EThumbnailSize::Medium;
+	ThumbnailSize = InArgs._InitialThumbnailSize;
 
 	// Get desktop metrics
 	FDisplayMetrics DisplayMetrics;
@@ -228,7 +228,7 @@ void SAssetView::Construct( const FArguments& InArgs )
 	TileViewThumbnailSize = 150;
 	TileViewThumbnailPadding = 9;
 
-	TileViewNameHeight = 100;
+	TileViewNameHeight = 50;
 
 	MinThumbnailScale = 0.2f * ThumbnailScaleRangeScalar;
 	MaxThumbnailScale = 1.9f * ThumbnailScaleRangeScalar;
@@ -254,6 +254,7 @@ void SAssetView::Construct( const FArguments& InArgs )
 	bShowPathInColumnView = InArgs._ShowPathInColumnView;
 	bShowTypeInColumnView = InArgs._ShowTypeInColumnView;
 	bSortByPathInColumnView = bShowPathInColumnView & InArgs._SortByPathInColumnView;
+	bShowTypeInTileView = InArgs._ShowTypeInTileView;
 	bForceShowEngineContent = InArgs._ForceShowEngineContent;
 	bForceShowPluginContent = InArgs._ForceShowPluginContent;
 
@@ -2631,7 +2632,7 @@ void SAssetView::PopulateViewButtonMenu(UToolMenu* Menu)
 			{
 				SizeSection.AddMenuEntry(
 					NAME_None,
-					ThumbnailSizeToDisplayName((EThumbnailSize)EnumValue),
+					SAssetView::ThumbnailSizeToDisplayName((EThumbnailSize)EnumValue),
 					FText::GetEmpty(),
 					FSlateIcon(),
 					FUIAction(
@@ -3012,6 +3013,11 @@ void SAssetView::SetCurrentViewType(EAssetViewType::Type NewType)
 	}
 }
 
+void SAssetView::SetCurrentThumbnailSize(EThumbnailSize NewThumbnailSize)
+{
+	OnThumbnailSizeChanged(NewThumbnailSize);
+}
+
 void SAssetView::SetCurrentViewTypeFromMenu(EAssetViewType::Type NewType)
 {
 	if (NewType != CurrentViewType)
@@ -3298,7 +3304,8 @@ TSharedRef<ITableRow> SAssetView::MakeTileViewWidget(TSharedPtr<FAssetViewItem>
 			.OnIsAssetValidForCustomToolTip(OnIsAssetValidForCustomToolTip)
 			.OnGetCustomAssetToolTip(OnGetCustomAssetToolTip)
 			.OnVisualizeAssetToolTip( OnVisualizeAssetToolTip )
-			.OnAssetToolTipClosing( OnAssetToolTipClosing );
+			.OnAssetToolTipClosing( OnAssetToolTipClosing )
+			.ShowType(bShowTypeInTileView);
 
 		TableRowWidget->SetContent(Item);
 
@@ -3882,7 +3889,6 @@ void SAssetView::ToggleThumbnailEditMode()
 	bThumbnailEditMode = !bThumbnailEditMode;
 }
 
-
 void SAssetView::OnThumbnailSizeChanged(EThumbnailSize NewThumbnailSize)
 {
 	ThumbnailSize = NewThumbnailSize;
@@ -3927,6 +3933,32 @@ bool SAssetView::IsThumbnailScalingAllowed() const
 	return GetCurrentViewType() != EAssetViewType::Column;
 }
 
+float SAssetView::GetTileViewTypeNameHeight() const
+{
+	float TypeNameHeight = 0;
+
+	if (bShowTypeInTileView)
+	{
+		TypeNameHeight = 50;
+	}
+	else
+	{
+		if (ThumbnailSize == EThumbnailSize::Small)
+		{
+			TypeNameHeight = 25;
+		}
+		else if (ThumbnailSize > EThumbnailSize::Small)
+		{
+			TypeNameHeight = -25;
+		}
+		else
+		{
+			TypeNameHeight = -40;
+		}
+	}
+	return TypeNameHeight;
+}
+
 float SAssetView::GetListViewItemHeight() const
 {
 	return (ListViewThumbnailSize + ListViewThumbnailPadding * 2) * FMath::Lerp(MinThumbnailScale, MaxThumbnailScale, GetThumbnailScale());
@@ -3934,7 +3966,7 @@ float SAssetView::GetListViewItemHeight() const
 
 float SAssetView::GetTileViewItemHeight() const
 {
-	return (TileViewNameHeight * FMath::Lerp(MinThumbnailScale, MaxThumbnailScale, GetThumbnailScale())) + GetTileViewItemBaseHeight() * FillScale;
+	return ((TileViewNameHeight + GetTileViewTypeNameHeight()) * FMath::Lerp(MinThumbnailScale, MaxThumbnailScale, GetThumbnailScale())) + GetTileViewItemBaseHeight() * FillScale;
 }
 
 float SAssetView::GetTileViewItemBaseHeight() const
diff --git a/Engine/Source/Editor/ContentBrowser/Public/AssetViewWidgets.h b/Engine/Source/Editor/ContentBrowser/Public/AssetViewWidgets.h
--- a/Engine/Source/Editor/ContentBrowser/Public/AssetViewWidgets.h
+++ b/Engine/Source/Editor/ContentBrowser/Public/AssetViewWidgets.h
@@ -42,20 +42,6 @@ DECLARE_DELEGATE_OneParam( FOnItemDestroyed, const TSharedPtr<FAssetViewItem>& /
 class SAssetListItem;
 class SAssetTileItem;
 
-enum class EThumbnailSize : uint8
-{
-	Tiny = 0,
-	Small,
-	Medium,
-	Large,
-	Huge,
-
-	// Not a size
-	MAX
-};
-
-
-
 namespace FAssetViewModeUtils 
 {
 	FReply OnViewModeKeyDown( const TSet< TSharedPtr<FAssetViewItem> >& SelectedItems, const FKeyEvent& InKeyEvent );
@@ -474,6 +460,7 @@ public:
 		, _CurrentThumbnailSize(EThumbnailSize::Medium)
 		, _ItemWidth(16)
 		, _ShouldAllowToolTip(true)
+		, _ShowType(true)
 		, _ThumbnailEditMode(false)
 	{}
 
@@ -516,6 +503,9 @@ public:
 		/** If false, the tooltip will not be displayed */
 		SLATE_ATTRIBUTE( bool, ShouldAllowToolTip )
 
+		/** If false, will not show type */
+		SLATE_ARGUMENT( bool, ShowType )
+
 		/** The string in the title to highlight (used when searching by string) */
 		SLATE_ATTRIBUTE( FText, HighlightText )
 
@@ -562,6 +552,11 @@ public:
 
 	/** Get the border image to display */
 	virtual const FSlateBrush* GetBorderImage() const;
+
+	static void InitializeAssetNameHeights();
+	static float GetRegularFontHeight() { return RegularFontHeight; }
+	static float GetSmallFontHeight() { return SmallFontHeight; }
+
 protected:
 	/** SAssetViewItem interface */
 	virtual float GetNameTextWrapWidth() const override { return LastGeometry.GetLocalSize().X - 15.f; }
@@ -594,9 +589,10 @@ protected:
 
 	int32 GetGenericThumbnailSize() const;
 
-	static void InitializeAssetNameHeights();
-
 private:
+	/** If false, the tooltip will not be displayed */
+	bool bShowType;
+
 	/** The handle to the thumbnail that this item is rendering */
 	TSharedPtr<FAssetThumbnail> AssetThumbnail;
 
@@ -606,6 +602,12 @@ private:
 	/** Max name height for each thumbnail size */
 	static float AssetNameHeights[(int32)EThumbnailSize::MAX];
 
+	/** Regular thumbnail font size */
+	static float RegularFontHeight;
+
+	/** Small thumbnail font size */
+	static float SmallFontHeight;
+
 	/** The padding allotted for the thumbnail */
 	float ThumbnailPadding;
 
diff --git a/Engine/Source/Editor/ContentBrowser/Public/SAssetView.h b/Engine/Source/Editor/ContentBrowser/Public/SAssetView.h
--- a/Engine/Source/Editor/ContentBrowser/Public/SAssetView.h
+++ b/Engine/Source/Editor/ContentBrowser/Public/SAssetView.h
@@ -46,8 +46,6 @@ DECLARE_DELEGATE(FOnSearchOptionChanged);
 /** Fires whenever asset view options menu is being opened, gives chance for external code to set additional context */
 DECLARE_DELEGATE_OneParam(FOnExtendAssetViewOptionsMenuContext, FToolMenuContext&);
 
-enum class EThumbnailSize : uint8;
-
 /**
  * A widget to display a list of filtered assets
  */
@@ -59,6 +57,7 @@ public:
 		, _ThumbnailLabel( EThumbnailLabel::ClassName )
 		, _AllowThumbnailHintLabel(true)
 		, _InitialViewType(EAssetViewType::Tile)
+		, _InitialThumbnailSize(EThumbnailSize::Medium)
 		, _ShowBottomToolbar(true)
 		, _ShowViewOptions(true)
 		, _AllowThumbnailEditMode(false)
@@ -78,6 +77,7 @@ public:
 		, _ShowPathInColumnView(false)
 		, _ShowTypeInColumnView(true)
 		, _SortByPathInColumnView(false)
+		, _ShowTypeInTileView(true)
 		, _ForceShowEngineContent(false)
 		, _ForceShowPluginContent(false)
 		{}
@@ -147,6 +147,10 @@ public:
 
 		/** The initial view type */
 		SLATE_ARGUMENT( EAssetViewType::Type, InitialViewType )
+
+		/** Initial thumbnail size */
+		SLATE_ARGUMENT( EThumbnailSize, InitialThumbnailSize )
+
 		/** Should the toolbar indicating number of selected assets, mode switch buttons, etc... be shown? */
 		SLATE_ARGUMENT( bool, ShowBottomToolbar )
 
@@ -203,6 +207,9 @@ public:
 
 		/** Sort by path in the column view. Only works if the initial view type is Column */
 		SLATE_ARGUMENT(bool, SortByPathInColumnView)
+		
+		/** Should show Type in column view if true */
+		SLATE_ARGUMENT(bool, ShowTypeInTileView)
 
 		/** Should always show engine content */
 		SLATE_ARGUMENT(bool, ForceShowEngineContent)
@@ -353,9 +360,18 @@ public:
 	/** Handler for when the view combo button is clicked */
 	TSharedRef<SWidget> GetViewButtonContent();
 
+	/** Sets the view type and updates lists accordingly */
+	void SetCurrentViewType(EAssetViewType::Type NewType);
+
+	/** Sets the thumbnail size and updates lists accordingly */
+	void SetCurrentThumbnailSize(EThumbnailSize NewThumbnailSize);
+
 	/** Gets the text for the asset count label */
 	FText GetAssetCountText() const;
 
+	/** Gets text name for given thumbnail */
+	static FText ThumbnailSizeToDisplayName(EThumbnailSize InSize);
+
 private:
 
 	/** Sets the pending selection to the current selection (used when changing views or refreshing the view). */
@@ -568,9 +584,6 @@ private:
 	/** @return true when we are organizing folders into virtual locations */
 	bool IsOrganizingFolders() const;
 
-	/** Sets the view type and updates lists accordingly */
-	void SetCurrentViewType(EAssetViewType::Type NewType);
-
 	/** Sets the view type and forcibly dismisses all currently open context menus */
 	void SetCurrentViewTypeFromMenu(EAssetViewType::Type NewType);
 
@@ -687,6 +700,9 @@ private:
 	/** Gets the current thumbnail size enum */
 	EThumbnailSize GetThumbnailSize() const { return ThumbnailSize; }
 
+	/** Gets the spacing dedicated to support type names */
+	float GetTileViewTypeNameHeight() const;
+
 	/** Gets the scaled item height for the list view */
 	float GetListViewItemHeight() const;
 	
@@ -1003,6 +1019,9 @@ private:
 	/** If true, it sorts by path and then name */
 	bool bSortByPathInColumnView;
 
+	/** If true, it will show type in the tile view */
+	bool bShowTypeInTileView;
+
 	/** If true, engine content is always shown */
 	bool bForceShowEngineContent;
 
diff --git a/Engine/Source/Editor/EditorStyle/Private/StarshipStyle.cpp b/Engine/Source/Editor/EditorStyle/Private/StarshipStyle.cpp
--- a/Engine/Source/Editor/EditorStyle/Private/StarshipStyle.cpp
+++ b/Engine/Source/Editor/EditorStyle/Private/StarshipStyle.cpp
@@ -7597,6 +7597,20 @@ void FStarshipEditorStyle::FStyle::SetupUMGEditorStyles()
 		.SetSelectedTextColor(InvertedForeground)
 	);
 
+	Set("UMGEditor.LibraryView", FTableRowStyle()
+		.SetEvenRowBackgroundBrush(FSlateColorBrush(FStyleColors::Background))
+		.SetEvenRowBackgroundHoveredBrush(FSlateColorBrush(FStyleColors::Background))
+		.SetOddRowBackgroundBrush(FSlateColorBrush(FStyleColors::Background))
+		.SetOddRowBackgroundHoveredBrush(FSlateColorBrush(FStyleColors::Background))
+		.SetSelectorFocusedBrush(FSlateNoResource())
+		.SetActiveBrush(FSlateNoResource())
+		.SetActiveHoveredBrush(FSlateNoResource())
+		.SetInactiveBrush(FSlateNoResource())
+		.SetInactiveHoveredBrush(FSlateNoResource())
+		.SetTextColor(DefaultForeground)
+		.SetSelectedTextColor(InvertedForeground)
+	);
+
 	// Style of the favorite toggle
 	const FCheckBoxStyle UMGEditorFavoriteToggleStyle = FCheckBoxStyle()
 		.SetCheckBoxType(ESlateCheckBoxType::CheckBox)
diff --git a/Engine/Source/Editor/UMGEditor/Private/Library/SLibraryView.cpp b/Engine/Source/Editor/UMGEditor/Private/Library/SLibraryView.cpp
--- a/Engine/Source/Editor/UMGEditor/Private/Library/SLibraryView.cpp
+++ b/Engine/Source/Editor/UMGEditor/Private/Library/SLibraryView.cpp
@@ -9,13 +9,12 @@
 #include "Editor.h"
 #include "Widgets/Images/SImage.h"
 #include "Widgets/Layout/SScrollBorder.h"
+#include "SAssetView.h"
 
 #if WITH_EDITOR
 	#include "EditorStyleSet.h"
 #endif // WITH_EDITOR
 
-
-
 #include "DragDrop/WidgetTemplateDragDropOp.h"
 
 #include "Templates/WidgetTemplateClass.h"
@@ -26,11 +25,11 @@
 #include "Widgets/Input/SCheckBox.h"
 
 #include "Settings/ContentBrowserSettings.h"
+#include "IContentBrowserSingleton.h"
 #include "WidgetBlueprintEditorUtils.h"
 
-
-
 #include "UMGEditorProjectSettings.h"
+#include "SLibraryView.h"
 
 #define LOCTEXT_NAMESPACE "UMG"
 
@@ -119,6 +118,73 @@ FReply SLibraryViewItem::OnMouseButtonDoubleClick(const FGeometry& InMyGeometry,
 	return WidgetViewModel->Template->OnDoubleClicked();
 };
 
+TSharedRef<SWidget> SLibraryView::ConstructViewOptions()
+{
+	FMenuBuilder ViewOptions(true, nullptr);
+
+	{
+		ViewOptions.BeginSection("AssetViewType", LOCTEXT("ViewTypeHeading", "View Type"));
+		ViewOptions.AddMenuEntry(
+			LOCTEXT("TileViewOption", "Tiles"),
+			LOCTEXT("TileViewOptionToolTip", "View assets as tiles in a grid."),
+			FSlateIcon(),
+			FUIAction(
+				FExecuteAction::CreateSP(this, &SLibraryView::SetCurrentViewTypeFromMenu, EAssetViewType::Tile),
+				FCanExecuteAction(),
+				FIsActionChecked::CreateSP(this, &SLibraryView::IsCurrentViewType, EAssetViewType::Tile)
+			),
+			"TileView",
+			EUserInterfaceActionType::RadioButton
+		);
+
+		ViewOptions.AddMenuEntry(
+			LOCTEXT("ListViewOption", "List"),
+			LOCTEXT("ListViewOptionToolTip", "View assets in a list with thumbnails."),
+			FSlateIcon(),
+			FUIAction(
+				FExecuteAction::CreateSP(this, &SLibraryView::SetCurrentViewTypeFromMenu, EAssetViewType::List),
+				FCanExecuteAction(),
+				FIsActionChecked::CreateSP(this, &SLibraryView::IsCurrentViewType, EAssetViewType::List)
+			),
+			"ListView",
+			EUserInterfaceActionType::RadioButton
+		);
+		ViewOptions.EndSection();
+	}
+
+	{
+		ViewOptions.BeginSection("AssetThumbnails", LOCTEXT("ThumbnailsHeading", "Thumbnails"));
+
+		auto CreateThumbnailSizeSubMenu = [this](FMenuBuilder& SubMenu)
+		{
+			for (int32 EnumValue = (int32)EThumbnailSize::Tiny; EnumValue < (int32)EThumbnailSize::MAX; ++EnumValue)
+			{
+				SubMenu.AddMenuEntry(
+					SAssetView::ThumbnailSizeToDisplayName((EThumbnailSize)EnumValue),
+					FText::GetEmpty(),
+					FSlateIcon(),
+					FUIAction(
+						FExecuteAction::CreateSP(this, &SLibraryView::OnThumbnailSizeChanged, (EThumbnailSize)EnumValue),
+						FCanExecuteAction(),
+						FIsActionChecked::CreateSP(this, &SLibraryView::IsThumbnailSizeChecked, (EThumbnailSize)EnumValue)
+					),
+					NAME_None,
+					EUserInterfaceActionType::RadioButton
+				);
+			}
+		};
+
+		ViewOptions.AddSubMenu(
+			LOCTEXT("ThumbnailSize", "Thumbnail Size"),
+			LOCTEXT("ThumbnailSizeToolTip", "Adjust the size of thumbnails."),
+			FNewMenuDelegate::CreateLambda(CreateThumbnailSizeSubMenu)
+		);
+		ViewOptions.EndSection();
+	}
+
+	return ViewOptions.MakeWidget();
+}
+
 void SLibraryView::Construct(const FArguments& InArgs, TSharedPtr<FWidgetBlueprintEditor> InBlueprintEditor)
 {
 	BlueprintEditor = InBlueprintEditor;
@@ -130,10 +196,10 @@ void SLibraryView::Construct(const FArguments& InArgs, TSharedPtr<FWidgetBluepri
 	LibraryViewModel->OnUpdating.AddRaw(this, &SLibraryView::OnViewModelUpdating);
 	LibraryViewModel->OnUpdated.AddRaw(this, &SLibraryView::OnViewModelUpdated);
 
-	WidgetFilter = MakeShareable(new WidgetViewModelTextFilter(
-		WidgetViewModelTextFilter::FItemToStringArray::CreateSP(this, &SLibraryView::GetWidgetFilterStrings)));
+	WidgetFilter = MakeShared<WidgetViewModelDelegateFilter>(
+		WidgetViewModelDelegateFilter::FPredicate::CreateSP(this, &SLibraryView::HandleFilterWidgetView));
 
-	FilterHandler = MakeShareable(new LibraryFilterHandler());
+	FilterHandler = MakeShared<LibraryFilterHandler>();
 	FilterHandler->SetFilter(WidgetFilter.Get());
 	FilterHandler->SetRootItems(&(LibraryViewModel->GetWidgetViewModels()), &TreeWidgetViewModels);
 	FilterHandler->SetGetChildrenDelegate(LibraryFilterHandler::FOnGetChildren::CreateRaw(this, &SLibraryView::OnGetChildren));
@@ -145,10 +211,9 @@ void SLibraryView::Construct(const FArguments& InArgs, TSharedPtr<FWidgetBluepri
 		.OnGetChildren(FilterHandler.ToSharedRef(), &LibraryFilterHandler::OnGetFilteredChildren)
 		.OnSelectionChanged(this, &SLibraryView::WidgetLibrary_OnSelectionChanged)
 		.TreeItemsSource(&TreeWidgetViewModels);
-		
 
 	FilterHandler->SetTreeView(WidgetTemplatesView.Get());
-
+	
 	ChildSlot
 	[
 		SNew(SVerticalBox)
@@ -157,9 +222,35 @@ void SLibraryView::Construct(const FArguments& InArgs, TSharedPtr<FWidgetBluepri
 		.Padding(4)
 		.AutoHeight()
 		[
-			SAssignNew(SearchBoxPtr, SSearchBox)
-			.HintText(LOCTEXT("SearchTemplates", "Search Library"))
-			.OnTextChanged(this, &SLibraryView::OnSearchChanged)
+			SNew(SHorizontalBox)
+			+ SHorizontalBox::Slot()
+			.FillWidth(1.0f)
+			.VAlign(VAlign_Center)
+			[
+				SAssignNew(SearchBoxPtr, SSearchBox)
+				.HintText(LOCTEXT("SearchTemplates", "Search Library"))
+				.OnTextChanged(this, &SLibraryView::OnSearchChanged)
+			]
+			+ SHorizontalBox::Slot()
+			.AutoWidth()
+			.Padding(2, 0, 0, 0)
+			[
+				SNew(SComboButton)
+				.ContentPadding(0)
+				.ComboButtonStyle(&FAppStyle::Get().GetWidgetStyle<FComboButtonStyle>("SimpleComboButton"))
+				.HasDownArrow(false)
+				.AddMetaData<FTagMetaData>(FTagMetaData(TEXT("ViewOptions")))
+				.ButtonContent()
+				[
+					SNew(SImage)
+					.ColorAndOpacity(FSlateColor::UseForeground())
+					.Image(FAppStyle::Get().GetBrush("Icons.Settings"))
+				]
+				.MenuContent()
+				[
+					ConstructViewOptions()
+				]
+			]
 		]
 
 		+ SVerticalBox::Slot()
@@ -178,6 +269,8 @@ void SLibraryView::Construct(const FArguments& InArgs, TSharedPtr<FWidgetBluepri
 	];
 
 	bRefreshRequested = true;
+	LastViewType = EAssetViewType::Type::Tile;
+	LastThumbnailSize = EThumbnailSize::Small;
 
 	LibraryViewModel->Update();
 	LoadItemExpansion();
@@ -204,8 +297,6 @@ void SLibraryView::OnSearchChanged(const FText& InFilterText)
 {
 	bRefreshRequested = true;
 	FilterHandler->SetIsEnabled(!InFilterText.IsEmpty());
-	WidgetFilter->SetRawFilterText(InFilterText);
-	SearchBoxPtr->SetError(WidgetFilter->GetFilterErrorText());
 	LibraryViewModel->SetSearchText(InFilterText);
 }
 
@@ -287,6 +378,43 @@ void SLibraryView::SaveItemExpansion()
 	}
 }
 
+bool SLibraryView::HandleFilterWidgetView(TSharedPtr<FWidgetViewModel> WidgetViewModel)
+{
+	return WidgetViewModel->HasFilteredChildTemplates();
+}
+
+void SLibraryView::SetCurrentViewTypeFromMenu(EAssetViewType::Type ViewType)
+{
+	LastViewType = ViewType;
+
+	for (TSharedPtr<FWidgetViewModel>& ViewModel : LibraryViewModel->GetWidgetTemplateListViewModels())
+	{
+		TSharedPtr<FWidgetTemplateListViewModel> TemplateListViewModel = StaticCastSharedPtr<FWidgetTemplateListViewModel>(ViewModel);
+		TemplateListViewModel->SetViewType(ViewType);
+	}
+}
+
+bool SLibraryView::IsCurrentViewType(EAssetViewType::Type ViewType)
+{
+	return LastViewType == ViewType;
+}
+
+void SLibraryView::OnThumbnailSizeChanged(EThumbnailSize ThumbnailSize)
+{
+	LastThumbnailSize = ThumbnailSize;
+
+	for (TSharedPtr<FWidgetViewModel>& ViewModel : LibraryViewModel->GetWidgetTemplateListViewModels())
+	{
+		TSharedPtr<FWidgetTemplateListViewModel> TemplateListViewModel = StaticCastSharedPtr<FWidgetTemplateListViewModel>(ViewModel);
+		TemplateListViewModel->SetThumbnailSize(ThumbnailSize);
+	}
+}
+
+bool SLibraryView::IsThumbnailSizeChecked(EThumbnailSize ThumbnailSize)
+{
+	return LastThumbnailSize == ThumbnailSize;
+}
+
 void SLibraryView::OnGetChildren(TSharedPtr<FWidgetViewModel> Item, TArray< TSharedPtr<FWidgetViewModel> >& Children)
 {
 	return Item->GetChildren(Children);
@@ -330,9 +458,4 @@ void SLibraryView::Tick(const FGeometry& AllottedGeometry, const double InCurren
 	}
 }
 
-void SLibraryView::GetWidgetFilterStrings(TSharedPtr<FWidgetViewModel> WidgetViewModel, TArray<FString>& OutStrings)
-{
-	WidgetViewModel->GetFilterStrings(OutStrings);
-}
-
 #undef LOCTEXT_NAMESPACE
diff --git a/Engine/Source/Editor/UMGEditor/Private/Library/SLibraryView.h b/Engine/Source/Editor/UMGEditor/Private/Library/SLibraryView.h
--- a/Engine/Source/Editor/UMGEditor/Private/Library/SLibraryView.h
+++ b/Engine/Source/Editor/UMGEditor/Private/Library/SLibraryView.h
@@ -6,9 +6,10 @@
 #include "Widgets/DeclarativeSyntaxSupport.h"
 #include "Widgets/SCompoundWidget.h"
 #include "WidgetBlueprintEditor.h"
-#include "Misc/TextFilter.h"
+#include "Misc/DelegateFilter.h"
 #include "Widgets/Views/STreeView.h"
 #include "Framework/Views/TreeFilterHandler.h"
+#include "IContentBrowserSingleton.h"
 
 class FWidgetTemplate;
 class UWidgetBlueprint;
@@ -72,7 +73,7 @@ private:
 class SLibraryView : public SCompoundWidget
 {
 public:
-	typedef TTextFilter<TSharedPtr<FWidgetViewModel>> WidgetViewModelTextFilter;
+	typedef TDelegateFilter<TSharedPtr<FWidgetViewModel>> WidgetViewModelDelegateFilter;
 
 public:
 	SLATE_BEGIN_ARGS( SLibraryView ){}
@@ -83,9 +84,6 @@ public:
 	
 	virtual void Tick( const FGeometry& AllottedGeometry, const double InCurrentTime, const float InDeltaTime ) override;
 
-	/** Gets the text currently displayed in the search box. */
-	FText GetSearchText() const;
-
 	/** On Selected template widget changed in Library view */
 	void WidgetLibrary_OnSelectionChanged(TSharedPtr<FWidgetViewModel> SelectedItem, ESelectInfo::Type SelectInfo);
 
@@ -94,6 +92,9 @@ public:
 
 private:
 
+	/** Builds settings menu. */
+	TSharedRef<SWidget> ConstructViewOptions();
+
 	void OnGetChildren(TSharedPtr<FWidgetViewModel> Item, TArray< TSharedPtr<FWidgetViewModel> >& Children);
 	TSharedRef<ITableRow> OnGenerateWidgetTemplateItem(TSharedPtr<FWidgetViewModel> Item, const TSharedRef<STableViewBase>& OwnerTable);
 
@@ -110,8 +111,20 @@ private:
 	/** Save the expansion state for the TreeView */
 	void SaveItemExpansion();
 
-	/** Gets an array of strings used for filtering/searching the specified widget. */
-	void GetWidgetFilterStrings(TSharedPtr<FWidgetViewModel> WidgetViewModel, TArray<FString>& OutStrings);
+	/** Callback to determine if a particular library view passes filtering */
+	bool HandleFilterWidgetView(TSharedPtr<FWidgetViewModel> WidgetViewModel);
+
+	/** Callback to set library asset view type */
+	void SetCurrentViewTypeFromMenu(EAssetViewType::Type ViewType);
+
+	/** Callback to determine library asset view type */
+	bool IsCurrentViewType(EAssetViewType::Type ViewType);
+
+	/** Callback to set library asset thumbnail size */
+	void OnThumbnailSizeChanged(EThumbnailSize InThumbnailSize);
+
+	/** Callback to determine library asset thumbnail size */
+	bool IsThumbnailSizeChecked(EThumbnailSize InThumbnailSize);
 
 	TWeakPtr<class FWidgetBlueprintEditor> BlueprintEditor;
 	TSharedPtr<FLibraryViewModel> LibraryViewModel;
@@ -130,7 +143,7 @@ private:
 	TSharedPtr<class SSearchBox> SearchBoxPtr;
 
 	/** The filter instance which is used by the TreeFilterHandler to filter the TreeView. */
-	TSharedPtr<WidgetViewModelTextFilter> WidgetFilter;
+	TSharedPtr<WidgetViewModelDelegateFilter> WidgetFilter;
 
 	/** Expended Items in the Tree view */
 	TSet<TSharedPtr<FWidgetViewModel>> ExpandedItems;
@@ -140,4 +153,11 @@ private:
 
 	/** Are editor widgets supported. */
 	bool bAllowEditorWidget;
+
+	/** Last view type set by user */
+	EAssetViewType::Type LastViewType;
+
+	/** Last view type set by user */
+	EThumbnailSize LastThumbnailSize;
+
 };
diff --git a/Engine/Source/Editor/UMGEditor/Private/Library/SLibraryViewModel.cpp b/Engine/Source/Editor/UMGEditor/Private/Library/SLibraryViewModel.cpp
--- a/Engine/Source/Editor/UMGEditor/Private/Library/SLibraryViewModel.cpp
+++ b/Engine/Source/Editor/UMGEditor/Private/Library/SLibraryViewModel.cpp
@@ -25,6 +25,7 @@
 #include "Settings/WidgetDesignerSettings.h"
 #include "UMGEditorProjectSettings.h"
 #include "WidgetPaletteFavorites.h"
+#include "SLibraryViewModel.h"
 
 #define LOCTEXT_NAMESPACE "UMG"
 
@@ -75,6 +76,17 @@ void FLibraryViewModel::RemoveFromFavorites(const FWidgetTemplateViewModel* Widg
 	Favorites->Remove(WidgetTemplateViewModel->GetName().ToString());
 }
 
+void FLibraryViewModel::SetSearchText(const FText& InSearchText)
+{
+	FFavortiesViewModel::SetSearchText(InSearchText);
+
+	for (TSharedPtr<FWidgetViewModel>& WidgetTemplateListViewModel : WidgetTemplateListViewModels)
+	{
+		TSharedPtr<FWidgetTemplateListViewModel> WidgetLibrary = StaticCastSharedPtr<FWidgetTemplateListViewModel>(WidgetTemplateListViewModel);
+		WidgetLibrary->SetSearchText(InSearchText);
+	}
+}
+
 void FLibraryViewModel::Update()
 {
 	if (bRebuildRequested)
@@ -86,7 +98,6 @@ void FLibraryViewModel::Update()
 	}
 }
 
-
 UWidgetBlueprint* FLibraryViewModel::GetBlueprint() const
 {
 	if (BlueprintEditor.IsValid())
@@ -103,6 +114,7 @@ void FLibraryViewModel::BuildWidgetList()
 	// Clear the current list of view models and categories
 	WidgetViewModels.Reset();
 	WidgetTemplateCategories.Reset();
+	WidgetTemplateListViewModels.Reset();
 
 	// Generate a list of templates
 	BuildClassWidgetList();
@@ -121,11 +133,12 @@ void FLibraryViewModel::BuildWidgetList()
 		TSharedPtr<FWidgetHeaderViewModel> Header = MakeShareable(new FWidgetHeaderViewModel());
 		Header->GroupName = FText::FromString(Entry.Key);
 
-
-		TSharedPtr<FWidgetTemplateListViewModel> TemplateViewModel = MakeShareable(new FWidgetTemplateListViewModel());
-		TemplateViewModel->Templates = Entry.Value;
+		TSharedPtr<FWidgetTemplateListViewModel> TemplateViewModel = MakeShared<FWidgetTemplateListViewModel>();
+		TemplateViewModel->ConstructListView(Entry.Value);
 		Header->Children.Add(TemplateViewModel);
+		WidgetTemplateListViewModels.Add(TemplateViewModel);
 
+		// @TODO: DarenC - Reference for when implementing favorites system
 		//for ( auto& Template : Entry.Value )
 		//{
 		//	TSharedPtr<FWidgetTemplateViewModel> TemplateViewModel = MakeShareable(new FWidgetTemplateViewModel());
@@ -428,60 +441,49 @@ void FLibraryViewModel::HandleOnAssetsDeleted(const TArray<UClass*>& DeletedAsse
 	}
 }
 
-#undef LOCTEXT_NAMESPACE
-
 FWidgetTemplateListViewModel::FWidgetTemplateListViewModel()
 {
 }
 
-FText FWidgetTemplateListViewModel::GetName() const
+void FWidgetTemplateListViewModel::ConstructListView(TArray<TSharedPtr<FWidgetTemplate>> InTemplates)
 {
-	return FText();
-}
+	Templates = InTemplates;
 
-bool FWidgetTemplateListViewModel::IsTemplate() const
-{
-	return false;
-}
-
-void FWidgetTemplateListViewModel::GetFilterStrings(TArray<FString>& OutStrings) const
-{
-}
-
-TSharedRef<ITableRow> FWidgetTemplateListViewModel::BuildRow(const TSharedRef<STableViewBase>& OwnerTable)
-{
-	// @TODO: DarenC - Temp construct on demand. Move to construct earlier?
 	if (!TemplatesFilter)
 	{
 		// Generate filter text
 		bool bHasFilters = false;
-		TStringBuilder<2048> FilterString; 
+		TStringBuilder<2048> FilterString;
 		for (TSharedPtr<FWidgetTemplate> Template : Templates)
 		{
 			if (TSharedPtr<FWidgetTemplateClass> TemplateClass = StaticCastSharedPtr<FWidgetTemplateClass>(Template))
 			{
 				TWeakObjectPtr<UClass> WidgetClass = TemplateClass->GetWidgetClass();
-				if (WidgetClass.IsValid())
+				FName TemplateName = WidgetClass.IsValid() ? WidgetClass->GetFName() : TemplateClass->GetWidgetAssetData().AssetName;
+				FString TemplateString = TemplateName.ToString();
+
+				if (!TemplateString.IsEmpty())
 				{
-					FString TemplateName = WidgetClass->GetFName().ToString();
-					if (!TemplateName.IsEmpty())
-					{
-						TemplateName.RemoveFromEnd(TEXT("_C"));
-						FilterString += bHasFilters ? "|+" : "+";
-						FilterString += TemplateName;
-						bHasFilters = true;
-					}
+					TemplateString.RemoveFromEnd(TEXT("_C"));
+					FilterString += bHasFilters ? "|+" : "+";
+					FilterString += TemplateString;
+					bHasFilters = true;
 				}
 			}
 		}
 
-		TSharedPtr<FFrontendFilter_Text> WidgetTextFilter = MakeShared<FFrontendFilter_Text>();
+		WidgetTextFilter = MakeShared<FFrontendFilter_Text>();
 		WidgetTextFilter->SetActive(true);
 		WidgetTextFilter->SetIncludeClassName(true);
 		WidgetTextFilter->SetRawFilterText(FText::FromString(FilterString.ToString()));
+		CachedLowercaseWidgetFilter = WidgetTextFilter->GetRawFilterText().ToString().ToLower();
+
+		SearchFilter = MakeShared<FFrontendFilter_Text>();
+		SearchFilter->SetActive(false);
 
 		TemplatesFilter = MakeShared<FAssetFilterCollectionType>();
 		TemplatesFilter->Add(WidgetTextFilter);
+		TemplatesFilter->Add(SearchFilter);
 	}
 
 	if (!AssetViewPtr)
@@ -493,17 +495,86 @@ TSharedRef<ITableRow> FWidgetTemplateListViewModel::BuildRow(const TSharedRef<ST
 		AssetViewPtr = SNew(SAssetView)
 			.InitialCategoryFilter(EContentBrowserItemCategoryFilter::IncludeAssets | EContentBrowserItemCategoryFilter::IncludeClasses)
 			.InitialBackendFilter(WidgetTemplateFilter)
+			.InitialThumbnailSize(EThumbnailSize::Small)
 			.FrontendFilters(TemplatesFilter)
 			.ForceShowEngineContent(true)
 			.ForceShowPluginContent(true)
+			.ShowTypeInTileView(false)
+			.ShowViewOptions(false)
+			.HighlightedText(this, &FWidgetTemplateListViewModel::GetSearchText)
 			;
 	}
+
 	AssetViewPtr->RequestSlowFullListRefresh();
+}
 
+FText FWidgetTemplateListViewModel::GetName() const
+{
+	return FText();
+}
+
+bool FWidgetTemplateListViewModel::IsTemplate() const
+{
+	return false;
+}
+
+void FWidgetTemplateListViewModel::GetFilterStrings(TArray<FString>& OutStrings) const
+{
+}
+
+bool FWidgetTemplateListViewModel::HasFilteredChildTemplates() const
+{
+	FText SearchText = SearchFilter->GetRawFilterText();
+	return SearchText.IsEmpty() || CachedLowercaseWidgetFilter.Contains(SearchText.ToString().ToLower());
+}
+
+TSharedRef<ITableRow> FWidgetTemplateListViewModel::BuildRow(const TSharedRef<STableViewBase>& OwnerTable)
+{
 	return SNew(STableRow<TSharedPtr<FWidgetViewModel>>, OwnerTable)
 		.Padding(2.0f)
-		//.OnDragDetected(this, &FWidgetTemplateViewModel::OnDraggingWidgetTemplateItem)
+		.ShowSelection(false)
+		.ShowWires(false)
+		.Style(FEditorStyle::Get(), "UMGEditor.LibraryView")
 		[
 			AssetViewPtr.ToSharedRef()
 		];
 }
+
+void FWidgetTemplateListViewModel::SetViewType(EAssetViewType::Type ViewType)
+{
+	if (AssetViewPtr)
+	{
+		AssetViewPtr->SetCurrentViewType(ViewType);
+	}
+}
+
+void FWidgetTemplateListViewModel::SetThumbnailSize(EThumbnailSize ThumbnailSize)
+{
+	if (AssetViewPtr)
+	{
+		AssetViewPtr->SetCurrentThumbnailSize(ThumbnailSize);
+	}
+}
+
+void FWidgetTemplateListViewModel::SetSearchText(const FText& InSearchText)
+{
+	if (SearchFilter && AssetViewPtr)
+	{
+		SearchFilter->SetActive(!InSearchText.IsEmpty());
+		SearchFilter->SetRawFilterText(InSearchText);
+
+		AssetViewPtr->SetUserSearching(!InSearchText.IsEmpty());
+	}
+}
+
+FText FWidgetTemplateListViewModel::GetSearchText() const
+{
+	if (SearchFilter)
+	{
+		return SearchFilter->GetRawFilterText();
+	}
+
+	return FText::GetEmpty();
+}
+
+#undef LOCTEXT_NAMESPACE
\ No newline at end of file
diff --git a/Engine/Source/Editor/UMGEditor/Private/Library/SLibraryViewModel.h b/Engine/Source/Editor/UMGEditor/Private/Library/SLibraryViewModel.h
--- a/Engine/Source/Editor/UMGEditor/Private/Library/SLibraryViewModel.h
+++ b/Engine/Source/Editor/UMGEditor/Private/Library/SLibraryViewModel.h
@@ -22,15 +22,23 @@ class FWidgetTemplateListViewModel : public FWidgetViewModel
 public:
 	FWidgetTemplateListViewModel();
 
-	virtual FText GetName() const override;
+	/** Builds list view for given templates */
+	void ConstructListView(TArray<TSharedPtr<FWidgetTemplate>> InTemplates);
 
+	//~ Begin FWidgetViewModel Interface
+	virtual FText GetName() const override;
 	virtual bool IsTemplate() const override;
-
 	virtual void GetFilterStrings(TArray<FString>& OutStrings) const override;
-
+	virtual bool HasFilteredChildTemplates() const override;
 	virtual TSharedRef<ITableRow> BuildRow(const TSharedRef<STableViewBase>& OwnerTable) override;
+	//~ End FWidgetViewModel Interface
+
+	void SetViewType(EAssetViewType::Type ViewType);
+
+	void SetThumbnailSize(EThumbnailSize ThumbnailSize);
 
-	// FReply OnDraggingWidgetTemplateItem(const FGeometry& MyGeometry, const FPointerEvent& MouseEvent);
+	void SetSearchText(const FText& InSearchText);
+	FText GetSearchText() const;
 
 	TArray<TSharedPtr<FWidgetTemplate>> Templates;
 	TSharedPtr<FAssetFilterCollectionType> TemplatesFilter;
@@ -38,6 +46,15 @@ private:
 
 	/** The asset view widget */
 	TSharedPtr<SAssetView> AssetViewPtr;
+
+	/** Filter we use to narrow down the libary to select widgets asset view */
+	TSharedPtr<class FFrontendFilter_Text> WidgetTextFilter;
+
+	/** Filter we forward our search text to for the asset view */
+	TSharedPtr<class FFrontendFilter_Text> SearchFilter;
+
+	/** Lowercase string used by widget text filter */
+	FString CachedLowercaseWidgetFilter;
 };
 
 class FLibraryViewModel : public FFavortiesViewModel
@@ -60,14 +77,15 @@ public:
 	/** Returns true if the view model needs to be updated */
 	bool NeedUpdate() const { return bRebuildRequested; }
 	   
-	/** Add the widget template to the list of favorites */
+	//~ Begin FFavortiesViewModel Interface
 	virtual void AddToFavorites(const FWidgetTemplateViewModel* WidgetTemplateViewModel) override;
-
-	/** Remove the widget template to the list of favorites */
 	virtual void RemoveFromFavorites(const FWidgetTemplateViewModel* WidgetTemplateViewModel) override;
+	virtual void SetSearchText(const FText& InSearchText) override;
+	//~ End FFavortiesViewModel Interface
 
 	typedef TArray< TSharedPtr<FWidgetViewModel> > ViewModelsArray;
 	ViewModelsArray& GetWidgetViewModels() { return WidgetViewModels; }
+	ViewModelsArray& GetWidgetTemplateListViewModels() { return WidgetTemplateListViewModels; }
 
 	/** Fires before the view model is updated */
 	FOnUpdating OnUpdating;
@@ -109,6 +127,9 @@ private:
 
 	/** The source root view models for the tree. */
 	ViewModelsArray WidgetViewModels;
+
+	/** Widget template list view models associated with each root. */
+	ViewModelsArray WidgetTemplateListViewModels;
 	   
 	/** Controls rebuilding the list of spawnable widgets */
 	bool bRebuildRequested;
diff --git a/Engine/Source/Editor/UMGEditor/Private/Palette/SPaletteViewModel.h b/Engine/Source/Editor/UMGEditor/Private/Palette/SPaletteViewModel.h
--- a/Engine/Source/Editor/UMGEditor/Private/Palette/SPaletteViewModel.h
+++ b/Engine/Source/Editor/UMGEditor/Private/Palette/SPaletteViewModel.h
@@ -28,6 +28,11 @@ public:
 	/** @param OutStrings - Returns an array of strings used for filtering/searching this item. */
 	virtual void GetFilterStrings(TArray<FString>& OutStrings) const = 0;
 
+	virtual bool HasFilteredChildTemplates() const
+	{
+		return false;
+	}
+
 	virtual TSharedRef<ITableRow> BuildRow(const TSharedRef<STableViewBase>& OwnerTable) = 0;
 
 	virtual void GetChildren(TArray< TSharedPtr<FWidgetViewModel> >& OutChildren)
@@ -102,6 +107,18 @@ public:
 		// it's widgets filtered out, so return an empty filter string.
 	}
 
+	virtual bool HasFilteredChildTemplates() const override
+	{
+		for (const TSharedPtr<FWidgetViewModel>& Child : Children)
+		{
+			if (Child && Child->HasFilteredChildTemplates())
+			{
+				return true;
+			}
+		}
+		return false;
+	}
+
 	virtual TSharedRef<ITableRow> BuildRow(const TSharedRef<STableViewBase>& OwnerTable) override;
 
 	virtual void GetChildren(TArray< TSharedPtr<FWidgetViewModel> >& OutChildren) override;
@@ -129,7 +146,7 @@ public:
 	/** Remove the widget template to the list of favorites */
 	virtual void RemoveFromFavorites(const FWidgetTemplateViewModel* WidgetTemplateViewModel) = 0;
 
-	void SetSearchText(const FText& inSearchText) { SearchText = inSearchText; }
+	virtual void SetSearchText(const FText& InSearchText) { SearchText = InSearchText; }
 	FText GetSearchText() const { return SearchText; }
 
 protected:
diff --git a/Engine/Source/Editor/UnrealEd/Public/AssetThumbnail.h b/Engine/Source/Editor/UnrealEd/Public/AssetThumbnail.h
--- a/Engine/Source/Editor/UnrealEd/Public/AssetThumbnail.h
+++ b/Engine/Source/Editor/UnrealEd/Public/AssetThumbnail.h
@@ -27,6 +27,17 @@ namespace EThumbnailLabel
 	};
 };
 
+enum class EThumbnailSize : uint8
+{
+	Tiny = 0,
+	Small,
+	Medium,
+	Large,
+	Huge,
+
+	// Not a size
+	MAX
+};
 
 /** A struct containing details about how the asset thumbnail should behave */
 struct FAssetThumbnailConfig
