Fix SLibraryView Non-Unity

diff --git a/Engine/Source/Editor/UMGEditor/Private/Library/SLibraryView.cpp b/Engine/Source/Editor/UMGEditor/Private/Library/SLibraryView.cpp
--- a/Engine/Source/Editor/UMGEditor/Private/Library/SLibraryView.cpp
+++ b/Engine/Source/Editor/UMGEditor/Private/Library/SLibraryView.cpp
@@ -7,6 +7,7 @@
 #include "UObject/UObjectIterator.h"
 #include "WidgetBlueprint.h"
 #include "Editor.h"
+#include "Framework/MultiBox/MultiBoxBuilder.h"
 #include "Widgets/Images/SImage.h"
 #include "Widgets/Layout/SScrollBorder.h"
 #include "SAssetView.h"
