Cleanup Widget Library

diff --git a/Engine/Source/Editor/UMGEditor/Private/Library/SLibraryView.cpp b/Engine/Source/Editor/UMGEditor/Private/Library/SLibraryView.cpp
--- a/Engine/Source/Editor/UMGEditor/Private/Library/SLibraryView.cpp
+++ b/Engine/Source/Editor/UMGEditor/Private/Library/SLibraryView.cpp
@@ -114,11 +114,6 @@ void SLibraryViewItem::Construct(const FArguments& InArgs, TSharedPtr<FWidgetTem
 		];
 }
 
-FReply SLibraryViewItem::OnMouseButtonDoubleClick(const FGeometry& InMyGeometry, const FPointerEvent& InMouseEvent)
-{
-	return WidgetViewModel->Template->OnDoubleClicked();
-};
-
 TSharedRef<SWidget> SLibraryView::ConstructViewOptions()
 {
 	FMenuBuilder ViewOptions(true, nullptr);
@@ -208,9 +203,8 @@ void SLibraryView::Construct(const FArguments& InArgs, TSharedPtr<FWidgetBluepri
 	SAssignNew(WidgetTemplatesView, STreeView< TSharedPtr<FWidgetViewModel> >)
 		.ItemHeight(1.0f)
 		.SelectionMode(ESelectionMode::Single)
-		.OnGenerateRow(this, &SLibraryView::OnGenerateWidgetTemplateItem)
+		.OnGenerateRow(this, &SLibraryView::OnGenerateWidgetTemplateLibrary)
 		.OnGetChildren(FilterHandler.ToSharedRef(), &LibraryFilterHandler::OnGetFilteredChildren)
-		.OnSelectionChanged(this, &SLibraryView::WidgetLibrary_OnSelectionChanged)
 		.TreeItemsSource(&TreeWidgetViewModels);
 
 	FilterHandler->SetTreeView(WidgetTemplatesView.Get());
@@ -301,61 +295,6 @@ void SLibraryView::OnSearchChanged(const FText& InFilterText)
 	LibraryViewModel->SetSearchText(InFilterText);
 }
 
-void SLibraryView::WidgetLibrary_OnSelectionChanged(TSharedPtr<FWidgetViewModel> SelectedItem, ESelectInfo::Type SelectInfo)
-{
-	if (!SelectedItem.IsValid()) 
-	{
-		return;
-	}
-
-	// Reset the selected
-	BlueprintEditor.Pin()->SetSelectedTemplate(nullptr);
-	BlueprintEditor.Pin()->SetSelectedUserWidget(FAssetData());
-
-	// If it's not a template, return
-	if (!SelectedItem->IsTemplate())
-	{
-		return;
-	}
-
-	TSharedPtr<FWidgetTemplateViewModel> SelectedTemplate = StaticCastSharedPtr<FWidgetTemplateViewModel>(SelectedItem);
-	if (SelectedTemplate.IsValid())
-	{
-		TSharedPtr<FWidgetTemplateClass> TemplateClass = StaticCastSharedPtr<FWidgetTemplateClass>(SelectedTemplate->Template);
-		if (TemplateClass.IsValid())
-		{
-			if (TemplateClass->GetWidgetClass().IsValid())
-			{
-				BlueprintEditor.Pin()->SetSelectedTemplate(TemplateClass->GetWidgetClass());
-			}
-			else
-			{
-				TSharedPtr<FWidgetTemplateBlueprintClass> UserCreatedTemplate = StaticCastSharedPtr<FWidgetTemplateBlueprintClass>(TemplateClass);
-				if (UserCreatedTemplate.IsValid())
-				{
-					// Then pass in the asset data of selected widget
-					FAssetData UserCreatedWidget = UserCreatedTemplate->GetWidgetAssetData();
-					BlueprintEditor.Pin()->SetSelectedUserWidget(UserCreatedWidget);
-				}
-			}
-		}
-	}
-}
-
-TSharedPtr<FWidgetTemplate> SLibraryView::GetSelectedTemplateWidget() const
-{
-	TArray<TSharedPtr<FWidgetViewModel>> SelectedTemplates = WidgetTemplatesView.Get()->GetSelectedItems();
-	if (SelectedTemplates.Num() == 1)
-	{
-		TSharedPtr<FWidgetTemplateViewModel> TemplateViewModel = StaticCastSharedPtr<FWidgetTemplateViewModel>(SelectedTemplates[0]);
-		if (TemplateViewModel.IsValid())
-		{
-			return TemplateViewModel->Template;
-		}
-	}
-	return nullptr;
-}
-
 void SLibraryView::LoadItemExpansion()
 {
 	// Restore the expansion state of the widget groups.
@@ -421,7 +360,7 @@ void SLibraryView::OnGetChildren(TSharedPtr<FWidgetViewModel> Item, TArray< TSha
 	return Item->GetChildren(Children);
 }
 
-TSharedRef<ITableRow> SLibraryView::OnGenerateWidgetTemplateItem(TSharedPtr<FWidgetViewModel> Item, const TSharedRef<STableViewBase>& OwnerTable)
+TSharedRef<ITableRow> SLibraryView::OnGenerateWidgetTemplateLibrary(TSharedPtr<FWidgetViewModel> Item, const TSharedRef<STableViewBase>& OwnerTable)
 {
 	return Item->BuildRow(OwnerTable);
 }
diff --git a/Engine/Source/Editor/UMGEditor/Private/Library/SLibraryView.h b/Engine/Source/Editor/UMGEditor/Private/Library/SLibraryView.h
--- a/Engine/Source/Editor/UMGEditor/Private/Library/SLibraryView.h
+++ b/Engine/Source/Editor/UMGEditor/Private/Library/SLibraryView.h
@@ -63,8 +63,6 @@ private:
 
 	EVisibility GetFavoritedStateVisibility() const;
 
-	virtual FReply OnMouseButtonDoubleClick(const FGeometry& InMyGeometry, const FPointerEvent& InMouseEvent) override;
-
 private:
 	TSharedPtr<FWidgetTemplateViewModel> WidgetViewModel;
 };
@@ -84,19 +82,13 @@ public:
 	
 	virtual void Tick( const FGeometry& AllottedGeometry, const double InCurrentTime, const float InDeltaTime ) override;
 
-	/** On Selected template widget changed in Library view */
-	void WidgetLibrary_OnSelectionChanged(TSharedPtr<FWidgetViewModel> SelectedItem, ESelectInfo::Type SelectInfo);
-
-	/** Gets the selected template widget in Library */
-	TSharedPtr<FWidgetTemplate> GetSelectedTemplateWidget() const;
-
 private:
 
 	/** Builds settings menu. */
 	TSharedRef<SWidget> ConstructViewOptions();
 
 	void OnGetChildren(TSharedPtr<FWidgetViewModel> Item, TArray< TSharedPtr<FWidgetViewModel> >& Children);
-	TSharedRef<ITableRow> OnGenerateWidgetTemplateItem(TSharedPtr<FWidgetViewModel> Item, const TSharedRef<STableViewBase>& OwnerTable);
+	TSharedRef<ITableRow> OnGenerateWidgetTemplateLibrary(TSharedPtr<FWidgetViewModel> Item, const TSharedRef<STableViewBase>& OwnerTable);
 
 	/** Called when the filter text is changed. */
 	void OnSearchChanged(const FText& InFilterText);
diff --git a/Engine/Source/Editor/UMGEditor/Private/Library/SLibraryViewModel.cpp b/Engine/Source/Editor/UMGEditor/Private/Library/SLibraryViewModel.cpp
--- a/Engine/Source/Editor/UMGEditor/Private/Library/SLibraryViewModel.cpp
+++ b/Engine/Source/Editor/UMGEditor/Private/Library/SLibraryViewModel.cpp
@@ -29,39 +29,47 @@
 
 #define LOCTEXT_NAMESPACE "UMG"
 
-FLibraryViewModel::FLibraryViewModel(TSharedPtr<FWidgetBlueprintEditor> InBlueprintEditor)
-	: bRebuildRequested(true)
+void FLibraryViewModel::BuildWidgetTemplateCategory(FString& Category, TArray<TSharedPtr<FWidgetTemplate>>& Templates)
 {
-	BlueprintEditor = InBlueprintEditor;
+	TSharedPtr<FWidgetHeaderViewModel> Header = MakeShareable(new FWidgetHeaderViewModel());
+	Header->GroupName = FText::FromString(Category);
 
-	FavoriteHeader = MakeShareable(new FWidgetHeaderViewModel());
-	FavoriteHeader->GroupName = LOCTEXT("Favorites", "Favorites");
-}
+	TSharedPtr<FWidgetTemplateListViewModel> TemplateViewModel = MakeShared<FWidgetTemplateListViewModel>();
+	TemplateViewModel->ConstructListView(Templates);
+	Header->Children.Add(TemplateViewModel);
+	WidgetTemplateListViewModels.Add(TemplateViewModel);
 
-void FLibraryViewModel::RegisterToEvents()
-{
-	// Register for events that can trigger a Library rebuild
-	GEditor->OnBlueprintReinstanced().AddRaw(this, &FLibraryViewModel::OnBlueprintReinstanced);
-	FEditorDelegates::OnAssetsDeleted.AddSP(this, &FLibraryViewModel::HandleOnAssetsDeleted);
-	FCoreUObjectDelegates::ReloadCompleteDelegate.AddSP(this, &FLibraryViewModel::OnReloadComplete);
+	// @TODO: DarenC - Reference for when implementing favorites system
+	//for ( auto& Template : Entry.Value )
+	//{
+	//	TSharedPtr<FWidgetTemplateViewModel> TemplateViewModel = MakeShareable(new FWidgetTemplateViewModel());
+	//	TemplateViewModel->Template = Template;
+	//	TemplateViewModel->FavortiesViewModel = this;
+	//	Header->Children.Add(TemplateViewModel);
 
-	// register for any objects replaced
-	FCoreUObjectDelegates::OnObjectsReplaced.AddRaw(this, &FLibraryViewModel::OnObjectsReplaced);
+	//	// If it's a favorite, we also add it to the Favorite section
+	//	int32 index = FavoritesList.Find(Template->Name.ToString());
+	//	if (index != INDEX_NONE)
+	//	{
+	//		TemplateViewModel->SetFavorite();
 
-	// Register for favorite list update to handle the case where a favorite is added in another window of the UMG Designer
-	UWidgetPaletteFavorites* Favorites = GetDefault<UWidgetDesignerSettings>()->Favorites;
-	Favorites->OnFavoritesUpdated.AddSP(this, &FLibraryViewModel::OnFavoritesUpdated);
-}
+	//		// We have to create a second copy of the ViewModel for the treeview has it doesn't support to have the same element twice.
+	//		TSharedPtr<FWidgetTemplateViewModel> FavoriteTemplateViewModel = MakeShareable(new FWidgetTemplateViewModel());
+	//		FavoriteTemplateViewModel->Template = Template;
+	//		FavoriteTemplateViewModel->FavortiesViewModel = this;
+	//		FavoriteTemplateViewModel->SetFavorite();
 
-FLibraryViewModel::~FLibraryViewModel()
-{
-	GEditor->OnBlueprintReinstanced().RemoveAll(this);
-	FEditorDelegates::OnAssetsDeleted.RemoveAll(this);
-	FCoreUObjectDelegates::ReloadCompleteDelegate.RemoveAll(this);
-	FCoreUObjectDelegates::OnObjectsReplaced.RemoveAll(this);
+	//		FavoriteHeader->Children.Add(FavoriteTemplateViewModel);
 
-	UWidgetPaletteFavorites* Favorites = GetDefault<UWidgetDesignerSettings>()->Favorites;
-	Favorites->OnFavoritesUpdated.RemoveAll(this);
+	//		// Remove the favorite from the temporary list
+	//		FavoritesList.RemoveAt(index);
+	//	}
+
+	//}
+
+	Header->Children.Sort([](TSharedPtr<FWidgetViewModel> L, TSharedPtr<FWidgetViewModel> R) { return R->GetName().CompareTo(L->GetName()) > 0; });
+
+	WidgetViewModels.Add(Header);
 }
 
 void FLibraryViewModel::AddToFavorites(const FWidgetTemplateViewModel* WidgetTemplateViewModel)
@@ -78,7 +86,7 @@ void FLibraryViewModel::RemoveFromFavorites(const FWidgetTemplateViewModel* Widg
 
 void FLibraryViewModel::SetSearchText(const FText& InSearchText)
 {
-	FFavortiesViewModel::SetSearchText(InSearchText);
+	FWidgetCatalogViewModel::SetSearchText(InSearchText);
 
 	for (TSharedPtr<FWidgetViewModel>& WidgetTemplateListViewModel : WidgetTemplateListViewModels)
 	{
@@ -87,358 +95,10 @@ void FLibraryViewModel::SetSearchText(const FText& InSearchText)
 	}
 }
 
-void FLibraryViewModel::Update()
-{
-	if (bRebuildRequested)
-	{
-		OnUpdating.Broadcast();
-		BuildWidgetList();
-		bRebuildRequested = false;
-		OnUpdated.Broadcast();
-	}
-}
-
-UWidgetBlueprint* FLibraryViewModel::GetBlueprint() const
-{
-	if (BlueprintEditor.IsValid())
-	{
-		UBlueprint* BP = BlueprintEditor.Pin()->GetBlueprintObj();
-		return Cast<UWidgetBlueprint>(BP);
-	}
-
-	return NULL;
-}
-
 void FLibraryViewModel::BuildWidgetList()
 {
-	// Clear the current list of view models and categories
-	WidgetViewModels.Reset();
-	WidgetTemplateCategories.Reset();
 	WidgetTemplateListViewModels.Reset();
-
-	// Generate a list of templates
-	BuildClassWidgetList();
-
-	// Clear the Favorite section
-	bool bHasFavorites = FavoriteHeader->Children.Num() != 0;
-	FavoriteHeader->Children.Reset();
-	
-	// Copy of the list of favorites to be able to do some cleanup in the real list
-	UWidgetPaletteFavorites* FavoritesLibrary = GetDefault<UWidgetDesignerSettings>()->Favorites;
-	TArray<FString> FavoritesList = FavoritesLibrary->GetFavorites();
-
-	// For each entry in the category create a view model for the widget template
-	for ( auto& Entry : WidgetTemplateCategories )
-	{
-		TSharedPtr<FWidgetHeaderViewModel> Header = MakeShareable(new FWidgetHeaderViewModel());
-		Header->GroupName = FText::FromString(Entry.Key);
-
-		TSharedPtr<FWidgetTemplateListViewModel> TemplateViewModel = MakeShared<FWidgetTemplateListViewModel>();
-		TemplateViewModel->ConstructListView(Entry.Value);
-		Header->Children.Add(TemplateViewModel);
-		WidgetTemplateListViewModels.Add(TemplateViewModel);
-
-		// @TODO: DarenC - Reference for when implementing favorites system
-		//for ( auto& Template : Entry.Value )
-		//{
-		//	TSharedPtr<FWidgetTemplateViewModel> TemplateViewModel = MakeShareable(new FWidgetTemplateViewModel());
-		//	TemplateViewModel->Template = Template;
-		//	TemplateViewModel->FavortiesViewModel = this;
-		//	Header->Children.Add(TemplateViewModel);
-
-		//	// If it's a favorite, we also add it to the Favorite section
-		//	int32 index = FavoritesList.Find(Template->Name.ToString());
-		//	if (index != INDEX_NONE)
-		//	{
-		//		TemplateViewModel->SetFavorite();
-
-		//		// We have to create a second copy of the ViewModel for the treeview has it doesn't support to have the same element twice.
-		//		TSharedPtr<FWidgetTemplateViewModel> FavoriteTemplateViewModel = MakeShareable(new FWidgetTemplateViewModel());
-		//		FavoriteTemplateViewModel->Template = Template;
-		//		FavoriteTemplateViewModel->FavortiesViewModel = this;
-		//		FavoriteTemplateViewModel->SetFavorite();
-
-		//		FavoriteHeader->Children.Add(FavoriteTemplateViewModel);
-
-		//		// Remove the favorite from the temporary list
-		//		FavoritesList.RemoveAt(index);
-		//	}
-
-		//}
-
-		Header->Children.Sort([] (TSharedPtr<FWidgetViewModel> L, TSharedPtr<FWidgetViewModel> R) { return R->GetName().CompareTo(L->GetName()) > 0; });
-
-		WidgetViewModels.Add(Header);
-	}	
-
-	// Remove all Favorites that may be left in the list.Typically happening when the list of favorite contains widget that were deleted since the last opening.
-	for (const FString& favoriteName : FavoritesList)
-	{
-		FavoritesLibrary->Remove(favoriteName);
-	}
-
-	// Sort the view models by name
-	WidgetViewModels.Sort([] (TSharedPtr<FWidgetViewModel> L, TSharedPtr<FWidgetViewModel> R) { return R->GetName().CompareTo(L->GetName()) > 0; });
-
-	// Add the Favorite section at the top
-	if (FavoriteHeader->Children.Num() != 0)
-	{
-		// We force expansion of the favorite header when we add favorites for the first time.
-		FavoriteHeader->SetForceExpansion(!bHasFavorites);
-		FavoriteHeader->Children.Sort([](TSharedPtr<FWidgetViewModel> L, TSharedPtr<FWidgetViewModel> R) { return R->GetName().CompareTo(L->GetName()) > 0; });
-		WidgetViewModels.Insert(FavoriteHeader, 0);
-	}
-	
-	// Take the Advanced Section, and put it at the end.
-	TSharedPtr<FWidgetViewModel>* advancedSectionPtr = WidgetViewModels.FindByPredicate([](TSharedPtr<FWidgetViewModel> widget) {return widget->GetName().CompareTo(LOCTEXT("Advanced", "Advanced")) == 0; });
-	if (advancedSectionPtr)
-	{
-		TSharedPtr<FWidgetViewModel> advancedSection = *advancedSectionPtr;
-		WidgetViewModels.Remove(advancedSection);
-		WidgetViewModels.Push(advancedSection);
-	}
-}
-
-void FLibraryViewModel::BuildClassWidgetList()
-{
-	static const FName DevelopmentStatusKey(TEXT("DevelopmentStatus"));
-
-	TMap<FName, TSubclassOf<UUserWidget>> LoadedWidgetBlueprintClassesByName;
-
-	auto ActiveWidgetBlueprintClass = GetBlueprint()->GeneratedClass;
-	FName ActiveWidgetBlueprintClassName = ActiveWidgetBlueprintClass->GetFName();
-
-	TArray<FSoftClassPath> WidgetClassesToHide = GetDefault<UUMGEditorProjectSettings>()->WidgetClassesToHide;
-
-	// Locate all UWidget classes from code and loaded widget BPs
-	for (TObjectIterator<UClass> ClassIt; ClassIt; ++ClassIt)
-	{
-		UClass* WidgetClass = *ClassIt;
-
-		if (!FWidgetBlueprintEditorUtils::IsUsableWidgetClass(WidgetClass))
-		{
-			continue;
-		}
-
-		if (WidgetClass->HasAnyClassFlags(CLASS_HideDropDown))
-		{
-			continue;
-		}
-
-		// Initialize AssetData for checking PackagePath
-		FAssetData WidgetAssetData = FAssetData(WidgetClass);
-
-		// Excludes engine content if user sets it to false
-		if (!GetDefault<UContentBrowserSettings>()->GetDisplayEngineFolder() || !GetDefault<UUMGEditorProjectSettings>()->bShowWidgetsFromEngineContent)
-		{
-			if (WidgetAssetData.PackagePath.ToString().Find(TEXT("/Engine")) == 0)
-			{
-				continue;
-			}
-		}
-
-		// Excludes developer content if user sets it to false
-		if (!GetDefault<UContentBrowserSettings>()->GetDisplayDevelopersFolder() || !GetDefault<UUMGEditorProjectSettings>()->bShowWidgetsFromDeveloperContent)
-		{
-			if (WidgetAssetData.PackagePath.ToString().Find(TEXT("/Game/Developers")) == 0)
-			{
-				continue;
-			}
-		}
-
-		// Excludes this widget if it is on the hide list
-		bool bIsOnList = false;
-		for (FSoftClassPath Widget : WidgetClassesToHide)
-		{
-			if (WidgetAssetData.ObjectPath.ToString().Find(Widget.ToString()) == 0)
-			{
-				bIsOnList = true;
-				break;
-			}
-		}
-		if (bIsOnList)
-		{
-			continue;
-		}
-
-		const bool bIsSameClass = WidgetClass->GetFName() == ActiveWidgetBlueprintClassName;
-
-		// Check that the asset that generated this class is valid (necessary b/c of a larger issue wherein force delete does not wipe the generated class object)
-		if ( bIsSameClass )
-		{
-			continue;
-		}
-
-		if (WidgetClass->IsChildOf(UUserWidget::StaticClass()))
-		{
-			if ( WidgetClass->ClassGeneratedBy )
-			{
-				// Track the widget blueprint classes that are already loaded
-				LoadedWidgetBlueprintClassesByName.Add(WidgetClass->ClassGeneratedBy->GetFName()) = WidgetClass;
-			}
-		}
-		else
-		{
-			TSharedPtr<FWidgetTemplateClass> Template = MakeShareable(new FWidgetTemplateClass(WidgetClass));
-
-			AddWidgetTemplate(Template);
-		}
-
-		//TODO UMG does not prevent deep nested circular references
-	}
-
-	// Locate all widget BP assets (include unloaded)
-	const FAssetRegistryModule& AssetRegistryModule = FModuleManager::LoadModuleChecked<FAssetRegistryModule>(TEXT("AssetRegistry"));
-	TArray<FAssetData> AllBPsAssetData;
-	AssetRegistryModule.Get().GetAssetsByClass(UBlueprint::StaticClass()->GetFName(), AllBPsAssetData, true);
-
-	for (FAssetData& BPAssetData : AllBPsAssetData)
-	{
-		// Blueprints get the class type actions for their parent native class - this avoids us having to load the blueprint
-		UClass* ParentClass = nullptr;
-		FString ParentClassName;
-		if (!BPAssetData.GetTagValue(FBlueprintTags::NativeParentClassPath, ParentClassName))
-		{
-			BPAssetData.GetTagValue(FBlueprintTags::ParentClassPath, ParentClassName);
-		}
-		if (!ParentClassName.IsEmpty())
-		{
-			UObject* Outer = nullptr;
-			ResolveName(Outer, ParentClassName, false, false);
-			ParentClass = FindObject<UClass>(ANY_PACKAGE, *ParentClassName);
-			// UUserWidgets have their own loading section, and we don't want to process any blueprints that don't have UWidget parents
-			if (!ParentClass->IsChildOf(UWidget::StaticClass()) || ParentClass->IsChildOf(UUserWidget::StaticClass()))
-			{
-				continue;
-			}
-		}
-
-		if (!FilterAssetData(BPAssetData))
-		{
-			// If this object isn't currently loaded, add it to the Library view
-			if (BPAssetData.ToSoftObjectPath().ResolveObject() == nullptr)
-			{
-				auto Template = MakeShareable(new FWidgetTemplateClass(BPAssetData, nullptr));
-				AddWidgetTemplate(Template);
-			}
-		}
-	}
-
-	TArray<FAssetData> AllWidgetBPsAssetData;
-	AssetRegistryModule.Get().GetAssetsByClass(UWidgetBlueprint::StaticClass()->GetFName(), AllWidgetBPsAssetData, true);
-
-	FName ActiveWidgetBlueprintName = ActiveWidgetBlueprintClass->ClassGeneratedBy->GetFName();
-	for (FAssetData& WidgetBPAssetData : AllWidgetBPsAssetData)
-	{
-		// Excludes the blueprint you're currently in
-		if (WidgetBPAssetData.AssetName == ActiveWidgetBlueprintName)
-		{
-			continue;
-		}
-
-		if (!FilterAssetData(WidgetBPAssetData))
-		{
-			// Excludes this widget if it is on the hide list
-			bool bIsOnList = false;
-			for (FSoftClassPath Widget : WidgetClassesToHide)
-			{
-				if (Widget.ToString().Find(WidgetBPAssetData.ObjectPath.ToString()) == 0)
-				{
-					bIsOnList = true;
-					break;
-				}
-			}
-			if (bIsOnList)
-			{
-				continue;
-			}
-
-			// If the blueprint generated class was found earlier, pass it to the template
-			TSubclassOf<UUserWidget> WidgetBPClass = nullptr;
-			auto LoadedWidgetBPClass = LoadedWidgetBlueprintClassesByName.Find(WidgetBPAssetData.AssetName);
-			if (LoadedWidgetBPClass)
-			{
-				WidgetBPClass = *LoadedWidgetBPClass;
-			}
-
-			uint32 BPFlags = WidgetBPAssetData.GetTagValueRef<uint32>(FBlueprintTags::ClassFlags);
-			if ((BPFlags & (CLASS_Abstract | CLASS_Deprecated | CLASS_HideDropDown)) == 0)
-			{
-				auto Template = MakeShareable(new FWidgetTemplateBlueprintClass(WidgetBPAssetData, WidgetBPClass));
-
-				AddWidgetTemplate(Template);
-			}
-		}
-	}
-}
-
-bool FLibraryViewModel::FilterAssetData(FAssetData &InAssetData)
-{
-	// Excludes engine content if user sets it to false
-	if (!GetDefault<UContentBrowserSettings>()->GetDisplayEngineFolder() || !GetDefault<UUMGEditorProjectSettings>()->bShowWidgetsFromEngineContent)
-	{
-		if (InAssetData.PackagePath.ToString().Find(TEXT("/Engine")) == 0)
-		{
-			return true;
-		}
-	}
-
-	// Excludes developer content if user sets it to false
-	if (!GetDefault<UContentBrowserSettings>()->GetDisplayDevelopersFolder() || !GetDefault<UUMGEditorProjectSettings>()->bShowWidgetsFromDeveloperContent)
-	{
-		if (InAssetData.PackagePath.ToString().Find(TEXT("/Game/Developers")) == 0)
-		{
-			return true;
-		}
-	}
-	return false;
-}
-
-void FLibraryViewModel::AddWidgetTemplate(TSharedPtr<FWidgetTemplate> Template)
-{
-	FString Category = Template->GetCategory().ToString();
-
-	// Hide user specific categories
-	TArray<FString> CategoriesToHide = GetDefault<UUMGEditorProjectSettings>()->CategoriesToHide;
-	for (FString CategoryName : CategoriesToHide)
-	{
-		if (Category == CategoryName)
-		{
-			return;
-		}
-	}
-	WidgetTemplateArray& Group = WidgetTemplateCategories.FindOrAdd(Category);
-	Group.Add(Template);
-}
-
-void FLibraryViewModel::OnObjectsReplaced(const TMap<UObject*, UObject*>& ReplacementMap)
-{
-}
-
-void FLibraryViewModel::OnBlueprintReinstanced()
-{
-	bRebuildRequested = true;
-}
-
-void FLibraryViewModel::OnFavoritesUpdated()
-{
-	bRebuildRequested = true;
-}
-
-void FLibraryViewModel::OnReloadComplete(EReloadCompleteReason Reason)
-{
-	bRebuildRequested = true;
-}
-
-void FLibraryViewModel::HandleOnAssetsDeleted(const TArray<UClass*>& DeletedAssetClasses)
-{
-	for (auto DeletedAssetClass : DeletedAssetClasses)
-	{
-		if (DeletedAssetClass->IsChildOf(UWidgetBlueprint::StaticClass()))
-		{
-			bRebuildRequested = true;
-		}
-	}
+	FWidgetCatalogViewModel::BuildWidgetList();
 }
 
 FWidgetTemplateListViewModel::FWidgetTemplateListViewModel()
diff --git a/Engine/Source/Editor/UMGEditor/Private/Library/SLibraryViewModel.h b/Engine/Source/Editor/UMGEditor/Private/Library/SLibraryViewModel.h
--- a/Engine/Source/Editor/UMGEditor/Private/Library/SLibraryViewModel.h
+++ b/Engine/Source/Editor/UMGEditor/Private/Library/SLibraryViewModel.h
@@ -57,83 +57,27 @@ private:
 	FString CachedLowercaseWidgetFilter;
 };
 
-class FLibraryViewModel : public FFavortiesViewModel
+class FLibraryViewModel : public FWidgetCatalogViewModel
 {
 public:
+	FLibraryViewModel(TSharedPtr<FWidgetBlueprintEditor> InBlueprintEditor) : FWidgetCatalogViewModel(InBlueprintEditor) { }
 
-	DECLARE_MULTICAST_DELEGATE(FOnUpdating)
-	DECLARE_MULTICAST_DELEGATE(FOnUpdated)
-
-public:
-	FLibraryViewModel(TSharedPtr<FWidgetBlueprintEditor> InBlueprintEditor);
-	virtual ~FLibraryViewModel();
-
-	/** Register the View Model to events that should trigger a update of the Library*/
-	void RegisterToEvents();
-
-	/** Update the view model if needed and returns true if it did. */
-	void Update();
-
-	/** Returns true if the view model needs to be updated */
-	bool NeedUpdate() const { return bRebuildRequested; }
-	   
-	//~ Begin FFavortiesViewModel Interface
+	//~ Begin FWidgetCatalogViewModel Interface
+	virtual void BuildWidgetTemplateCategory(FString& Category, TArray<TSharedPtr<FWidgetTemplate>>& Templates) override;
 	virtual void AddToFavorites(const FWidgetTemplateViewModel* WidgetTemplateViewModel) override;
 	virtual void RemoveFromFavorites(const FWidgetTemplateViewModel* WidgetTemplateViewModel) override;
 	virtual void SetSearchText(const FText& InSearchText) override;
-	//~ End FFavortiesViewModel Interface
+	//~ End FWidgetCatalogViewModel Interface
 
-	typedef TArray< TSharedPtr<FWidgetViewModel> > ViewModelsArray;
 	ViewModelsArray& GetWidgetViewModels() { return WidgetViewModels; }
 	ViewModelsArray& GetWidgetTemplateListViewModels() { return WidgetTemplateListViewModels; }
 
-	/** Fires before the view model is updated */
-	FOnUpdating OnUpdating;
-
-	/** Fires after the view model is updated */
-	FOnUpdated OnUpdated;
-
-private:
-	FLibraryViewModel() {};
-
-	UWidgetBlueprint* GetBlueprint() const;
-
-	void BuildWidgetList();
-	void BuildClassWidgetList();
-
-	static bool FilterAssetData(FAssetData &BPAssetData);
-
-	void AddWidgetTemplate(TSharedPtr<FWidgetTemplate> Template);
-
-	/** Called when a Blueprint is recompiled and live objects are swapped out for replacements */
-	void OnObjectsReplaced(const TMap<UObject*, UObject*>& ReplacementMap);
-
-	/** Requests a rebuild of the widget list if a widget blueprint was compiled */
-	void OnBlueprintReinstanced();
-
-	/** Called when the favorite list is changed */
-	void OnFavoritesUpdated();
-
-	/** Requests a rebuild of the widget list */
-	void OnReloadComplete(EReloadCompleteReason Reason);
-
-	/** Requests a rebuild of the widget list if a widget blueprint was deleted */
-	void HandleOnAssetsDeleted(const TArray<UClass*>& DeletedAssetClasses);
+protected:
+	//~ Begin FWidgetCatalogViewModel Interface
+	virtual void BuildWidgetList() override;
+	//~ End FWidgetCatalogViewModel Interface
 	
-	TWeakPtr<class FWidgetBlueprintEditor> BlueprintEditor;
-
-	typedef TArray<TSharedPtr<FWidgetTemplate>> WidgetTemplateArray;
-	TMap<FString, WidgetTemplateArray> WidgetTemplateCategories;
-
-	/** The source root view models for the tree. */
-	ViewModelsArray WidgetViewModels;
-
 	/** Widget template list view models associated with each root. */
 	ViewModelsArray WidgetTemplateListViewModels;
-	   
-	/** Controls rebuilding the list of spawnable widgets */
-	bool bRebuildRequested;
-
-	TSharedPtr<FWidgetHeaderViewModel> FavoriteHeader;
 };
 
diff --git a/Engine/Source/Editor/UMGEditor/Private/Palette/SPaletteViewModel.cpp b/Engine/Source/Editor/UMGEditor/Private/Palette/SPaletteViewModel.cpp
--- a/Engine/Source/Editor/UMGEditor/Private/Palette/SPaletteViewModel.cpp
+++ b/Engine/Source/Editor/UMGEditor/Private/Palette/SPaletteViewModel.cpp
@@ -54,7 +54,7 @@ TSharedRef<ITableRow> FWidgetTemplateViewModel::BuildRow(const TSharedRef<STable
 		.OnDragDetected(this, &FWidgetTemplateViewModel::OnDraggingWidgetTemplateItem)
 		[
 			SNew(SPaletteViewItem, SharedThis(this))
-			.HighlightText(FavortiesViewModel, &FFavortiesViewModel::GetSearchText)
+			.HighlightText(FavortiesViewModel, &FWidgetCatalogViewModel::GetSearchText)
 		];
 }
 
@@ -97,7 +97,7 @@ void FWidgetHeaderViewModel::GetChildren(TArray< TSharedPtr<FWidgetViewModel> >&
 	}
 }
 
-FPaletteViewModel::FPaletteViewModel(TSharedPtr<FWidgetBlueprintEditor> InBlueprintEditor)
+FWidgetCatalogViewModel::FWidgetCatalogViewModel(TSharedPtr<FWidgetBlueprintEditor> InBlueprintEditor)
 	: bRebuildRequested(true)
 {
 	BlueprintEditor = InBlueprintEditor;
@@ -106,22 +106,22 @@ FPaletteViewModel::FPaletteViewModel(TSharedPtr<FWidgetBlueprintEditor> InBluepr
 	FavoriteHeader->GroupName = LOCTEXT("Favorites", "Favorites");
 }
 
-void FPaletteViewModel::RegisterToEvents()
+void FWidgetCatalogViewModel::RegisterToEvents()
 {
 	// Register for events that can trigger a palette rebuild
-	GEditor->OnBlueprintReinstanced().AddRaw(this, &FPaletteViewModel::OnBlueprintReinstanced);
-	FEditorDelegates::OnAssetsDeleted.AddSP(this, &FPaletteViewModel::HandleOnAssetsDeleted);
-	FCoreUObjectDelegates::ReloadCompleteDelegate.AddSP(this, &FPaletteViewModel::OnReloadComplete);
+	GEditor->OnBlueprintReinstanced().AddRaw(this, &FWidgetCatalogViewModel::OnBlueprintReinstanced);
+	FEditorDelegates::OnAssetsDeleted.AddSP(this, &FWidgetCatalogViewModel::HandleOnAssetsDeleted);
+	FCoreUObjectDelegates::ReloadCompleteDelegate.AddSP(this, &FWidgetCatalogViewModel::OnReloadComplete);
 
 	// register for any objects replaced
-	FCoreUObjectDelegates::OnObjectsReplaced.AddRaw(this, &FPaletteViewModel::OnObjectsReplaced);
+	FCoreUObjectDelegates::OnObjectsReplaced.AddRaw(this, &FWidgetCatalogViewModel::OnObjectsReplaced);
 
 	// Register for favorite list update to handle the case where a favorite is added in another window of the UMG Designer
 	UWidgetPaletteFavorites* Favorites = GetDefault<UWidgetDesignerSettings>()->Favorites;
-	Favorites->OnFavoritesUpdated.AddSP(this, &FPaletteViewModel::OnFavoritesUpdated);
+	Favorites->OnFavoritesUpdated.AddSP(this, &FWidgetCatalogViewModel::OnFavoritesUpdated);
 }
 
-FPaletteViewModel::~FPaletteViewModel()
+FWidgetCatalogViewModel::~FWidgetCatalogViewModel()
 {
 	GEditor->OnBlueprintReinstanced().RemoveAll(this);
 	FEditorDelegates::OnAssetsDeleted.RemoveAll(this);
@@ -132,19 +132,7 @@ FPaletteViewModel::~FPaletteViewModel()
 	Favorites->OnFavoritesUpdated.RemoveAll(this);
 }
 
-void FPaletteViewModel::AddToFavorites(const FWidgetTemplateViewModel* WidgetTemplateViewModel)
-{
-	UWidgetPaletteFavorites* Favorites = GetDefault<UWidgetDesignerSettings>()->Favorites;
-	Favorites->Add(WidgetTemplateViewModel->GetName().ToString());
-}
-
-void FPaletteViewModel::RemoveFromFavorites(const FWidgetTemplateViewModel* WidgetTemplateViewModel)
-{
-	UWidgetPaletteFavorites* Favorites = GetDefault<UWidgetDesignerSettings>()->Favorites;
-	Favorites->Remove(WidgetTemplateViewModel->GetName().ToString());
-}
-
-void FPaletteViewModel::Update()
+void FWidgetCatalogViewModel::Update()
 {
 	if (bRebuildRequested)
 	{
@@ -155,8 +143,7 @@ void FPaletteViewModel::Update()
 	}
 }
 
-
-UWidgetBlueprint* FPaletteViewModel::GetBlueprint() const
+UWidgetBlueprint* FWidgetCatalogViewModel::GetBlueprint() const
 {
 	if (BlueprintEditor.IsValid())
 	{
@@ -167,7 +154,7 @@ UWidgetBlueprint* FPaletteViewModel::GetBlueprint() const
 	return NULL;
 }
 
-void FPaletteViewModel::BuildWidgetList()
+void FWidgetCatalogViewModel::BuildWidgetList()
 {
 	// Clear the current list of view models and categories
 	WidgetViewModels.Reset();
@@ -187,39 +174,7 @@ void FPaletteViewModel::BuildWidgetList()
 	// For each entry in the category create a view model for the widget template
 	for ( auto& Entry : WidgetTemplateCategories )
 	{
-		TSharedPtr<FWidgetHeaderViewModel> Header = MakeShareable(new FWidgetHeaderViewModel());
-		Header->GroupName = FText::FromString(Entry.Key);
-
-		for ( auto& Template : Entry.Value )
-		{
-			TSharedPtr<FWidgetTemplateViewModel> TemplateViewModel = MakeShareable(new FWidgetTemplateViewModel());
-			TemplateViewModel->Template = Template;
-			TemplateViewModel->FavortiesViewModel = this;
-			Header->Children.Add(TemplateViewModel);
-
-			// If it's a favorite, we also add it to the Favorite section
-			int32 index = FavoritesList.Find(Template->Name.ToString());
-			if (index != INDEX_NONE)
-			{
-				TemplateViewModel->SetFavorite();
-
-				// We have to create a second copy of the ViewModel for the treeview has it doesn't support to have the same element twice.
-				TSharedPtr<FWidgetTemplateViewModel> FavoriteTemplateViewModel = MakeShareable(new FWidgetTemplateViewModel());
-				FavoriteTemplateViewModel->Template = Template;
-				FavoriteTemplateViewModel->FavortiesViewModel = this;
-				FavoriteTemplateViewModel->SetFavorite();
-
-				FavoriteHeader->Children.Add(FavoriteTemplateViewModel);
-
-				// Remove the favorite from the temporary list
-				FavoritesList.RemoveAt(index);
-			}
-
-		}
-
-		Header->Children.Sort([] (TSharedPtr<FWidgetViewModel> L, TSharedPtr<FWidgetViewModel> R) { return R->GetName().CompareTo(L->GetName()) > 0; });
-
-		WidgetViewModels.Add(Header);
+		BuildWidgetTemplateCategory(Entry.Key, Entry.Value);
 	}	
 
 	// Remove all Favorites that may be left in the list.Typically happening when the list of favorite contains widget that were deleted since the last opening.
@@ -250,7 +205,7 @@ void FPaletteViewModel::BuildWidgetList()
 	}
 }
 
-void FPaletteViewModel::BuildClassWidgetList()
+void FWidgetCatalogViewModel::BuildClassWidgetList()
 {
 	static const FName DevelopmentStatusKey(TEXT("DevelopmentStatus"));
 
@@ -423,7 +378,7 @@ void FPaletteViewModel::BuildClassWidgetList()
 	}
 }
 
-bool FPaletteViewModel::FilterAssetData(FAssetData &InAssetData)
+bool FWidgetCatalogViewModel::FilterAssetData(FAssetData &InAssetData)
 {
 	// Excludes engine content if user sets it to false
 	if (!GetDefault<UContentBrowserSettings>()->GetDisplayEngineFolder() || !GetDefault<UUMGEditorProjectSettings>()->bShowWidgetsFromEngineContent)
@@ -445,7 +400,7 @@ bool FPaletteViewModel::FilterAssetData(FAssetData &InAssetData)
 	return false;
 }
 
-void FPaletteViewModel::AddWidgetTemplate(TSharedPtr<FWidgetTemplate> Template)
+void FWidgetCatalogViewModel::AddWidgetTemplate(TSharedPtr<FWidgetTemplate> Template)
 {
 	FString Category = Template->GetCategory().ToString();
 
@@ -462,26 +417,26 @@ void FPaletteViewModel::AddWidgetTemplate(TSharedPtr<FWidgetTemplate> Template)
 	Group.Add(Template);
 }
 
-void FPaletteViewModel::OnObjectsReplaced(const TMap<UObject*, UObject*>& ReplacementMap)
+void FWidgetCatalogViewModel::OnObjectsReplaced(const TMap<UObject*, UObject*>& ReplacementMap)
 {
 }
 
-void FPaletteViewModel::OnBlueprintReinstanced()
+void FWidgetCatalogViewModel::OnBlueprintReinstanced()
 {
 	bRebuildRequested = true;
 }
 
-void FPaletteViewModel::OnFavoritesUpdated()
+void FWidgetCatalogViewModel::OnFavoritesUpdated()
 {
 	bRebuildRequested = true;
 }
 
-void FPaletteViewModel::OnReloadComplete(EReloadCompleteReason Reason)
+void FWidgetCatalogViewModel::OnReloadComplete(EReloadCompleteReason Reason)
 {
 	bRebuildRequested = true;
 }
 
-void FPaletteViewModel::HandleOnAssetsDeleted(const TArray<UClass*>& DeletedAssetClasses)
+void FWidgetCatalogViewModel::HandleOnAssetsDeleted(const TArray<UClass*>& DeletedAssetClasses)
 {
 	for (auto DeletedAssetClass : DeletedAssetClasses)
 	{
@@ -492,4 +447,57 @@ void FPaletteViewModel::HandleOnAssetsDeleted(const TArray<UClass*>& DeletedAsse
 	}
 }
 
+void FPaletteViewModel::BuildWidgetTemplateCategory(FString& Category, TArray<TSharedPtr<FWidgetTemplate>>& Templates)
+{
+	TSharedPtr<FWidgetHeaderViewModel> Header = MakeShareable(new FWidgetHeaderViewModel());
+	Header->GroupName = FText::FromString(Category);
+
+	// Copy of the list of favorites to be able to do some cleanup in the real list
+	UWidgetPaletteFavorites* FavoritesPalette = GetDefault<UWidgetDesignerSettings>()->Favorites;
+	TArray<FString> FavoritesList = FavoritesPalette->GetFavorites();
+
+	for (auto& Template : Templates)
+	{
+		TSharedPtr<FWidgetTemplateViewModel> TemplateViewModel = MakeShareable(new FWidgetTemplateViewModel());
+		TemplateViewModel->Template = Template;
+		TemplateViewModel->FavortiesViewModel = this;
+		Header->Children.Add(TemplateViewModel);
+
+		// If it's a favorite, we also add it to the Favorite section
+		int32 index = FavoritesList.Find(Template->Name.ToString());
+		if (index != INDEX_NONE)
+		{
+			TemplateViewModel->SetFavorite();
+
+			// We have to create a second copy of the ViewModel for the treeview has it doesn't support to have the same element twice.
+			TSharedPtr<FWidgetTemplateViewModel> FavoriteTemplateViewModel = MakeShareable(new FWidgetTemplateViewModel());
+			FavoriteTemplateViewModel->Template = Template;
+			FavoriteTemplateViewModel->FavortiesViewModel = this;
+			FavoriteTemplateViewModel->SetFavorite();
+
+			FavoriteHeader->Children.Add(FavoriteTemplateViewModel);
+
+			// Remove the favorite from the temporary list
+			FavoritesList.RemoveAt(index);
+		}
+
+	}
+
+	Header->Children.Sort([](TSharedPtr<FWidgetViewModel> L, TSharedPtr<FWidgetViewModel> R) { return R->GetName().CompareTo(L->GetName()) > 0; });
+
+	WidgetViewModels.Add(Header);
+}
+
+void FPaletteViewModel::AddToFavorites(const FWidgetTemplateViewModel* WidgetTemplateViewModel)
+{
+	UWidgetPaletteFavorites* Favorites = GetDefault<UWidgetDesignerSettings>()->Favorites;
+	Favorites->Add(WidgetTemplateViewModel->GetName().ToString());
+}
+
+void FPaletteViewModel::RemoveFromFavorites(const FWidgetTemplateViewModel* WidgetTemplateViewModel)
+{
+	UWidgetPaletteFavorites* Favorites = GetDefault<UWidgetDesignerSettings>()->Favorites;
+	Favorites->Remove(WidgetTemplateViewModel->GetName().ToString());
+}
+
 #undef LOCTEXT_NAMESPACE
diff --git a/Engine/Source/Editor/UMGEditor/Private/Palette/SPaletteViewModel.h b/Engine/Source/Editor/UMGEditor/Private/Palette/SPaletteViewModel.h
--- a/Engine/Source/Editor/UMGEditor/Private/Palette/SPaletteViewModel.h
+++ b/Engine/Source/Editor/UMGEditor/Private/Palette/SPaletteViewModel.h
@@ -13,7 +13,7 @@ class FWidgetTemplate;
 class FWidgetBlueprintEditor;
 class UWidgetBlueprint;
 class SPaletteView;
-class FFavortiesViewModel;
+class FWidgetCatalogViewModel;
 
 /** View model for the items in the widget template list */
 class FWidgetViewModel : public TSharedFromThis<FWidgetViewModel>
@@ -78,7 +78,7 @@ public:
 	virtual void SetFavorite() override { bIsFavorite = true; }
 
 	TSharedPtr<FWidgetTemplate> Template;
-	FFavortiesViewModel* FavortiesViewModel;
+	FWidgetCatalogViewModel* FavortiesViewModel;
 private:
 	/** True is the widget is a favorite. It's keep as a state to prevent a search in the favorite list. */
 	bool bIsFavorite;
@@ -135,25 +135,7 @@ private:
 	bool bForceExpansion = false;
 };
 
-class FFavortiesViewModel : public TSharedFromThis<FFavortiesViewModel>
-{
-public:
-	virtual ~FFavortiesViewModel() { }
-
-	/** Add the widget template to the list of favorites */
-	virtual void AddToFavorites(const FWidgetTemplateViewModel* WidgetTemplateViewModel) = 0;
-
-	/** Remove the widget template to the list of favorites */
-	virtual void RemoveFromFavorites(const FWidgetTemplateViewModel* WidgetTemplateViewModel) = 0;
-
-	virtual void SetSearchText(const FText& InSearchText) { SearchText = InSearchText; }
-	FText GetSearchText() const { return SearchText; }
-
-protected:
-	FText SearchText;
-};
-
-class FPaletteViewModel : public FFavortiesViewModel
+class FWidgetCatalogViewModel : public TSharedFromThis<FWidgetCatalogViewModel>
 {
 public:
 
@@ -161,10 +143,10 @@ public:
 	DECLARE_MULTICAST_DELEGATE(FOnUpdated)
 
 public:
-	FPaletteViewModel(TSharedPtr<FWidgetBlueprintEditor> InBlueprintEditor);
-	virtual ~FPaletteViewModel();
+	FWidgetCatalogViewModel(TSharedPtr<FWidgetBlueprintEditor> InBlueprintEditor);
+	virtual ~FWidgetCatalogViewModel();
 
-	/** Register the View Model to events that should trigger a update of the Palette*/
+	/** Register the View Model to events that should trigger a update */
 	void RegisterToEvents();
 
 	/** Update the view model if needed and returns true if it did. */
@@ -172,12 +154,12 @@ public:
 
 	/** Returns true if the view model needs to be updated */
 	bool NeedUpdate() const { return bRebuildRequested; }
-	   
+
 	/** Add the widget template to the list of favorites */
-	void AddToFavorites(const FWidgetTemplateViewModel* WidgetTemplateViewModel);
+	virtual void AddToFavorites(const FWidgetTemplateViewModel* WidgetTemplateViewModel) = 0;
 
 	/** Remove the widget template to the list of favorites */
-	void RemoveFromFavorites(const FWidgetTemplateViewModel* WidgetTemplateViewModel);
+	virtual void RemoveFromFavorites(const FWidgetTemplateViewModel* WidgetTemplateViewModel) = 0;
 
 	typedef TArray< TSharedPtr<FWidgetViewModel> > ViewModelsArray;
 	ViewModelsArray& GetWidgetViewModels() { return WidgetViewModels; }
@@ -188,15 +170,19 @@ public:
 	/** Fires after the view model is updated */
 	FOnUpdated OnUpdated;
 
-private:
-	FPaletteViewModel() {};
+	virtual void SetSearchText(const FText& InSearchText) { SearchText = InSearchText; }
+	FText GetSearchText() const { return SearchText; }
+
+protected:
+	FWidgetCatalogViewModel() {};
 
 	UWidgetBlueprint* GetBlueprint() const;
 
-	void BuildWidgetList();
+	virtual void BuildWidgetList();
+	virtual void BuildWidgetTemplateCategory(FString& Category, TArray<TSharedPtr<FWidgetTemplate>>& Templates) = 0;
 	void BuildClassWidgetList();
 
-	static bool FilterAssetData(FAssetData &BPAssetData);
+	static bool FilterAssetData(FAssetData& BPAssetData);
 
 	void AddWidgetTemplate(TSharedPtr<FWidgetTemplate> Template);
 
@@ -227,5 +213,19 @@ private:
 	bool bRebuildRequested;
 
 	TSharedPtr<FWidgetHeaderViewModel> FavoriteHeader;
+
+	FText SearchText;
+};
+
+class FPaletteViewModel : public FWidgetCatalogViewModel
+{
+public:
+	FPaletteViewModel(TSharedPtr<FWidgetBlueprintEditor> InBlueprintEditor) : FWidgetCatalogViewModel(InBlueprintEditor) { }
+
+	//~ Begin FWidgetCatalogViewModel Interface
+	virtual void BuildWidgetTemplateCategory(FString& Category, TArray<TSharedPtr<FWidgetTemplate>>& Templates) override;
+	virtual void AddToFavorites(const FWidgetTemplateViewModel* WidgetTemplateViewModel) override;
+	virtual void RemoveFromFavorites(const FWidgetTemplateViewModel* WidgetTemplateViewModel) override;
+	//~ End FWidgetCatalogViewModel Interface
 };
 
diff --git a/Engine/Source/Editor/UMGEditor/Private/TabFactory/LibraryTabSummoner.cpp b/Engine/Source/Editor/UMGEditor/Private/TabFactory/LibraryTabSummoner.cpp
--- a/Engine/Source/Editor/UMGEditor/Private/TabFactory/LibraryTabSummoner.cpp
+++ b/Engine/Source/Editor/UMGEditor/Private/TabFactory/LibraryTabSummoner.cpp
@@ -13,7 +13,7 @@ FLibraryTabSummoner::FLibraryTabSummoner(TSharedPtr<class FWidgetBlueprintEditor
 		, BlueprintEditor(InBlueprintEditor)
 {
 	TabLabel = LOCTEXT("WidgetLibraryTabLabel", "Library");
-	TabIcon = FSlateIcon(FUMGStyle::GetStyleSetName(), "Palette.TabIcon");
+	TabIcon = FSlateIcon(FAppStyle::GetAppStyleSetName(), "Icons.FolderClosed");
 
 	bIsSingleton = true;
 
diff --git a/Engine/Source/Editor/UMGEditor/Private/WidgetBlueprintEditor.cpp b/Engine/Source/Editor/UMGEditor/Private/WidgetBlueprintEditor.cpp
--- a/Engine/Source/Editor/UMGEditor/Private/WidgetBlueprintEditor.cpp
+++ b/Engine/Source/Editor/UMGEditor/Private/WidgetBlueprintEditor.cpp
@@ -125,13 +125,13 @@ void FWidgetBlueprintEditor::InitWidgetBlueprintEditor(const EToolkitMode::Type
 
 	TSharedPtr<FWidgetBlueprintEditor> ThisPtr(SharedThis(this));
 
-	PaletteViewModel = MakeShareable(new FPaletteViewModel(ThisPtr));
+	PaletteViewModel = MakeShared<FPaletteViewModel>(ThisPtr);
 	PaletteViewModel->RegisterToEvents();
 
-	LibraryViewModel = MakeShareable(new FLibraryViewModel(ThisPtr));
+	LibraryViewModel = MakeShared<FLibraryViewModel>(ThisPtr);
 	LibraryViewModel->RegisterToEvents();
 
-	WidgetToolbar = MakeShareable(new FWidgetBlueprintEditorToolbar(ThisPtr));
+	WidgetToolbar = MakeShared<FWidgetBlueprintEditorToolbar>(ThisPtr);
 
 	BindToolkitCommands();
 
@@ -145,7 +145,7 @@ void FWidgetBlueprintEditor::InitWidgetBlueprintEditor(const EToolkitMode::Type
 
 	UpdatePreview(GetWidgetBlueprintObj(), true);
 
-	DesignerCommandList = MakeShareable(new FUICommandList);
+	DesignerCommandList = MakeShared<FUICommandList>();
 
 	DesignerCommandList->MapAction(FGenericCommands::Get().Delete,
 		FExecuteAction::CreateSP(this, &FWidgetBlueprintEditor::DeleteSelectedWidgets),
