Fix Hiearchy drag drop crashing

diff --git a/Engine/Source/Editor/UMGEditor/Private/Hierarchy/SHierarchyViewItem.cpp b/Engine/Source/Editor/UMGEditor/Private/Hierarchy/SHierarchyViewItem.cpp
--- a/Engine/Source/Editor/UMGEditor/Private/Hierarchy/SHierarchyViewItem.cpp
+++ b/Engine/Source/Editor/UMGEditor/Private/Hierarchy/SHierarchyViewItem.cpp
@@ -151,7 +151,6 @@ TOptional<EItemDropZone> ProcessHierarchyDragDrop(const FDragDropEvent& DragDrop
 {
 	UWidget* TargetTemplate = TargetItem.GetTemplate();
 
-
 	if (TSharedPtr<FHierarchyWidgetDragDropOp> HierarchyDragDropOp = DragDropEvent.GetOperationAs<FHierarchyWidgetDragDropOp>())
 	{
 		if (!HierarchyDragDropOp->HasOriginatedFrom(BlueprintEditor))
@@ -190,7 +189,7 @@ TOptional<EItemDropZone> ProcessHierarchyDragDrop(const FDragDropEvent& DragDrop
 
 	// Is this a drag/drop op to create a new widget in the tree?
 	TSharedPtr<FDragDropOperation> DragDropOp = DragDropEvent.GetOperation();
-	if (DragDropOp.IsValid() )
+	if (DragDropOp.IsValid() && !DragDropOp->IsOfType<FHierarchyWidgetDragDropOp>())
 	{
 		TSharedPtr<FDecoratedDragDropOp> DecoratedDragDropOp = nullptr;
 		if (DragDropOp->IsOfType<FDecoratedDragDropOp>())
@@ -887,7 +886,7 @@ TOptional<EItemDropZone> FNamedSlotModel::HandleCanAcceptDrop(const FDragDropEve
 	UWidgetBlueprint* Blueprint = BlueprintEditor.Pin()->GetWidgetBlueprintObj();
 
 	TSharedPtr<FDragDropOperation> DragDropOp = DragDropEvent.GetOperation();
-	if (DragDropOp.IsValid())
+	if (DragDropOp.IsValid() && !DragDropOp->IsOfType<FHierarchyWidgetDragDropOp>())
 	{
 		TSharedPtr<FDecoratedDragDropOp> DecoratedDragDropOp = nullptr;
 		if (DragDropOp->IsOfType<FDecoratedDragDropOp>())
@@ -975,7 +974,7 @@ FReply FNamedSlotModel::HandleAcceptDrop(FDragDropEvent const& DragDropEvent, EI
 	}
 
 	TSharedPtr<FDragDropOperation> DragDropOp = DragDropEvent.GetOperation();
-	if (DragDropOp.IsValid())
+	if (DragDropOp.IsValid() && !DragDropOp->IsOfType<FHierarchyWidgetDragDropOp>())
 	{
 		FScopedTransaction Transaction(LOCTEXT("AddWidgetFromTemplate", "Add Widget"));
 
