Actually enable Widget Library by default.
Fix widget library locking up when scrolling edge is just before a new asset view.
Review feedback.

diff --git a/Engine/Source/Editor/UMGEditor/Private/BlueprintModes/WidgetDesignerApplicationMode.cpp b/Engine/Source/Editor/UMGEditor/Private/BlueprintModes/WidgetDesignerApplicationMode.cpp
--- a/Engine/Source/Editor/UMGEditor/Private/BlueprintModes/WidgetDesignerApplicationMode.cpp
+++ b/Engine/Source/Editor/UMGEditor/Private/BlueprintModes/WidgetDesignerApplicationMode.cpp
@@ -53,7 +53,7 @@ FWidgetDesignerApplicationMode::FWidgetDesignerApplicationMode(TSharedPtr<FWidge
 					->SetSizeCoefficient(0.5f)
 					->SetForegroundTab(FPaletteTabSummoner::TabID)
 					->AddTab(FPaletteTabSummoner::TabID, ETabState::OpenedTab)
-					->AddTab(FLibraryTabSummoner::TabID, ETabState::ClosedTab)
+					->AddTab(FLibraryTabSummoner::TabID, ETabState::OpenedTab)
 				)
 				->Split
 				(
diff --git a/Engine/Source/Editor/UMGEditor/Private/Library/SLibraryView.cpp b/Engine/Source/Editor/UMGEditor/Private/Library/SLibraryView.cpp
--- a/Engine/Source/Editor/UMGEditor/Private/Library/SLibraryView.cpp
+++ b/Engine/Source/Editor/UMGEditor/Private/Library/SLibraryView.cpp
@@ -10,6 +10,7 @@
 #include "Framework/MultiBox/MultiBoxBuilder.h"
 #include "Widgets/Images/SImage.h"
 #include "Widgets/Layout/SScrollBorder.h"
+#include "Widgets/Layout/SScrollBox.h"
 #include "SAssetView.h"
 
 #if WITH_EDITOR
@@ -251,13 +252,20 @@ void SLibraryView::Construct(const FArguments& InArgs, TSharedPtr<FWidgetBluepri
 		+ SVerticalBox::Slot()
 		.FillHeight(1.0f)
 		[
-			SNew(SScrollBorder, WidgetTemplatesView.ToSharedRef())
+			SNew(SBorder)
+			.BorderImage(FAppStyle::Get().GetBrush("Brushes.Recessed"))
+			.Padding(0)
 			[
-				SNew(SBorder)
-				.BorderImage(FAppStyle::Get().GetBrush("Brushes.Recessed"))
-				.Padding(0)
+				SNew(SScrollBox)
+				.ScrollBarPadding(FMargin(2.0, 0))
+				+SScrollBox::Slot()
 				[
-					WidgetTemplatesView.ToSharedRef()
+					SNew(SBorder)
+					.BorderImage(FAppStyle::Get().GetBrush("Brushes.Recessed"))
+					.Padding(0)
+					[
+						WidgetTemplatesView.ToSharedRef()
+					]
 				]
 			]
 		]
diff --git a/Engine/Source/Editor/UMGEditor/Private/Library/SLibraryViewModel.cpp b/Engine/Source/Editor/UMGEditor/Private/Library/SLibraryViewModel.cpp
--- a/Engine/Source/Editor/UMGEditor/Private/Library/SLibraryViewModel.cpp
+++ b/Engine/Source/Editor/UMGEditor/Private/Library/SLibraryViewModel.cpp
@@ -67,7 +67,7 @@ void FLibraryViewModel::BuildWidgetTemplateCategory(FString& Category, TArray<TS
 
 	//}
 
-	Header->Children.Sort([](TSharedPtr<FWidgetViewModel> L, TSharedPtr<FWidgetViewModel> R) { return R->GetName().CompareTo(L->GetName()) > 0; });
+	Header->Children.Sort([](const TSharedPtr<FWidgetViewModel>& L, const TSharedPtr<FWidgetViewModel>& R) { return R->GetName().CompareTo(L->GetName()) > 0; });
 
 	WidgetViewModels.Add(Header);
 }
diff --git a/Engine/Source/Editor/UMGEditor/Private/Palette/SPaletteViewModel.cpp b/Engine/Source/Editor/UMGEditor/Private/Palette/SPaletteViewModel.cpp
--- a/Engine/Source/Editor/UMGEditor/Private/Palette/SPaletteViewModel.cpp
+++ b/Engine/Source/Editor/UMGEditor/Private/Palette/SPaletteViewModel.cpp
@@ -482,8 +482,8 @@ void FPaletteViewModel::BuildWidgetTemplateCategory(FString& Category, TArray<TS
 		}
 
 	}
-
-	Header->Children.Sort([](TSharedPtr<FWidgetViewModel> L, TSharedPtr<FWidgetViewModel> R) { return R->GetName().CompareTo(L->GetName()) > 0; });
+	
+	Header->Children.Sort([](const TSharedPtr<FWidgetViewModel>& L, const TSharedPtr<FWidgetViewModel>& R) { return R->GetName().CompareTo(L->GetName()) > 0; });
 
 	WidgetViewModels.Add(Header);
 }
