Use FRenderTarget of canvas for widget blueprint thumbnail rendering in order to reduce the number of render targets used.
#rnx

diff --git a/Engine/Source/Editor/UMGEditor/Private/WidgetBlueprintEditor.cpp b/Engine/Source/Editor/UMGEditor/Private/WidgetBlueprintEditor.cpp
--- a/Engine/Source/Editor/UMGEditor/Private/WidgetBlueprintEditor.cpp
+++ b/Engine/Source/Editor/UMGEditor/Private/WidgetBlueprintEditor.cpp
@@ -312,12 +312,7 @@ void FWidgetBlueprintEditor::CaptureThumbnail()
 	}
 
 	UTextureRenderTarget2D* RenderTarget2D = NewObject<UTextureRenderTarget2D>();
-	RenderTarget2D->Filter = TF_Bilinear;
-	RenderTarget2D->ClearColor = FLinearColor::Transparent;
-	RenderTarget2D->SRGB = true;
-	RenderTarget2D->TargetGamma = 1;
-
-	TOptional<FWidgetBlueprintEditorUtils::FWidgetThumbnailProperties> ScaleAndOffset = FWidgetBlueprintEditorUtils::DrawSWidgetInRenderTargetForThumbnail(PreviewWidget, RenderTarget2D, FVector2D(256.f, 256.f), TOptional<FVector2D>());
+	TOptional<FWidgetBlueprintEditorUtils::FWidgetThumbnailProperties> ScaleAndOffset = FWidgetBlueprintEditorUtils::DrawSWidgetInRenderTargetForThumbnail(PreviewWidget, RenderTarget2D, FVector2D(256.f, 256.f), TOptional<FVector2D>(), EThumbnailPreviewSizeMode::MatchDesignerMode);
 
 	if (!ScaleAndOffset.IsSet())
 	{
diff --git a/Engine/Source/Editor/UMGEditor/Private/WidgetBlueprintEditorUtils.cpp b/Engine/Source/Editor/UMGEditor/Private/WidgetBlueprintEditorUtils.cpp
--- a/Engine/Source/Editor/UMGEditor/Private/WidgetBlueprintEditorUtils.cpp
+++ b/Engine/Source/Editor/UMGEditor/Private/WidgetBlueprintEditorUtils.cpp
@@ -2120,18 +2120,27 @@ EDesignPreviewSizeMode FWidgetBlueprintEditorUtils::ConvertThumbnailSizeModeToDe
 
 TOptional<FWidgetBlueprintEditorUtils::FWidgetThumbnailProperties> FWidgetBlueprintEditorUtils::DrawSWidgetInRenderTarget(UUserWidget* WidgetInstance, UTextureRenderTarget2D* RenderTarget2D)
 {
-	return DrawSWidgetInRenderTargetInternal(WidgetInstance, RenderTarget2D, FVector2D(256.f,256.f), false, TOptional<FVector2D>(), EThumbnailPreviewSizeMode::MatchDesignerMode);
+	return DrawSWidgetInRenderTargetInternal(WidgetInstance, nullptr, RenderTarget2D, FVector2D(256.f,256.f), false, TOptional<FVector2D>(), EThumbnailPreviewSizeMode::MatchDesignerMode);
+}
+
+TOptional<FWidgetBlueprintEditorUtils::FWidgetThumbnailProperties> FWidgetBlueprintEditorUtils::DrawSWidgetInRenderTargetForThumbnail(UUserWidget* WidgetInstance, FRenderTarget* RenderTarget2D, FVector2D ThumbnailSize, TOptional<FVector2D> ThumbnailCustomSize, EThumbnailPreviewSizeMode ThumbnailSizeMode)
+{
+	return DrawSWidgetInRenderTargetInternal(WidgetInstance, RenderTarget2D, nullptr,ThumbnailSize, true, ThumbnailCustomSize, ThumbnailSizeMode);
 }
 
 TOptional<FWidgetBlueprintEditorUtils::FWidgetThumbnailProperties> FWidgetBlueprintEditorUtils::DrawSWidgetInRenderTargetForThumbnail(UUserWidget* WidgetInstance, UTextureRenderTarget2D* RenderTarget2D, FVector2D ThumbnailSize, TOptional<FVector2D> ThumbnailCustomSize, EThumbnailPreviewSizeMode ThumbnailSizeMode)
 {
-	return DrawSWidgetInRenderTargetInternal(WidgetInstance, RenderTarget2D, ThumbnailSize, true, ThumbnailCustomSize, ThumbnailSizeMode);
+	return DrawSWidgetInRenderTargetInternal(WidgetInstance, nullptr, RenderTarget2D,ThumbnailSize, true, ThumbnailCustomSize, ThumbnailSizeMode);
 }
 
-TOptional<FWidgetBlueprintEditorUtils::FWidgetThumbnailProperties>  FWidgetBlueprintEditorUtils::DrawSWidgetInRenderTargetInternal(UUserWidget* WidgetInstance, UTextureRenderTarget2D* RenderTarget2D, FVector2D ThumbnailSize, bool bIsForThumbnail, TOptional<FVector2D> ThumbnailCustomSize,EThumbnailPreviewSizeMode ThumbnailSizeMode)
+TOptional<FWidgetBlueprintEditorUtils::FWidgetThumbnailProperties>  FWidgetBlueprintEditorUtils::DrawSWidgetInRenderTargetInternal(UUserWidget* WidgetInstance, FRenderTarget* RenderTarget2D, UTextureRenderTarget2D* TextureRenderTarget,FVector2D ThumbnailSize, bool bIsForThumbnail, TOptional<FVector2D> ThumbnailCustomSize, EThumbnailPreviewSizeMode ThumbnailSizeMode)
 {
+	if (TextureRenderTarget == nullptr && RenderTarget2D == nullptr)
+	{
+		return TOptional<FWidgetBlueprintEditorUtils::FWidgetThumbnailProperties>(); 
+	}
 	//Create Window
-	FVector2D Offset(0.f,0.f);
+	FVector2D Offset(0.f, 0.f);
 	FVector2D ScaledSize(0.f, 0.f);
 	TSharedPtr<SWidget> WindowContent = WidgetInstance->TakeWidget();
 
@@ -2140,7 +2149,7 @@ TOptional<FWidgetBlueprintEditorUtils::FWidgetThumbnailProperties>  FWidgetBluep
 		return TOptional<FWidgetBlueprintEditorUtils::FWidgetThumbnailProperties>();
 	}
 
-	TSharedRef<SVirtualWindow> Window = SNew(SVirtualWindow).Size(ThumbnailSize);
+	TSharedRef<SVirtualWindow> Window = SNew(SVirtualWindow);
 	TUniquePtr<FHittestGrid> HitTestGrid = MakeUnique<FHittestGrid>();
 	Window->SetContent(WindowContent.ToSharedRef());
 	Window->Resize(ThumbnailSize);
@@ -2149,54 +2158,64 @@ TOptional<FWidgetBlueprintEditorUtils::FWidgetThumbnailProperties>  FWidgetBluep
 	FGeometry WindowGeometry = FGeometry::MakeRoot(ThumbnailSize, FSlateLayoutTransform(1.0f));
 	Window->SlatePrepass(1.0f);
 	FVector2D DesiredSizeWindow = Window->GetDesiredSize();
-
-
+	
 	if (DesiredSizeWindow.X < SMALL_NUMBER || DesiredSizeWindow.Y < SMALL_NUMBER)
 	{
 		return TOptional<FWidgetBlueprintEditorUtils::FWidgetThumbnailProperties>();
 	}
 
-	FVector2D UnscaledSize = FWidgetBlueprintEditorUtils::GetWidgetPreviewUnScaledCustomSize(DesiredSizeWindow, WidgetInstance, ThumbnailCustomSize,ThumbnailSizeMode);
+	FVector2D UnscaledSize = FWidgetBlueprintEditorUtils::GetWidgetPreviewUnScaledCustomSize(DesiredSizeWindow, WidgetInstance, ThumbnailCustomSize, ThumbnailSizeMode);
 	if (UnscaledSize.X < SMALL_NUMBER || UnscaledSize.Y < SMALL_NUMBER)
 	{
 		return TOptional<FWidgetBlueprintEditorUtils::FWidgetThumbnailProperties>();
 	}
 
-	float Scale;
+	float Scale = 1.f;
 	// Change some configuration if it is for thumbnail creation
-	if (bIsForThumbnail) 
+	if (bIsForThumbnail)
 	{
 		TTuple<float, FVector2D> ScaleAndOffset = FWidgetBlueprintEditorUtils::GetThumbnailImageScaleAndOffset(UnscaledSize, ThumbnailSize);
 		Scale = ScaleAndOffset.Get<0>();
 		Offset = ScaleAndOffset.Get<1>();
 	}
-	else
-	{
-		RenderTarget2D->Filter = TF_Bilinear;
-		RenderTarget2D->ClearColor = FLinearColor::Transparent;
-		RenderTarget2D->SRGB = true;
-		RenderTarget2D->RenderTargetFormat = RTF_RGBA8;
-		Scale = 1.f;
-		Offset = FVector2D(0.f, 0.f);
-	}
 
 	ScaledSize = UnscaledSize * Scale;
 	if (ScaledSize.X < 1.f || ScaledSize.Y < 1.f)
 	{
 		return TOptional<FWidgetBlueprintEditorUtils::FWidgetThumbnailProperties>();
 	}
-	const EPixelFormat RequestedFormat = FSlateApplication::Get().GetRenderer()->GetSlateRecommendedColorFormat();
-	RenderTarget2D->InitCustomFormat(ScaledSize.X, ScaledSize.Y, RequestedFormat, false);
 
 	// Create Renderer Target and WidgetRenderer
-	bool bApplyGammaCorrection = false;
+	bool bApplyGammaCorrection = RenderTarget2D? true : false;
 	FWidgetRenderer* WidgetRenderer = new FWidgetRenderer(bApplyGammaCorrection);
-	const bool bIsLinearSpace = !bApplyGammaCorrection;
+
 	if (!bIsForThumbnail)
 	{
 		WidgetRenderer->SetIsPrepassNeeded(false);
 	}
-	WidgetRenderer->DrawWindow(RenderTarget2D, *HitTestGrid, Window, Scale, ScaledSize, 0.1f);
+
+	if (TextureRenderTarget)
+	{
+		TextureRenderTarget->Filter = TF_Bilinear;
+		TextureRenderTarget->ClearColor = FLinearColor::Transparent;
+		TextureRenderTarget->SRGB = true;
+		TextureRenderTarget->RenderTargetFormat = RTF_RGBA8;
+
+		uint32 ScaledSizeX = static_cast<uint32>(ScaledSize.X);
+		uint32 ScaledSizeY = static_cast<uint32>(ScaledSize.Y);
+
+		const bool bForceLinearGamma = false;
+		const EPixelFormat RequestedFormat = FSlateApplication::Get().GetRenderer()->GetSlateRecommendedColorFormat();
+		TextureRenderTarget->InitCustomFormat(ScaledSizeX, ScaledSizeY, RequestedFormat, bForceLinearGamma);
+		WidgetRenderer->DrawWindow(TextureRenderTarget, *HitTestGrid, Window, Scale, ScaledSize, 0.1f);
+	}
+	else
+	{
+		ensure(RenderTarget2D != nullptr);
+		WidgetRenderer->SetShouldClearTarget(false);
+		WidgetRenderer->ViewOffset = Offset;
+		WidgetRenderer->DrawWindow(RenderTarget2D, *HitTestGrid, Window, Scale, ScaledSize, 0.1f);
+	}
 
 	if (WidgetRenderer)
 	{
diff --git a/Engine/Source/Editor/UMGEditor/Private/WidgetBlueprintEditorUtils.h b/Engine/Source/Editor/UMGEditor/Private/WidgetBlueprintEditorUtils.h
--- a/Engine/Source/Editor/UMGEditor/Private/WidgetBlueprintEditorUtils.h
+++ b/Engine/Source/Editor/UMGEditor/Private/WidgetBlueprintEditorUtils.h
@@ -88,7 +88,9 @@ public:
 
 	static void SetTextureAsAssetThumbnail(UWidgetBlueprint* WidgetBlueprint, UTexture2D* ThumbnailTexture);
 
-	static TOptional<FWidgetThumbnailProperties> DrawSWidgetInRenderTargetForThumbnail(UUserWidget* WidgetInstance, UTextureRenderTarget2D* RenderTarget2D, FVector2D ThumbnailSize, TOptional<FVector2D> ThumbnailCustomSize, EThumbnailPreviewSizeMode ThumbnailSizeMode = EThumbnailPreviewSizeMode::MatchDesignerMode);
+	static TOptional<FWidgetThumbnailProperties> DrawSWidgetInRenderTargetForThumbnail(UUserWidget* WidgetInstance, FRenderTarget* RenderTarget2D, FVector2D ThumbnailSize, TOptional<FVector2D> ThumbnailCustomSize, EThumbnailPreviewSizeMode ThumbnailSizeMode = EThumbnailPreviewSizeMode::MatchDesignerMode);
+
+	static TOptional<FWidgetThumbnailProperties> DrawSWidgetInRenderTargetForThumbnail(UUserWidget* WidgetInstance, UTextureRenderTarget2D* RenderTarget2D, FVector2D ThumbnailSize, TOptional<FVector2D> ThumbnailCustomSize, EThumbnailPreviewSizeMode ThumbnailSizeMode);
 
 	static TOptional<FWidgetThumbnailProperties> DrawSWidgetInRenderTarget(UUserWidget* WidgetInstance, UTextureRenderTarget2D* RenderTarget2D);
 
@@ -128,5 +130,5 @@ private:
 
 	static bool ShouldContinueReplaceOperation(UWidgetBlueprint* BP, const TArray<FText>& WidgetNames);	
 
-	static TOptional<FWidgetThumbnailProperties> DrawSWidgetInRenderTargetInternal(UUserWidget* WidgetInstance, UTextureRenderTarget2D* RenderTarget2D, FVector2D ThumbnailSize, bool bIsForThumbnail, TOptional<FVector2D> ThumbnailCustomSize, EThumbnailPreviewSizeMode ThumbnailSizeMode);
+	static TOptional<FWidgetThumbnailProperties> DrawSWidgetInRenderTargetInternal(UUserWidget* WidgetInstance, FRenderTarget* RenderTarget2D, UTextureRenderTarget2D* TextureRenderTarget,FVector2D ThumbnailSize, bool bIsForThumbnail, TOptional<FVector2D> ThumbnailCustomSize, EThumbnailPreviewSizeMode ThumbnailSizeMode);
 };
diff --git a/Engine/Source/Editor/UMGEditor/Private/WidgetBlueprintThumbnailRenderer.cpp b/Engine/Source/Editor/UMGEditor/Private/WidgetBlueprintThumbnailRenderer.cpp
--- a/Engine/Source/Editor/UMGEditor/Private/WidgetBlueprintThumbnailRenderer.cpp
+++ b/Engine/Source/Editor/UMGEditor/Private/WidgetBlueprintThumbnailRenderer.cpp
@@ -21,7 +21,6 @@ public:
 	struct FInstance
 	{
 		TWeakObjectPtr<UClass> WidgetClass;
-		TWeakObjectPtr<UTextureRenderTarget2D> RenderTarget2D;
 		TWeakObjectPtr<UUserWidget> Widget;
 	};
 
@@ -127,6 +126,11 @@ void UWidgetBlueprintThumbnailRenderer::Draw(UObject* Object, int32 X, int32 Y,
 		return;
 	}
 
+	if (!RenderTarget)
+	{
+		return;
+	}
+
 	UWidgetBlueprint* WidgetBlueprintToRender = Cast<UWidgetBlueprint>(Object);
 	const bool bIsBlueprintValid = IsValid(WidgetBlueprintToRender)
 		&& IsValid(WidgetBlueprintToRender->GeneratedClass)
@@ -142,11 +146,10 @@ void UWidgetBlueprintThumbnailRenderer::Draw(UObject* Object, int32 X, int32 Y,
 
 	// Create a plain gray background for the thumbnail
 	const int32 SizeOfUV = 1;
-	FLinearColor GrayBackgroundColor(FVector4(.03f, .03f, .03f, 1.f));
-	Canvas->DrawTile(
-		0.0f, 0.0f, Width, Height,
-		0.0f, 0.0f, SizeOfUV, SizeOfUV,
-		GrayBackgroundColor);
+	FLinearColor GrayBackgroundColor(FVector4(.03f, .03f, .03f, 1.0f));
+	FCanvasTileItem TileItem(FVector2D(X, Y),  GWhiteTexture, FVector2D(Width, Height), FVector2D(0, 0), FVector2D(SizeOfUV, SizeOfUV), GrayBackgroundColor);
+	TileItem.BlendMode = SE_BLEND_AlphaBlend;
+	TileItem.Draw(Canvas);
 
 	// check if an image is used instead of auto generating the thumbnail
 	if (WidgetBlueprintToRender->ThumbnailImage)
@@ -167,8 +170,8 @@ void UWidgetBlueprintThumbnailRenderer::Draw(UObject* Object, int32 X, int32 Y,
 	}
 	else
 	{
+		Canvas->Flush_GameThread();
 		UUserWidget* WidgetInstance = nullptr;
-		UTextureRenderTarget2D* RenderTarget2D = nullptr;
 
 		{
 			UClass* ClassToGenerate = WidgetBlueprintToRender->GeneratedClass;
@@ -183,44 +186,23 @@ void UWidgetBlueprintThumbnailRenderer::Draw(UObject* Object, int32 X, int32 Y,
 
 				WidgetInstance = ThumbnailInstance.Widget.Get();
 			}
-
-			RenderTarget2D = ThumbnailInstance.RenderTarget2D.Get();
-			if (!RenderTarget2D)
-			{
-				ThumbnailInstance.RenderTarget2D = NewObject<UTextureRenderTarget2D>(GetTransientPackage());
-				ThumbnailInstance.RenderTarget2D->Filter = TF_Bilinear;
-				ThumbnailInstance.RenderTarget2D->ClearColor = FLinearColor::Transparent;
-				ThumbnailInstance.RenderTarget2D->SRGB = true;
-				ThumbnailInstance.RenderTarget2D->TargetGamma = 1;
-
-				RenderTarget2D = ThumbnailInstance.RenderTarget2D.Get();
-			}
 		}
 
 		FVector2D ThumbnailSize(Width, Height);
 		TOptional<FWidgetBlueprintEditorUtils::FWidgetThumbnailProperties> ScaleAndOffset;
-
 		if (WidgetBlueprintToRender->ThumbnailSizeMode == EThumbnailPreviewSizeMode::Custom)
 		{
-			ScaleAndOffset = FWidgetBlueprintEditorUtils::DrawSWidgetInRenderTargetForThumbnail(WidgetInstance, RenderTarget2D, ThumbnailSize, WidgetBlueprintToRender->ThumbnailCustomSize, WidgetBlueprintToRender->ThumbnailSizeMode);
+			ScaleAndOffset = FWidgetBlueprintEditorUtils::DrawSWidgetInRenderTargetForThumbnail(WidgetInstance, Canvas->GetRenderTarget(), ThumbnailSize, WidgetBlueprintToRender->ThumbnailCustomSize, WidgetBlueprintToRender->ThumbnailSizeMode);
 		}
 		else
 		{
-			ScaleAndOffset = FWidgetBlueprintEditorUtils::DrawSWidgetInRenderTargetForThumbnail(WidgetInstance, RenderTarget2D, ThumbnailSize, TOptional<FVector2D>(), WidgetBlueprintToRender->ThumbnailSizeMode);
+			ScaleAndOffset = FWidgetBlueprintEditorUtils::DrawSWidgetInRenderTargetForThumbnail(WidgetInstance, Canvas->GetRenderTarget(), ThumbnailSize, TOptional<FVector2D>(), WidgetBlueprintToRender->ThumbnailSizeMode);
 
 		}
 		if (!ScaleAndOffset.IsSet())
 		{
 			return;
 		}
-		FVector2D Offset = ScaleAndOffset.GetValue().Offset;
-		FVector2D ScaledSize = ScaleAndOffset.GetValue().ScaledSize;
-
-		// Draw the SWidget on canvas
-		FCanvasTileItem CanvasTile(FVector2D(X + Offset.X, Y + Offset.Y), RenderTarget2D->GetResource(), ScaledSize, FLinearColor::White);
-		CanvasTile.BlendMode = SE_BLEND_Translucent;
-		FlushRenderingCommands();
-		CanvasTile.Draw(Canvas);
 	}
 #endif
 }
