Add anim drawer to UMG editor.

diff --git a/Engine/Source/Editor/EditorStyle/Private/StarshipStyle.cpp b/Engine/Source/Editor/EditorStyle/Private/StarshipStyle.cpp
--- a/Engine/Source/Editor/EditorStyle/Private/StarshipStyle.cpp
+++ b/Engine/Source/Editor/EditorStyle/Private/StarshipStyle.cpp
@@ -7815,6 +7815,7 @@ void FStarshipEditorStyle::FStyle::SetupUMGEditorStyles()
 
 	Set("UMGEditor.ResizeResolutionFont", DEFAULT_FONT("Bold", 10));
 	Set("UMGEditor.CategoryIcon", new IMAGE_BRUSH("Icons/hiererchy_16x", Icon16x16));
+	Set("UMGEditor.AnimTabIcon", new IMAGE_BRUSH_SVG("Starship/MainToolbar/cinematics", Icon16x16));
 }
 
 void FStarshipEditorStyle::FStyle::SetupTranslationEditorStyles()
diff --git a/Engine/Source/Editor/Kismet/Private/BlueprintEditor.cpp b/Engine/Source/Editor/Kismet/Private/BlueprintEditor.cpp
--- a/Engine/Source/Editor/Kismet/Private/BlueprintEditor.cpp
+++ b/Engine/Source/Editor/Kismet/Private/BlueprintEditor.cpp
@@ -1674,6 +1674,8 @@ FBlueprintEditor::FBlueprintEditor()
 	: bSaveIntermediateBuildProducts(false)
 	, bIsReparentingBlueprint(false)
 	, bPendingDeferredClose(false)
+	, CachedNumWarnings(0)
+	, CachedNumErrors(0)
 	, bRequestedSavingOpenDocumentState(false)
 	, bBlueprintModifiedOnOpen (false)
 	, PinVisibility(SGraphEditor::Pin_Show)
@@ -3787,7 +3789,10 @@ void FBlueprintEditor::Compile()
 
 		LogResults.EndEvent();
 
-		const bool bForceMessageDisplay = ((LogResults.NumWarnings > 0) || (LogResults.NumErrors > 0)) && !BlueprintObj->bIsRegeneratingOnLoad;
+		CachedNumWarnings = LogResults.NumWarnings;
+		CachedNumErrors = LogResults.NumErrors;
+
+		const bool bForceMessageDisplay = (LogResults.NumWarnings > 0 || LogResults.NumErrors > 0) && !BlueprintObj->bIsRegeneratingOnLoad;
 		DumpMessagesToCompilerLog(LogResults.Messages, bForceMessageDisplay);
 
 		UBlueprintEditorSettings const* BpEditorSettings = GetDefault<UBlueprintEditorSettings>();
diff --git a/Engine/Source/Editor/Kismet/Public/BlueprintEditor.h b/Engine/Source/Editor/Kismet/Public/BlueprintEditor.h
--- a/Engine/Source/Editor/Kismet/Public/BlueprintEditor.h
+++ b/Engine/Source/Editor/Kismet/Public/BlueprintEditor.h
@@ -1339,6 +1339,10 @@ protected:
 	TSharedPtr<class SWidget> CompilerResults;
 	TSharedPtr<class IMessageLogListing> CompilerResultsListing;
 	
+	/** Limited cache of last compile results */
+	int32 CachedNumWarnings;
+	int32 CachedNumErrors;
+	
 	/** Find results log as well as the search filter */
 	TSharedPtr<class SFindInBlueprints> FindResults;
 
diff --git a/Engine/Source/Editor/StatusBar/Private/SStatusBar.cpp b/Engine/Source/Editor/StatusBar/Private/SStatusBar.cpp
--- a/Engine/Source/Editor/StatusBar/Private/SStatusBar.cpp
+++ b/Engine/Source/Editor/StatusBar/Private/SStatusBar.cpp
@@ -769,8 +769,14 @@ void SStatusBar::OnGlobalFocusChanging(const FFocusEvent& FocusEvent, const FWea
 		}
 
 		bool bShouldDismiss = false;
+
+		// If we aren't focusing any new widgets, act as if the drawer is in the path 
+		bool bDrawerInPath = NewFocusedWidgetPath.ContainsWidget(ActiveDrawerOverlayContent.Get()) 
+			|| NewFocusedWidgetPath.ContainsWidget(this) 
+			|| NewFocusedWidgetPath.Widgets.Num() == 0;
+
 		// Do not close due to slow tasks as those opening send window activation events
-		if (!GIsSlowTask && !FSlateApplication::Get().GetActiveModalWindow().IsValid() && ActiveDrawerOverlayContent.IsValid() && (!NewFocusedWidgetPath.ContainsWidget(ActiveDrawerOverlayContent.Get()) && !NewFocusedWidgetPath.ContainsWidget(this)))
+		if (!GIsSlowTask && !bDrawerInPath && !FSlateApplication::Get().GetActiveModalWindow().IsValid() && ActiveDrawerOverlayContent.IsValid())
 		{
 			if (TSharedPtr<SWidget> MenuHost = FSlateApplication::Get().GetMenuHostWidget())
 			{
@@ -962,6 +968,11 @@ bool SStatusBar::IsAnyOtherDrawerOpened(const FName DrawerId) const
 	return OpenedDrawer.IsValid() && OpenedDrawer.DrawerId != DrawerId ? true : false;
 }
 
+FName SStatusBar::GetStatusBarName() const
+{
+	return StatusBarName;
+}
+
 FReply SStatusBar::OnDrawerButtonClicked(const FName DrawerId)
 {
 	if (!IsDrawerOpened(DrawerId))
@@ -1184,14 +1195,14 @@ void SStatusBar::CloseDrawerImmediatelyInternal(const FOpenDrawerData& Data)
 	}
 }
 
-void SStatusBar::RegisterDrawer(FStatusBarDrawer&& Drawer)
+void SStatusBar::RegisterDrawer(FStatusBarDrawer&& Drawer, int32 SlotIndex)
 {
 	const int32 NumDrawers = RegisteredDrawers.Num();
 	RegisteredDrawers.AddUnique(Drawer);
 
 	if (RegisteredDrawers.Num() > NumDrawers)
 	{
-		DrawerBox->AddSlot()
+		DrawerBox->InsertSlot(SlotIndex)
 		.Padding(1.0f, 0.0f)
 		.AutoWidth()
 		[
@@ -1251,13 +1262,14 @@ void SStatusBar::OpenDrawer(const FName DrawerId)
 
 			OpenedDrawer = MoveTemp(NewlyOpenedDrawer);
 
-			DrawerData->OnDrawerOpenedDelegate.ExecuteIfBound(ThisStatusBar);
+			DrawerData->OnDrawerOpenedDelegate.ExecuteIfBound(ThisStatusBar->StatusBarName);
 		}
 	}
 }
 
-void SStatusBar::DismissDrawer(const TSharedPtr<SWidget>& NewlyFocusedWidget)
+bool SStatusBar::DismissDrawer(const TSharedPtr<SWidget>& NewlyFocusedWidget)
 {
+	bool bWasDismissed = false;
 	if (OpenedDrawer.IsValid())
 	{
 		FStatusBarDrawer* Drawer = RegisteredDrawers.FindByKey(OpenedDrawer.DrawerId);
@@ -1268,7 +1280,10 @@ void SStatusBar::DismissDrawer(const TSharedPtr<SWidget>& NewlyFocusedWidget)
 		OpenedDrawer = FOpenDrawerData();
 
 		Drawer->OnDrawerDismissedDelegate.ExecuteIfBound(NewlyFocusedWidget);
+		bWasDismissed = true;
 	}
+
+	return bWasDismissed;
 }
 
 void SStatusBar::CloseDrawerImmediately(FName DrawerId)
diff --git a/Engine/Source/Editor/StatusBar/Private/SStatusBar.h b/Engine/Source/Editor/StatusBar/Private/SStatusBar.h
--- a/Engine/Source/Editor/StatusBar/Private/SStatusBar.h
+++ b/Engine/Source/Editor/StatusBar/Private/SStatusBar.h
@@ -9,9 +9,6 @@
 #include "Animation/CurveSequence.h"
 #include "Application/ThrottleManager.h"
 
-
-DECLARE_DELEGATE_OneParam(FOnStatusBarDrawerOpened, TSharedRef<SStatusBar>&)
-DECLARE_DELEGATE_OneParam(FOnStatusBarDrawerDismissed, const TSharedPtr<SWidget>&)
 DECLARE_DELEGATE_OneParam(FOnStatusBarDrawerTargetHeightChanged, float);
 
 class SWindow;
@@ -61,33 +58,6 @@ struct FStatusBarProgress
 	FProgressNotificationHandle Handle;
 };
 
-struct FStatusBarDrawer
-{
-	FStatusBarDrawer(FName InUniqueId)
-		: UniqueId(InUniqueId)
-	{}
-
-	FName UniqueId;
-	FOnGetContent GetDrawerContentDelegate;
-	FOnStatusBarDrawerOpened OnDrawerOpenedDelegate;
-	FOnStatusBarDrawerDismissed OnDrawerDismissedDelegate;
-
-	TSharedPtr<SWidget> CustomWidget;
-	FText ButtonText;
-	FText ToolTipText;
-	const FSlateBrush* Icon;
-
-	bool operator==(const FName& OtherId) const
-	{
-		return UniqueId == OtherId;
-	}
-
-	bool operator==(const FStatusBarDrawer& Other) const
-	{
-		return UniqueId == Other.UniqueId;
-	}
-};
-
 struct FOpenDrawerData
 {
 	FName DrawerId;
@@ -169,7 +139,7 @@ public:
 	/**
 	 * Registers a new drawer with this status bar. Registering will add a button to open and close the drawer
 	 */
-	void RegisterDrawer(FStatusBarDrawer&& Drawer);
+	void RegisterDrawer(FStatusBarDrawer&& Drawer, int32 SlotIndex = INDEX_NONE);
 
 	/**
 	 * Opens a drawer
@@ -180,8 +150,11 @@ public:
 
 	/**
 	 * Dismisses an open drawer with an animation.  The drawer contents are removed once the animation is complete
+	 * 
+	 * @param NewlyFocusedWidget Optional widget to focus 
+	 * @return true if any open drawer was dismissed
 	 */
-	void DismissDrawer(const TSharedPtr<SWidget>& NewlyFocusedWidget);
+	bool DismissDrawer(const TSharedPtr<SWidget>& NewlyFocusedWidget);
 
 	/**
 	 * Closes a drawer immediately with no closing animation. Needed for when UI is shutting down or to prevent conflicts where a drawer is open in two places at once. 
@@ -195,6 +168,9 @@ public:
 	/** Is any drawer other than the one specified by DrawerId opened */
 	bool IsAnyOtherDrawerOpened(const FName DrawerId) const;
 
+	/** Get unique name of this status bar */
+	FName GetStatusBarName() const;
+
 private:
 	/** Called when global focus changes which is used to determine if we should close an opened content browser drawer */
 	void OnGlobalFocusChanging(const FFocusEvent& FocusEvent, const FWeakWidgetPath& OldFocusedWidgetPath, const TSharedPtr<SWidget>& OldFocusedWidget, const FWidgetPath& NewFocusedWidgetPath, const TSharedPtr<SWidget>& NewFocusedWidget);
diff --git a/Engine/Source/Editor/StatusBar/Private/StatusBarSubsystem.cpp b/Engine/Source/Editor/StatusBar/Private/StatusBarSubsystem.cpp
--- a/Engine/Source/Editor/StatusBar/Private/StatusBarSubsystem.cpp
+++ b/Engine/Source/Editor/StatusBar/Private/StatusBarSubsystem.cpp
@@ -337,6 +337,51 @@ bool UStatusBarSubsystem::OpenOutputLogDrawer()
 	return false;
 }
 
+bool UStatusBarSubsystem::TryToggleDrawer(const FName DrawerId)
+{
+	bool bToggledSuccessfully = false;
+
+	TSharedPtr<SWindow> ParentWindow = FSlateApplication::Get().GetActiveTopLevelWindow();
+	if (!ParentWindow.IsValid())
+	{
+		if (TSharedPtr<SDockTab> ActiveTab = FGlobalTabmanager::Get()->GetActiveTab())
+		{
+			if (TSharedPtr<SDockTab> ActiveMajorTab = FGlobalTabmanager::Get()->GetMajorTabForTabManager(ActiveTab->GetTabManager()))
+			{
+				ParentWindow = ActiveMajorTab->GetParentWindow();
+			}
+		}
+	}
+
+	if (ParentWindow.IsValid() && ParentWindow->GetType() == EWindowType::Normal)
+	{
+		for (auto StatusBar : StatusBars)
+		{
+			FStatusBarData& SBData = StatusBar.Value;
+			if (TSharedPtr<SStatusBar> StatusBarPinned = SBData.StatusBarWidget.Pin())
+			{
+				TSharedPtr<SDockTab> ParentTab = StatusBarPinned->GetParentTab();
+				if (ParentTab && ParentTab->IsForeground() && ParentTab->GetParentWindow() == ParentWindow)
+				{
+					if (StatusBarPinned->IsDrawerOpened(DrawerId))
+					{
+						StatusBarPinned->DismissDrawer(nullptr);
+					}
+					else
+					{
+						StatusBarPinned->OpenDrawer(DrawerId);
+					}
+
+					bToggledSuccessfully = true;
+					break;
+				}
+			}
+		}
+	}
+
+	return bToggledSuccessfully;
+}
+
 bool UStatusBarSubsystem::ForceDismissDrawer()
 {
 	bool bWasDismissed = false;
@@ -344,11 +389,7 @@ bool UStatusBarSubsystem::ForceDismissDrawer()
 	{
 		if (TSharedPtr<SStatusBar> StatusBarPinned = StatusBar.Value.StatusBarWidget.Pin())
 		{
-			if (StatusBarPinned->IsDrawerOpened(StatusBarDrawerIds::ContentBrowser))
-			{
-				StatusBarPinned->DismissDrawer(nullptr);
-				bWasDismissed = true;
-			}
+			bWasDismissed |= StatusBarPinned->DismissDrawer(nullptr);
 		}
 	}
 
@@ -455,6 +496,14 @@ TSharedRef<SWidget> UStatusBarSubsystem::MakeStatusBarWidget(FName StatusBarName
 	return StatusBar;
 }
 
+void UStatusBarSubsystem::RegisterDrawer(FName StatusBarName, FStatusBarDrawer&& Drawer, int32 SlotIndex)
+{
+	if (TSharedPtr<SStatusBar> StatusBar = GetStatusBar(StatusBarName))
+	{
+		StatusBar->RegisterDrawer(MoveTemp(Drawer), SlotIndex);
+	}
+}
+
 FStatusBarMessageHandle UStatusBarSubsystem::PushStatusBarMessage(FName StatusBarName, const TAttribute<FText>& InMessage, const TAttribute<FText>& InHintText)
 {
 	if (TSharedPtr<SStatusBar> StatusBar = GetStatusBar(StatusBarName))
@@ -641,7 +690,7 @@ TSharedRef<SWidget> UStatusBarSubsystem::OnGetContentBrowser()
 	return StatusBarContentBrowser.ToSharedRef();
 }
 
-void UStatusBarSubsystem::OnContentBrowserOpened(TSharedRef<SStatusBar>& StatusBarWithContentBrowser)
+void UStatusBarSubsystem::OnContentBrowserOpened(FName StatusBarWithDrawerName)
 {
 	SNewUserTipNotification::Dismiss();
 
@@ -650,7 +699,7 @@ void UStatusBarSubsystem::OnContentBrowserOpened(TSharedRef<SStatusBar>& StatusB
 	{
 		if (TSharedPtr<SStatusBar> StatusBarPinned = StatusBar.Value.StatusBarWidget.Pin())
 		{
-			if (StatusBarWithContentBrowser != StatusBarPinned || StatusBarPinned->IsAnyOtherDrawerOpened(StatusBarDrawerIds::ContentBrowser))
+			if (StatusBarWithDrawerName != StatusBarPinned->GetStatusBarName() || StatusBarPinned->IsAnyOtherDrawerOpened(StatusBarDrawerIds::ContentBrowser))
 			{
 				StatusBarPinned->CloseDrawerImmediately();
 			}
@@ -704,14 +753,14 @@ TSharedRef<SWidget> UStatusBarSubsystem::OnGetOutputLog()
 	return StatusBarOutputLog.ToSharedRef();
 }
 
-void UStatusBarSubsystem::OnOutputLogOpened(TSharedRef<SStatusBar>& StatusBarWithContentBrowser)
+void UStatusBarSubsystem::OnOutputLogOpened(FName StatusBarWithDrawerName)
 {
 	// Dismiss any other content browser that is opened when one status bar opens it.  The content browser is a shared resource and shouldn't be in the layout twice
 	for (auto StatusBar : StatusBars)
 	{
 		if (TSharedPtr<SStatusBar> StatusBarPinned = StatusBar.Value.StatusBarWidget.Pin())
 		{
-			if (StatusBarWithContentBrowser != StatusBarPinned || StatusBarPinned->IsAnyOtherDrawerOpened(StatusBarDrawerIds::OutputLog))
+			if (StatusBarWithDrawerName != StatusBarPinned->GetStatusBarName() || StatusBarPinned->IsAnyOtherDrawerOpened(StatusBarDrawerIds::OutputLog))
 			{
 				StatusBarPinned->CloseDrawerImmediately();
 			}
diff --git a/Engine/Source/Editor/StatusBar/Public/StatusBarSubsystem.h b/Engine/Source/Editor/StatusBar/Public/StatusBarSubsystem.h
--- a/Engine/Source/Editor/StatusBar/Public/StatusBarSubsystem.h
+++ b/Engine/Source/Editor/StatusBar/Public/StatusBarSubsystem.h
@@ -5,9 +5,12 @@
 #include "CoreMinimal.h"
 #include "EditorSubsystem.h"
 #include "Framework/Notifications/NotificationManager.h"
+#include "Framework/SlateDelegates.h"
 
 #include "StatusBarSubsystem.generated.h"
 
+DECLARE_DELEGATE_OneParam(FOnStatusBarDrawerOpened, FName StatusBarName)
+DECLARE_DELEGATE_OneParam(FOnStatusBarDrawerDismissed, const TSharedPtr<SWidget>&)
 
 class SStatusBar;
 class SWindow;
@@ -47,6 +50,33 @@ private:
 	int32 Id;
 };
 
+struct FStatusBarDrawer
+{
+	FStatusBarDrawer(FName InUniqueId)
+		: UniqueId(InUniqueId)
+	{}
+
+	FName UniqueId;
+	FOnGetContent GetDrawerContentDelegate;
+	FOnStatusBarDrawerOpened OnDrawerOpenedDelegate;
+	FOnStatusBarDrawerDismissed OnDrawerDismissedDelegate;
+
+	TSharedPtr<SWidget> CustomWidget;
+	FText ButtonText;
+	FText ToolTipText;
+	const FSlateBrush* Icon;
+
+	bool operator==(const FName& OtherId) const
+	{
+		return UniqueId == OtherId;
+	}
+
+	bool operator==(const FStatusBarDrawer& Other) const
+	{
+		return UniqueId == Other.UniqueId;
+	}
+};
+
 struct FStatusBarData
 {
 	TWeakPtr<SStatusBar> StatusBarWidget;
@@ -94,7 +124,12 @@ public:
 	bool OpenOutputLogDrawer();
 
 	/**
-	 * Forces the drawer to dismiss. Usually it dismisses with focus. Only call this if there is some reason an open drawer would be invalid for the current state of the editor.
+	 * Tries to toggle the given drawer
+	 */
+	bool TryToggleDrawer(const FName DrawerId);
+
+	/**
+	 * Forces all drawers to dismiss. Usually it dismisses with focus. Only call this if there is some reason an open drawer would be invalid for the current state of the editor.
 	 */
 	bool ForceDismissDrawer();
 
@@ -105,6 +140,15 @@ public:
 	 */
 	TSharedRef<SWidget> MakeStatusBarWidget(FName StatusBarName, const TSharedRef<SDockTab>& InParentTab);
 
+	/**
+	 * Creates a new instance of a status bar widget
+	 *
+	 * @param StatusBarName	The name of the status bar to add the drawer to
+	 * @param Drawer		The drawer to add to the status bar
+	 * @param SlotIndex		The position at which to add the new drawer
+	 */
+	void RegisterDrawer(FName StatusBarName, FStatusBarDrawer&& Drawer, int32 SlotIndex = INDEX_NONE);
+
 	/** 
 	 * Pushes a new status bar message
 	 *
@@ -144,12 +188,12 @@ private:
 
 	TSharedPtr<SStatusBar> GetStatusBar(FName StatusBarName) const;
 	TSharedRef<SWidget> OnGetContentBrowser();
-	void OnContentBrowserOpened(TSharedRef<SStatusBar>& StatusBarWithContentBrowser);
+	void OnContentBrowserOpened(FName StatusBarWithDrawerName);
 	void OnContentBrowserDismissed(const TSharedPtr<SWidget>& NewlyFocusedWidget);
 	void HandleDeferredOpenContentBrowser(TSharedPtr<SWindow> ParentWindow);
 
 	TSharedRef<SWidget> OnGetOutputLog();
-	void OnOutputLogOpened(TSharedRef<SStatusBar>& StatusBarWithContentBrowser);
+	void OnOutputLogOpened(FName StatusBarWithDrawerName);
 	void OnOutputLogDismised(const TSharedPtr<SWidget>& NewlyFocusedWidget);
 
 	void OnDebugConsoleDrawerClosed();
diff --git a/Engine/Source/Editor/UMGEditor/Private/BlueprintModes/WidgetDesignerApplicationMode.cpp b/Engine/Source/Editor/UMGEditor/Private/BlueprintModes/WidgetDesignerApplicationMode.cpp
--- a/Engine/Source/Editor/UMGEditor/Private/BlueprintModes/WidgetDesignerApplicationMode.cpp
+++ b/Engine/Source/Editor/UMGEditor/Private/BlueprintModes/WidgetDesignerApplicationMode.cpp
@@ -8,6 +8,7 @@
 
 #include "WidgetBlueprintEditorToolbar.h"
 #include "UMGEditorModule.h"
+#include "StatusBarSubsystem.h"
 #include "ToolMenus.h"
 
 #include "TabFactory/PaletteTabSummoner.h"
@@ -15,7 +16,6 @@
 #include "TabFactory/HierarchyTabSummoner.h"
 #include "TabFactory/BindWidgetTabSummoner.h"
 #include "TabFactory/DesignerTabSummoner.h"
-#include "TabFactory/SequencerTabSummoner.h"
 #include "TabFactory/DetailsTabSummoner.h"
 #include "TabFactory/AnimationTabSummoner.h"
 #include "TabFactory/NavigationTabSummoner.h"
@@ -32,75 +32,64 @@ FWidgetDesignerApplicationMode::FWidgetDesignerApplicationMode(TSharedPtr<FWidge
 	// Override the default created category here since "Designer Editor" sounds awkward
 	WorkspaceMenuCategory = FWorkspaceItem::NewGroup(LOCTEXT("WorkspaceMenu_WidgetDesigner", "Widget Designer"));
 
-	TabLayout = FTabManager::NewLayout("WidgetBlueprintEditor_Designer_Layout_v4_3")
+	TabLayout = FTabManager::NewLayout("WidgetBlueprintEditor_Designer_Layout_v4_555")
 	->AddArea
 	(
 		FTabManager::NewPrimaryArea()
-		->SetOrientation(Orient_Vertical)
+		->SetOrientation(Orient_Horizontal)
 		->Split
 		(
 			FTabManager::NewSplitter()
-			->SetOrientation(Orient_Horizontal)
-			->SetSizeCoefficient( 0.70f )
+			->SetSizeCoefficient( 0.15f )
+			->SetOrientation(Orient_Vertical)
 			->Split
 			(
-				FTabManager::NewSplitter()
-				->SetSizeCoefficient( 0.15f )
-				->SetOrientation(Orient_Vertical)
-				->Split
-				(
-					FTabManager::NewStack()
-					->SetSizeCoefficient(0.5f)
-					->SetForegroundTab(FPaletteTabSummoner::TabID)
-					->AddTab(FPaletteTabSummoner::TabID, ETabState::OpenedTab)
-					->AddTab(FLibraryTabSummoner::TabID, ETabState::OpenedTab)
-				)
-				->Split
-				(
-					FTabManager::NewStack()
-					->SetSizeCoefficient(0.5f)
-					->SetForegroundTab(FHierarchyTabSummoner::TabID)
-					->AddTab(FHierarchyTabSummoner::TabID, ETabState::OpenedTab)
-					->AddTab(FBindWidgetTabSummoner::TabID, ETabState::OpenedTab)
-				)
+				FTabManager::NewStack()
+				->SetSizeCoefficient(0.5f)
+				->SetForegroundTab(FPaletteTabSummoner::TabID)
+				->AddTab(FPaletteTabSummoner::TabID, ETabState::OpenedTab)
+				->AddTab(FLibraryTabSummoner::TabID, ETabState::OpenedTab)
+			)
+			->Split
+			(
+				FTabManager::NewStack()
+				->SetSizeCoefficient(0.5f)
+				->SetForegroundTab(FHierarchyTabSummoner::TabID)
+				->AddTab(FHierarchyTabSummoner::TabID, ETabState::OpenedTab)
+				->AddTab(FBindWidgetTabSummoner::TabID, ETabState::OpenedTab)
 			)
+		)
+		->Split
+		(
+			FTabManager::NewSplitter()
+			->SetSizeCoefficient(0.85f)
+			->SetOrientation(Orient_Vertical)
 			->Split
 			(
 				FTabManager::NewSplitter()
-				->SetSizeCoefficient( 0.85f )
+				->SetSizeCoefficient(0.7f)
 				->SetOrientation(Orient_Horizontal)
 				->Split
 				(
 					FTabManager::NewStack()
 					->SetHideTabWell(true)
+					->SetSizeCoefficient(0.85f)
 					->AddTab( FDesignerTabSummoner::TabID, ETabState::OpenedTab )
 				)
 				->Split
 				(
 					FTabManager::NewStack()
-					->SetSizeCoefficient( 0.35f )
-					->AddTab( FDetailsTabSummoner::TabID, ETabState::OpenedTab )
+					->SetSizeCoefficient(0.35f)
+					->AddTab(FDetailsTabSummoner::TabID, ETabState::OpenedTab)
 				)
 			)
-		)
-		->Split
-		(
-			FTabManager::NewSplitter()
-			->SetOrientation(Orient_Horizontal)
-			->SetSizeCoefficient( 0.30f )
-			->Split
-			(
-				FTabManager::NewStack()
-				->SetSizeCoefficient(.15f)
-				->AddTab(FAnimationTabSummoner::TabID, ETabState::OpenedTab)
-			)
 			->Split
 			(
 				FTabManager::NewStack()
-				->SetSizeCoefficient(.85f)
-				->SetForegroundTab(FSequencerTabSummoner::TabID)
-				->AddTab(FSequencerTabSummoner::TabID, ETabState::OpenedTab)
-				->AddTab(FBlueprintEditorTabs::CompilerResultsID, ETabState::OpenedTab)
+				->SetSizeCoefficient( 0.3f )
+				->AddTab(FAnimationTabSummoner::TabID, ETabState::ClosedTab)
+				->AddTab(FBlueprintEditorTabs::CompilerResultsID, ETabState::ClosedTab)
+				->SetForegroundTab(FAnimationTabSummoner::TabID)
 			)
 		)
 	);
@@ -113,7 +102,6 @@ FWidgetDesignerApplicationMode::FWidgetDesignerApplicationMode(TSharedPtr<FWidge
 	TabFactories.RegisterFactory(MakeShareable(new FBindWidgetTabSummoner(InWidgetEditor)));
 	TabFactories.RegisterFactory(MakeShareable(new FPaletteTabSummoner(InWidgetEditor)));
 	TabFactories.RegisterFactory(MakeShareable(new FLibraryTabSummoner(InWidgetEditor)));
-	TabFactories.RegisterFactory(MakeShareable(new FSequencerTabSummoner(InWidgetEditor)));
 	TabFactories.RegisterFactory(MakeShareable(new FAnimationTabSummoner(InWidgetEditor)));
 	TabFactories.RegisterFactory(MakeShareable(new FCompilerResultsSummoner(InWidgetEditor)));
 	TabFactories.RegisterFactory(MakeShareable(new FNavigationTabSummoner(InWidgetEditor)));
@@ -151,8 +139,18 @@ void FWidgetDesignerApplicationMode::PreDeactivateMode()
 void FWidgetDesignerApplicationMode::PostActivateMode()
 {
 	//FWidgetBlueprintApplicationMode::PostActivateMode();
+
 	TSharedPtr<FWidgetBlueprintEditor> BP = GetBlueprintEditor();
 
+	FStatusBarDrawer WidgetAnimSequencerDrawer(FAnimationTabSummoner::WidgetAnimSequencerDrawerID);
+	WidgetAnimSequencerDrawer.GetDrawerContentDelegate.BindSP(BP.Get(), &FWidgetBlueprintEditor::OnGetWidgetAnimSequencer);
+	WidgetAnimSequencerDrawer.OnDrawerOpenedDelegate.BindSP(BP.Get(), &FWidgetBlueprintEditor::OnWidgetAnimSequencerOpened);
+	WidgetAnimSequencerDrawer.OnDrawerDismissedDelegate.BindSP(BP.Get(), &FWidgetBlueprintEditor::OnWidgetAnimSequencerDismissed);
+	WidgetAnimSequencerDrawer.ButtonText = LOCTEXT("StatusBar_WidgetAnimSequencer", "Animations");
+	WidgetAnimSequencerDrawer.ToolTipText = LOCTEXT("StatusBar_WidgetAnimSequencerToolTip", "Opens animation sequencer (Ctrl+Shift+Space Bar).");
+	WidgetAnimSequencerDrawer.Icon = FAppStyle::Get().GetBrush("UMGEditor.AnimTabIcon");
+	BP->RegisterDrawer(MoveTemp(WidgetAnimSequencerDrawer), 1);
+
 	BP->OnEnteringDesigner();
 }
 
diff --git a/Engine/Source/Editor/UMGEditor/Private/Settings/WidgetDesignerSettings.cpp b/Engine/Source/Editor/UMGEditor/Private/Settings/WidgetDesignerSettings.cpp
--- a/Engine/Source/Editor/UMGEditor/Private/Settings/WidgetDesignerSettings.cpp
+++ b/Engine/Source/Editor/UMGEditor/Private/Settings/WidgetDesignerSettings.cpp
@@ -14,6 +14,8 @@ UWidgetDesignerSettings::UWidgetDesignerSettings()
 	bShowOutlines = true;
 	bExecutePreConstructEvent = true;
 	bRespectLocks = true;
+	CreateOnCompile = EDisplayOnCompile::DoC_ErrorsOrWarnings;
+	DismissOnCompile = EDisplayOnCompile::DoC_ErrorsOrWarnings;
 }
 
 #if WITH_EDITOR
diff --git a/Engine/Source/Editor/UMGEditor/Private/TabFactory/AnimationTabSummoner.cpp b/Engine/Source/Editor/UMGEditor/Private/TabFactory/AnimationTabSummoner.cpp
--- a/Engine/Source/Editor/UMGEditor/Private/TabFactory/AnimationTabSummoner.cpp
+++ b/Engine/Source/Editor/UMGEditor/Private/TabFactory/AnimationTabSummoner.cpp
@@ -8,6 +8,7 @@
 #include "Animation/WidgetAnimation.h"
 #include "WidgetBlueprint.h"
 #include "WidgetBlueprintEditorUtils.h"
+#include "Widgets/Images/SImage.h"
 #include "Widgets/Input/SButton.h"
 #include "Widgets/Layout/SScrollBorder.h"
 #include "Widgets/Views/STableViewBase.h"
@@ -19,6 +20,7 @@
 #endif // WITH_EDITOR
 #include "Blueprint/WidgetTree.h"
 
+#include "UMGEditorActions.h"
 #include "UMGStyle.h"
 #include "Widgets/Input/SSearchBox.h"
 #include "Widgets/Text/SInlineEditableTextBlock.h"
@@ -32,10 +34,12 @@
 #define LOCTEXT_NAMESPACE "UMG"
 
 const FName FAnimationTabSummoner::TabID(TEXT("Animations"));
+const FName FAnimationTabSummoner::WidgetAnimSequencerDrawerID(TEXT("WidgetAnimSequencer"));
 
-FAnimationTabSummoner::FAnimationTabSummoner(TSharedPtr<class FWidgetBlueprintEditor> InBlueprintEditor)
+FAnimationTabSummoner::FAnimationTabSummoner(TSharedPtr<class FWidgetBlueprintEditor> InBlueprintEditor, bool bInIsDrawerTab)
 		: FWorkflowTabFactory(TabID, InBlueprintEditor)
 		, BlueprintEditor(InBlueprintEditor)
+		, bIsDrawerTab(bInIsDrawerTab)
 {
 	TabLabel = LOCTEXT("AnimationsTabLabel", "Animations");
 	TabIcon = FSlateIcon(FUMGStyle::GetStyleSetName(), "Animations.TabIcon");
@@ -239,6 +243,11 @@ private:
 			Blueprint->Animations.Add(WidgetAnimation);
 			ListItem.Pin()->bNewAnimation = false;
 			FBlueprintEditorUtils::MarkBlueprintAsStructurallyModified(Blueprint);
+
+			if (TSharedPtr<FWidgetBlueprintEditor> WidgetBlueprintEditorPin = BlueprintEditor.Pin())
+			{
+				WidgetBlueprintEditorPin->NotifyWidgetAnimListChanged();
+			}
 		}
 	}
 private:
@@ -254,12 +263,22 @@ public:
 	SLATE_BEGIN_ARGS( SUMGAnimationList	) {}
 	SLATE_END_ARGS()
 
-	void Construct( const FArguments& InArgs, TSharedPtr<FWidgetBlueprintEditor> InBlueprintEditor )
+	~SUMGAnimationList()
+	{
+		if (TSharedPtr<FWidgetBlueprintEditor> WidgetBlueprintEditorPin = BlueprintEditor.Pin())
+		{
+			WidgetBlueprintEditorPin->OnWidgetAnimationsUpdated.RemoveAll(this);
+		}
+	}
+
+	void Construct( const FArguments& InArgs, TSharedPtr<FWidgetBlueprintEditor> InBlueprintEditor, bool bInIsDrawerTab )
 	{
 		BlueprintEditor = InBlueprintEditor;
+		bIsDrawerTab = bInIsDrawerTab;
 
 		InBlueprintEditor->GetOnWidgetBlueprintTransaction().AddSP( this, &SUMGAnimationList::OnWidgetBlueprintTransaction );
 		InBlueprintEditor->OnEnterWidgetDesigner.AddSP(this, &SUMGAnimationList::OnEnteringDesignerMode);
+		InBlueprintEditor->OnWidgetAnimationsUpdated.AddSP(this, &SUMGAnimationList::OnUpdatedAnimationList);
 
 		SAssignNew(AnimationListView, SWidgetAnimationListView)
 			.ItemHeight(20.0f)
@@ -274,38 +293,62 @@ public:
 		[
 			SNew(SBorder)
 			.BorderImage(FEditorStyle::GetBrush("ToolPanel.GroupBorder"))
+			.Padding( FMargin(bIsDrawerTab ? 8.0 : 2.0, 2.0) )
 			[
-				SNew(SVerticalBox)
-				+ SVerticalBox::Slot()
-				.Padding( 2 )
-				.AutoHeight()
+				SNew(SOverlay)
+				+SOverlay::Slot()
+				.HAlign(HAlign_Fill)
+				.VAlign(VAlign_Fill)
 				[
-					SNew( SHorizontalBox )
-					+ SHorizontalBox::Slot()
-					.Padding(0)
-					.VAlign( VAlign_Center )
-					.AutoWidth()
+					SNew(SSplitter)
+					+ SSplitter::Slot()
+					.Value(0.15f)
 					[
-						SNew(SEditorHeaderButton)
-						.OnClicked( this, &SUMGAnimationList::OnNewAnimationClicked )
-						.Text( LOCTEXT("NewAnimationButtonText", "Animation") )
+						SNew(SVerticalBox)
+						+ SVerticalBox::Slot()
+						.Padding( 2 )
+						.AutoHeight()
+						[
+							SNew( SHorizontalBox )
+							+ SHorizontalBox::Slot()
+							.Padding(0)
+							.VAlign( VAlign_Center )
+							.AutoWidth()
+							[
+								SNew(SEditorHeaderButton)
+								.OnClicked( this, &SUMGAnimationList::OnNewAnimationClicked )
+								.Text( LOCTEXT("NewAnimationButtonText", "Animation") )
+							]
+							+ SHorizontalBox::Slot()
+							.Padding(2.0f, 0.0f)
+							.VAlign( VAlign_Center )
+							[
+								SAssignNew(SearchBoxPtr, SSearchBox)
+								.HintText(LOCTEXT("Search Animations", "Search Animations"))
+								.OnTextChanged(this, &SUMGAnimationList::OnSearchChanged)
+							]
+						]
+						+ SVerticalBox::Slot()
+						.FillHeight(1.0f)
+						[
+							SNew(SScrollBorder, AnimationListView.ToSharedRef())
+							[
+								AnimationListView.ToSharedRef()
+							]
+						]
 					]
-					+ SHorizontalBox::Slot()
-					.Padding(2.0f, 0.0f)
-					.VAlign( VAlign_Center )
+					+ SSplitter::Slot()
+					.Value(0.85f)
 					[
-						SAssignNew(SearchBoxPtr, SSearchBox)
-						.HintText(LOCTEXT("Search Animations", "Search Animations"))
-						.OnTextChanged(this, &SUMGAnimationList::OnSearchChanged)
+						bIsDrawerTab ? BlueprintEditor.Pin()->CreateSequencerDrawerWidget() : BlueprintEditor.Pin()->CreateSequencerTabWidget()
 					]
 				]
-				+ SVerticalBox::Slot()
-				.FillHeight(1.0f)
+				+SOverlay::Slot()
+				.HAlign(HAlign_Right)
+				.VAlign(VAlign_Top)
+				.Padding(FMargin(24.0, 10.0))
 				[
-					SNew(SScrollBorder, AnimationListView.ToSharedRef())
-					[
-						AnimationListView.ToSharedRef()
-					]
+					CreateDrawerDockButton()
 				]
 			]
 		];
@@ -315,7 +358,7 @@ public:
 		CreateCommandList();
 	}
 
-	virtual FReply OnKeyDown( const FGeometry& MyGeometry, const FKeyEvent& InKeyEvent )
+	virtual FReply OnKeyDown( const FGeometry& MyGeometry, const FKeyEvent& InKeyEvent ) override
 	{
 		FReply Reply = FReply::Unhandled();
 		if( CommandList->ProcessCommandBindings( InKeyEvent ) )
@@ -328,20 +371,30 @@ public:
 
 private:
 
-	void UpdateAnimationList()
+	void OnUpdatedAnimationList()
 	{
 		Animations.Empty();
 
 		const TArray<UWidgetAnimation*>& WidgetAnimations = BlueprintEditor.Pin()->GetWidgetBlueprintObj()->Animations;
 
-		for( UWidgetAnimation* Animation : WidgetAnimations )
+		for (UWidgetAnimation* Animation : WidgetAnimations)
 		{
-			Animations.Add( MakeShareable( new FWidgetAnimationListItem( Animation ) ) );
+			Animations.Add(MakeShareable(new FWidgetAnimationListItem(Animation)));
 		}
-		
+
 		AnimationListView->RequestListRefresh();
 	}
 
+	void UpdateAnimationList()
+	{
+		// There may be multiple sequencers acting as a view for our widget
+		// Let the BP editor handle updates as it as aware of all possible sequencers
+		if (TSharedPtr<FWidgetBlueprintEditor> WidgetBlueprintEditorPin = BlueprintEditor.Pin())
+		{
+			WidgetBlueprintEditorPin->NotifyWidgetAnimListChanged();
+		}
+	}
+
 	void OnEnteringDesignerMode()
 	{
 		UpdateAnimationList();
@@ -545,6 +598,55 @@ private:
 			FExecuteAction::CreateSP(this, &SUMGAnimationList::OnRenameAnimation),
 			FCanExecuteAction::CreateSP(this, &SUMGAnimationList::CanExecuteContextMenuAction)
 			);
+
+		CommandList->MapAction(FUMGEditorCommands::Get().OpenAnimDrawer,
+			FExecuteAction::CreateSP(this, &SUMGAnimationList::ToggleAnimDrawer)
+		);
+	}
+
+	FReply CreateDrawerDockButtonClicked()
+	{
+		if (TSharedPtr<FWidgetBlueprintEditor> WidgetEditor = BlueprintEditor.Pin())
+		{
+			WidgetEditor->DockInLayoutClicked();
+		}
+
+		return FReply::Handled();
+	}
+
+	TSharedRef<SWidget> CreateDrawerDockButton()
+	{
+		if(bIsDrawerTab)
+		{
+			return 
+				SNew(SButton)
+				.ButtonStyle(FAppStyle::Get(), "SimpleButton")
+				.ToolTipText(LOCTEXT("DockInLayout_Tooltip", "Docks animation drawer in tab."))
+				.ContentPadding(FMargin(1, 0))
+				.OnClicked(this, &SUMGAnimationList::CreateDrawerDockButtonClicked)
+				[
+					SNew(SHorizontalBox)
+					+ SHorizontalBox::Slot()
+					.AutoWidth()
+					.VAlign(VAlign_Center)
+					.Padding(4.0, 0.0f)
+					[
+						SNew(SImage)
+						.ColorAndOpacity(FSlateColor::UseForeground())
+						.Image(FAppStyle::Get().GetBrush("EditorViewport.SubMenu.Layouts"))
+					]
+					+ SHorizontalBox::Slot()
+					.VAlign(VAlign_Center)
+					.Padding(4.0, 0.0f)
+					[
+						SNew(STextBlock)
+						.Text(LOCTEXT("DockInLayout", "Dock in Layout"))
+						.ColorAndOpacity(FSlateColor::UseForeground())
+					]
+				];
+		}
+
+		return SNullWidget::NullWidget;
 	}
 
 	bool CanExecuteContextMenuAction() const
@@ -619,12 +721,22 @@ private:
 		AnimationListView->RequestScrollIntoView( SelectedAnimation );
 	}
 
+	void ToggleAnimDrawer()
+	{
+		if (TSharedPtr<FWidgetBlueprintEditor> WidgetEditor = BlueprintEditor.Pin())
+		{
+			WidgetEditor->ToggleAnimDrawer();
+		}
+	}
+
+
 private:
 	TSharedPtr<FUICommandList> CommandList;
 	TWeakPtr<FWidgetBlueprintEditor> BlueprintEditor;
 	TSharedPtr<SWidgetAnimationListView> AnimationListView;
 	TArray< TSharedPtr<FWidgetAnimationListItem> > Animations;
 	TSharedPtr<SSearchBox> SearchBoxPtr;
+	bool bIsDrawerTab;
 };
 
 
@@ -632,7 +744,7 @@ TSharedRef<SWidget> FAnimationTabSummoner::CreateTabBody(const FWorkflowTabSpawn
 {
 	TSharedPtr<FWidgetBlueprintEditor> BlueprintEditorPinned = BlueprintEditor.Pin();
 
-	return SNew( SUMGAnimationList, BlueprintEditorPinned );
+	return SNew( SUMGAnimationList, BlueprintEditorPinned, bIsDrawerTab);
 	
 }
 
diff --git a/Engine/Source/Editor/UMGEditor/Private/TabFactory/AnimationTabSummoner.h b/Engine/Source/Editor/UMGEditor/Private/TabFactory/AnimationTabSummoner.h
--- a/Engine/Source/Editor/UMGEditor/Private/TabFactory/AnimationTabSummoner.h
+++ b/Engine/Source/Editor/UMGEditor/Private/TabFactory/AnimationTabSummoner.h
@@ -11,12 +11,14 @@ struct FAnimationTabSummoner : public FWorkflowTabFactory
 {
 public:
 	static const FName TabID;
+	static const FName WidgetAnimSequencerDrawerID;
 	
 public:
-	FAnimationTabSummoner(TSharedPtr<class FWidgetBlueprintEditor> InBlueprintEditor);
+	FAnimationTabSummoner(TSharedPtr<class FWidgetBlueprintEditor> InBlueprintEditor, bool bInIsDrawerTab = false);
 	
 	virtual TSharedRef<SWidget> CreateTabBody(const FWorkflowTabSpawnInfo& Info) const override;
 	
 protected:
 	TWeakPtr<class FWidgetBlueprintEditor> BlueprintEditor;
+	bool bIsDrawerTab;
 };
diff --git a/Engine/Source/Editor/UMGEditor/Private/TabFactory/SequencerTabSummoner.cpp b/Engine/Source/Editor/UMGEditor/Private/TabFactory/SequencerTabSummoner.cpp
deleted file mode 100644
diff --git a/Engine/Source/Editor/UMGEditor/Private/TabFactory/SequencerTabSummoner.h b/Engine/Source/Editor/UMGEditor/Private/TabFactory/SequencerTabSummoner.h
deleted file mode 100644
diff --git a/Engine/Source/Editor/UMGEditor/Private/UMGEditorActions.cpp b/Engine/Source/Editor/UMGEditor/Private/UMGEditorActions.cpp
--- a/Engine/Source/Editor/UMGEditor/Private/UMGEditorActions.cpp
+++ b/Engine/Source/Editor/UMGEditor/Private/UMGEditorActions.cpp
@@ -11,6 +11,22 @@ void FUMGEditorCommands::RegisterCommands()
 	UI_COMMAND(SetImageAsThumbnail, "Capture and Set as Thumbnail", "Captures the current state of the widget blueprint and sets it as the thumbnail for this asset.", EUserInterfaceActionType::Button, FInputChord());
 	UI_COMMAND(ClearCustomThumbnail, "Clear Thumbnail", "Removes the image used as thumbnail and enables automatic thumbnail generation.", EUserInterfaceActionType::Button, FInputChord());
 
+	UI_COMMAND(DismissOnCompile_ErrorsAndWarnings, "No Errors or Warnings", "Automatically dismiss compile tab when no errors or warnings occur", EUserInterfaceActionType::RadioButton, FInputChord());
+	UI_COMMAND(DismissOnCompile_Errors, "No Errors", "Automatically dismiss compile tab when no errors occur", EUserInterfaceActionType::RadioButton, FInputChord());
+	UI_COMMAND(DismissOnCompile_Warnings, "No Warnings", "Automatically dismiss compile tab when no warnings occur (even if there are errors)", EUserInterfaceActionType::RadioButton, FInputChord());
+	UI_COMMAND(DismissOnCompile_Never, "Never", "Never automatically dismiss the compile tab", EUserInterfaceActionType::RadioButton, FInputChord());
+
+	UI_COMMAND(CreateOnCompile_ErrorsAndWarnings, "On Errors or Warnings", "Automatically create compile tab when errors or warnings occur", EUserInterfaceActionType::RadioButton, FInputChord());
+	UI_COMMAND(CreateOnCompile_Errors, "On Errors", "Automatically create compile tab when errors occur", EUserInterfaceActionType::RadioButton, FInputChord());
+	UI_COMMAND(CreateOnCompile_Warnings, "On Warnings", "Automatically create compile tab when warnings occur (but not errors)", EUserInterfaceActionType::RadioButton, FInputChord());
+	UI_COMMAND(CreateOnCompile_Never, "Never", "Never automatically create the compile tab", EUserInterfaceActionType::RadioButton, FInputChord());
+
+#if PLATFORM_MAC
+	// On mac command and ctrl are automatically swapped. Command + Space is spotlight search so we use ctrl+space on mac to avoid the conflict
+	UI_COMMAND(OpenAnimDrawer, "Open Animation Browser Drawer", "Opens the animation drawer from the status bar", EUserInterfaceActionType::Button, FInputChord(EModifierKey::Command | EModifierKey::Shift, EKeys::SpaceBar));
+#else
+	UI_COMMAND(OpenAnimDrawer, "Open Animation Browser Drawer", "Opens the animation drawer from the status bar", EUserInterfaceActionType::Button, FInputChord(EModifierKey::Control | EModifierKey::Shift, EKeys::SpaceBar));
+#endif
 }
 
 #undef LOCTEXT_NAMESPACE
\ No newline at end of file
diff --git a/Engine/Source/Editor/UMGEditor/Private/UMGEditorActions.h b/Engine/Source/Editor/UMGEditor/Private/UMGEditorActions.h
--- a/Engine/Source/Editor/UMGEditor/Private/UMGEditorActions.h
+++ b/Engine/Source/Editor/UMGEditor/Private/UMGEditorActions.h
@@ -30,6 +30,17 @@ public:
 	TSharedPtr< FUICommandInfo > ExportAsPNG;
 	TSharedPtr< FUICommandInfo > SetImageAsThumbnail;
 	TSharedPtr< FUICommandInfo > ClearCustomThumbnail;
+	TSharedPtr< FUICommandInfo > OpenAnimDrawer;
+
+	TSharedPtr< FUICommandInfo > DismissOnCompile_ErrorsAndWarnings;
+	TSharedPtr< FUICommandInfo > DismissOnCompile_Errors;
+	TSharedPtr< FUICommandInfo > DismissOnCompile_Warnings;
+	TSharedPtr< FUICommandInfo > DismissOnCompile_Never;
+
+	TSharedPtr< FUICommandInfo > CreateOnCompile_ErrorsAndWarnings;
+	TSharedPtr< FUICommandInfo > CreateOnCompile_Errors;
+	TSharedPtr< FUICommandInfo > CreateOnCompile_Warnings;
+	TSharedPtr< FUICommandInfo > CreateOnCompile_Never;
 
 	/**
 	 * Initialize commands
diff --git a/Engine/Source/Editor/UMGEditor/Private/WidgetBlueprintEditor.cpp b/Engine/Source/Editor/UMGEditor/Private/WidgetBlueprintEditor.cpp
--- a/Engine/Source/Editor/UMGEditor/Private/WidgetBlueprintEditor.cpp
+++ b/Engine/Source/Editor/UMGEditor/Private/WidgetBlueprintEditor.cpp
@@ -13,6 +13,7 @@
 #include "Blueprint/WidgetBlueprintGeneratedClass.h"
 #include "Kismet2/BlueprintEditorUtils.h"
 #include "WidgetBlueprint.h"
+#include "StatusBarSubsystem.h"
 #include "Editor.h"
 
 #if WITH_EDITOR
@@ -69,6 +70,8 @@
 #include "ImageUtils.h"
 #include "Serialization/BufferArchive.h"
 #include "Widgets/SVirtualWindow.h"
+#include "TabFactory/AnimationTabSummoner.h"
+#include "Kismet/Public/BlueprintEditorTabs.h"
 
 #define LOCTEXT_NAMESPACE "UMG"
 
@@ -77,6 +80,7 @@ FWidgetBlueprintEditor::FWidgetBlueprintEditor()
 	, PreviewBlueprint(nullptr)
 	, bIsSimulateEnabled(false)
 	, bIsRealTime(true)
+	, bIsSequencerDrawerOpen(false)
 	, bRefreshGeneratedClassAnimations(false)
 	, bUpdatingSequencerSelection(false)
 	, bUpdatingExternalSelection(false)
@@ -102,12 +106,15 @@ FWidgetBlueprintEditor::~FWidgetBlueprintEditor()
 	}
 
 	FCoreUObjectDelegates::OnObjectsReplaced.RemoveAll(this);
-	
-	if ( Sequencer.IsValid() )
+
+	for (TWeakPtr<ISequencer> SequencerPtr : Sequencers)
 	{
-		Sequencer->OnMovieSceneDataChanged().RemoveAll( this );
-		Sequencer->OnMovieSceneBindingsPasted().RemoveAll( this );
-		Sequencer.Reset();
+		if (TSharedPtr<ISequencer> Sequencer = SequencerPtr.Pin())
+		{
+			Sequencer->OnMovieSceneDataChanged().RemoveAll(this);
+			Sequencer->OnMovieSceneBindingsPasted().RemoveAll(this);
+			Sequencer.Reset();
+		}
 	}
 
 	// Un-Register sequencer menu extenders.
@@ -137,6 +144,13 @@ void FWidgetBlueprintEditor::InitWidgetBlueprintEditor(const EToolkitMode::Type
 
 	InitBlueprintEditor(Mode, InitToolkitHost, InBlueprints, bShouldOpenInDefaultsMode);
 
+	// We only show compile tab results on error
+	TSharedPtr<SDockTab> CompileResultsTab = GetToolkitHost()->GetTabManager()->FindExistingLiveTab(FBlueprintEditorTabs::CompilerResultsID);
+	if (CompileResultsTab)
+	{
+		CompileResultsTab->RequestCloseTab();
+	}
+
 	// register for any objects replaced
 	FCoreUObjectDelegates::OnObjectsReplaced.AddSP(this, &FWidgetBlueprintEditor::OnObjectsReplaced);
 
@@ -200,6 +214,12 @@ TSharedPtr<FExtender> FWidgetBlueprintEditor::CreateMenuExtender()
 		GetToolkitCommands(),
 		FMenuExtensionDelegate::CreateSP(this, &FWidgetBlueprintEditor::FillAssetMenu));
 
+	MenuExtender->AddMenuExtension(
+		"FileLoadAndSave",
+		EExtensionHook::After,
+		GetToolkitCommands(),
+		FMenuExtensionDelegate::CreateSP(this, &FWidgetBlueprintEditor::CustomizeWidgetCompileOptions));
+
 	return MenuExtender;
 }
 
@@ -222,6 +242,37 @@ void FWidgetBlueprintEditor::FillAssetMenu(FMenuBuilder& MenuBuilder)
 	MenuBuilder.EndSection();
 }
 
+void FWidgetBlueprintEditor::CustomizeWidgetCompileOptions(FMenuBuilder& InMenuBuilder)
+{
+	InMenuBuilder.AddSubMenu(
+		LOCTEXT("CreateCompileTab", "Create Compile Tab"),
+		LOCTEXT("CreateCompileTab_ToolTip", "Displays Compile tab when hidden based on compilation results."),
+		FNewMenuDelegate::CreateStatic(&FWidgetBlueprintEditor::AddCreateCompileTabSubMenu));
+
+	InMenuBuilder.AddSubMenu(
+		LOCTEXT("DismissCompileTab", "Dismiss Compile Tab"),
+		LOCTEXT("DismissCompileTab_ToolTip", "Dismisses compile tab based on compilation results."),
+		FNewMenuDelegate::CreateStatic(&FWidgetBlueprintEditor::AddDismissCompileTabSubMenu));
+}
+
+void FWidgetBlueprintEditor::AddCreateCompileTabSubMenu(FMenuBuilder& InMenuBuilder)
+{
+	const FUMGEditorCommands& Commands = FUMGEditorCommands::Get();
+	InMenuBuilder.AddMenuEntry(Commands.CreateOnCompile_ErrorsAndWarnings);
+	InMenuBuilder.AddMenuEntry(Commands.CreateOnCompile_Errors);
+	InMenuBuilder.AddMenuEntry(Commands.CreateOnCompile_Warnings);
+	InMenuBuilder.AddMenuEntry(Commands.CreateOnCompile_Never);
+}
+
+void FWidgetBlueprintEditor::AddDismissCompileTabSubMenu(FMenuBuilder& InMenuBuilder)
+{
+	const FUMGEditorCommands& Commands = FUMGEditorCommands::Get();
+	InMenuBuilder.AddMenuEntry(Commands.DismissOnCompile_ErrorsAndWarnings);
+	InMenuBuilder.AddMenuEntry(Commands.DismissOnCompile_Errors);
+	InMenuBuilder.AddMenuEntry(Commands.DismissOnCompile_Warnings);
+	InMenuBuilder.AddMenuEntry(Commands.DismissOnCompile_Never);
+}
+
 void FWidgetBlueprintEditor::BindToolkitCommands()
 {
 	FUMGEditorCommands::Register();
@@ -252,6 +303,38 @@ void FWidgetBlueprintEditor::BindToolkitCommands()
 			FCanExecuteAction::CreateSP(this, &FWidgetBlueprintEditor::IsImageUsedAsThumbnail)
 		)
 	);
+
+	GetToolkitCommands()->MapAction(FUMGEditorCommands::Get().OpenAnimDrawer,
+		FExecuteAction::CreateSP(this, &FWidgetBlueprintEditor::ToggleAnimDrawer)
+	);
+
+	auto MapCreateOnCompileAction = [&](const TSharedPtr<FUICommandInfo>& InUICommand, EDisplayOnCompile InCreateOnCompile)
+	{
+		ToolkitCommands->MapAction(
+			InUICommand,
+			FExecuteAction::CreateStatic(&FWidgetBlueprintEditor::SetCreateOnCompileSetting, InCreateOnCompile),
+			FCanExecuteAction(),
+			FIsActionChecked::CreateStatic(&FWidgetBlueprintEditor::IsCreateOnCompileSet, InCreateOnCompile)
+		);
+	};
+	MapCreateOnCompileAction(FUMGEditorCommands::Get().CreateOnCompile_ErrorsAndWarnings, EDisplayOnCompile::DoC_ErrorsOrWarnings);
+	MapCreateOnCompileAction(FUMGEditorCommands::Get().CreateOnCompile_Errors, EDisplayOnCompile::DoC_ErrorsOnly);
+	MapCreateOnCompileAction(FUMGEditorCommands::Get().CreateOnCompile_Warnings, EDisplayOnCompile::DoC_WarningsOnly);
+	MapCreateOnCompileAction(FUMGEditorCommands::Get().CreateOnCompile_Never, EDisplayOnCompile::DoC_Never);
+
+	auto MapDismissOnCompileAction = [&](const TSharedPtr<FUICommandInfo>& InUICommand, EDisplayOnCompile InDismissOnCompile)
+	{
+		ToolkitCommands->MapAction(
+			InUICommand,
+			FExecuteAction::CreateStatic(&FWidgetBlueprintEditor::SetDismissOnCompileSetting, InDismissOnCompile),
+			FCanExecuteAction(),
+			FIsActionChecked::CreateStatic(&FWidgetBlueprintEditor::IsDismissOnCompileSet, InDismissOnCompile)
+		);
+	};
+	MapDismissOnCompileAction(FUMGEditorCommands::Get().DismissOnCompile_ErrorsAndWarnings, EDisplayOnCompile::DoC_ErrorsOrWarnings);
+	MapDismissOnCompileAction(FUMGEditorCommands::Get().DismissOnCompile_Errors, EDisplayOnCompile::DoC_ErrorsOnly);
+	MapDismissOnCompileAction(FUMGEditorCommands::Get().DismissOnCompile_Warnings, EDisplayOnCompile::DoC_WarningsOnly);
+	MapDismissOnCompileAction(FUMGEditorCommands::Get().DismissOnCompile_Never, EDisplayOnCompile::DoC_Never);
 }
 
 void FWidgetBlueprintEditor::TakeSnapshot()
@@ -345,6 +428,32 @@ bool FWidgetBlueprintEditor::IsPreviewWidgetInitialized()
 	return GetPreview() != nullptr;
 }
 
+void FWidgetBlueprintEditor::SetCreateOnCompileSetting(EDisplayOnCompile InCreateOnCompile)
+{
+	UWidgetDesignerSettings* Settings = GetMutableDefault<UWidgetDesignerSettings>();
+	Settings->CreateOnCompile = InCreateOnCompile;
+	Settings->SaveConfig();
+}
+
+void FWidgetBlueprintEditor::SetDismissOnCompileSetting(EDisplayOnCompile InDismissOnCompile)
+{
+	UWidgetDesignerSettings* Settings = GetMutableDefault<UWidgetDesignerSettings>();
+	Settings->DismissOnCompile = InDismissOnCompile;
+	Settings->SaveConfig();
+}
+
+bool FWidgetBlueprintEditor::IsCreateOnCompileSet(EDisplayOnCompile InCreateOnCompile)
+{
+	const UWidgetDesignerSettings* Settings = GetDefault<UWidgetDesignerSettings>();
+	return Settings->CreateOnCompile == InCreateOnCompile;
+}
+
+bool FWidgetBlueprintEditor::IsDismissOnCompileSet(EDisplayOnCompile InDismissOnCompile)
+{
+	const UWidgetDesignerSettings* Settings = GetDefault<UWidgetDesignerSettings>();
+	return Settings->DismissOnCompile == InDismissOnCompile;
+}
+
 void FWidgetBlueprintEditor::OpenCreateNativeBaseClassDialog()
 {
 	FGameProjectGenerationModule::Get().OpenAddCodeToProjectDialog(
@@ -458,7 +567,8 @@ bool FWidgetBlueprintEditor::IsBindingSelected(const FMovieSceneBinding& InBindi
 		return true;
 	}
 
-	UMovieSceneSequence* AnimationSequence = GetSequencer().Get()->GetFocusedMovieSceneSequence();
+	TSharedPtr<ISequencer>& ActiveSequencer = GetSequencer();
+	UMovieSceneSequence* AnimationSequence = ActiveSequencer->GetFocusedMovieSceneSequence();
 	UObject* BindingContext = GetAnimationPlaybackContext();
 	TArray<UObject*, TInlineAllocator<1>> BoundObjects = AnimationSequence->LocateBoundObjects(InBinding.GetObjectGuid(), BindingContext);
 
@@ -880,34 +990,154 @@ void FWidgetBlueprintEditor::PostRedo(bool bSuccessful)
 	OnWidgetBlueprintTransaction.Broadcast();
 }
 
-TSharedRef<SWidget> FWidgetBlueprintEditor::CreateSequencerWidget()
+TSharedRef<SWidget> FWidgetBlueprintEditor::CreateSequencerTabWidget()
 {
 	TSharedRef<SOverlay> SequencerOverlayRef =
 		SNew(SOverlay)
 		.AddMetaData<FTagMetaData>(FTagMetaData(TEXT("Sequencer")));
-	SequencerOverlay = SequencerOverlayRef;
+	TabSequencerOverlay = SequencerOverlayRef;
 
-	TSharedRef<STextBlock> NoAnimationTextBlockRef = 
-		SNew(STextBlock)
-		.TextStyle(FEditorStyle::Get(), "UMGEditor.NoAnimationFont")
-		.Text(LOCTEXT("NoAnimationSelected", "No Animation Selected"));
-	NoAnimationTextBlock = NoAnimationTextBlockRef;
+	TSharedPtr<STextBlock> NoAnimationTextBlockPtr;
+	if (!NoAnimationTextBlockTab.IsValid())
+	{
+		NoAnimationTextBlockPtr =
+			SNew(STextBlock)
+			.TextStyle(FEditorStyle::Get(), "UMGEditor.NoAnimationFont")
+			.Text(LOCTEXT("NoAnimationSelected", "No Animation Selected"));
+		NoAnimationTextBlockTab = NoAnimationTextBlockPtr;
+	}
 
 	SequencerOverlayRef->AddSlot(0)
 	[
-		GetSequencer()->GetSequencerWidget()
+		GetTabSequencer()->GetSequencerWidget()
 	];
 
 	SequencerOverlayRef->AddSlot(1)
 		.HAlign(HAlign_Center)
 		.VAlign(VAlign_Center)
 	[
-		NoAnimationTextBlockRef
+		NoAnimationTextBlockTab.Pin().ToSharedRef()
 	];
 
 	return SequencerOverlayRef;
 }
 
+TSharedRef<SWidget> FWidgetBlueprintEditor::CreateSequencerDrawerWidget()
+{
+	TSharedRef<SOverlay> SequencerOverlayRef =
+		SNew(SOverlay)
+		.AddMetaData<FTagMetaData>(FTagMetaData(TEXT("Sequencer")));
+	DrawerSequencerOverlay = SequencerOverlayRef;
+
+	TSharedPtr<STextBlock> NoAnimationTextBlockPtr;
+	if (!NoAnimationTextBlockDrawer.IsValid())
+	{
+		NoAnimationTextBlockPtr =
+			SNew(STextBlock)
+			.TextStyle(FEditorStyle::Get(), "UMGEditor.NoAnimationFont")
+			.Text(LOCTEXT("NoAnimationSelected", "No Animation Selected"));
+		NoAnimationTextBlockDrawer = NoAnimationTextBlockPtr;
+	}
+
+	SequencerOverlayRef->AddSlot(0)
+	[
+		GetDrawerSequencer()->GetSequencerWidget()
+	];
+
+	SequencerOverlayRef->AddSlot(1)
+		.HAlign(HAlign_Center)
+		.VAlign(VAlign_Center)
+	[
+		NoAnimationTextBlockDrawer.Pin().ToSharedRef()
+	];
+
+	return SequencerOverlayRef;
+}
+
+TSharedRef<SWidget> FWidgetBlueprintEditor::OnGetWidgetAnimSequencer()
+{
+	if (!AnimDrawerWidget.IsValid())
+	{
+		FAnimationTabSummoner AnimDrawerSummoner(SharedThis(this), true);
+		FWorkflowTabSpawnInfo SpawnInfo;
+		AnimDrawerWidget = AnimDrawerSummoner.CreateTabBody(SpawnInfo);
+	}
+
+	return AnimDrawerWidget.ToSharedRef();
+}
+
+void FWidgetBlueprintEditor::ToggleAnimDrawer()
+{
+	GEditor->GetEditorSubsystem<UStatusBarSubsystem>()->TryToggleDrawer(FAnimationTabSummoner::WidgetAnimSequencerDrawerID);
+}
+
+void FWidgetBlueprintEditor::NotifyWidgetAnimListChanged()
+{
+	OnWidgetAnimationsUpdated.Broadcast();
+	
+	// Check if any animations viewed are invalid, if so select null animation
+	// This can happen when a secondardary sequencer deletes our animatio
+	for (TWeakPtr<ISequencer>& SequencerPtr : Sequencers)
+	{
+		if (TSharedPtr<ISequencer> Sequencer = SequencerPtr.Pin())
+		{
+			UWidgetAnimation* WidgetAnimation = Cast<UWidgetAnimation>(Sequencer->GetFocusedMovieSceneSequence());
+			if (!GetWidgetBlueprintObj()->Animations.Contains(WidgetAnimation))
+			{
+				Sequencer->ResetToNewRootSequence(*UWidgetAnimation::GetNullAnimation());
+				Sequencer->GetSequencerWidget()->SetEnabled(false);
+				Sequencer->SetAutoChangeMode(EAutoChangeMode::None);
+			}
+		}
+	}
+}
+
+void FWidgetBlueprintEditor::OnWidgetAnimSequencerOpened(FName StatusBarWithDrawerName)
+{
+	bIsSequencerDrawerOpen = true;
+
+	if (TSharedPtr<ISequencer>& ActiveSequencer = GetSequencer())
+	{
+		UWidgetAnimation* WidgetAnimation = Cast<UWidgetAnimation>(ActiveSequencer->GetFocusedMovieSceneSequence());
+		if (WidgetAnimation)
+		{
+			ChangeViewedAnimation(*WidgetAnimation);
+		}
+	}
+
+	for (TWeakPtr<ISequencer> SequencerPtr : Sequencers)
+	{
+		if (TSharedPtr<ISequencer> Sequencer = SequencerPtr.Pin())
+		{
+			Sequencer->RefreshTree();
+		}
+	}
+}
+
+void FWidgetBlueprintEditor::OnWidgetAnimSequencerDismissed(const TSharedPtr<SWidget>& NewlyFocusedWidget)
+{
+	bIsSequencerDrawerOpen = false;
+
+	if (TSharedPtr<ISequencer>& ActiveSequencer = GetSequencer())
+	{
+		UWidgetAnimation* WidgetAnimation = Cast<UWidgetAnimation>(ActiveSequencer->GetFocusedMovieSceneSequence());
+		if (WidgetAnimation)
+		{
+			ChangeViewedAnimation(*WidgetAnimation);
+		}
+	}
+
+	for (TWeakPtr<ISequencer> SequencerPtr : Sequencers)
+	{
+		if (TSharedPtr<ISequencer> Sequencer = SequencerPtr.Pin())
+		{
+			Sequencer->RefreshTree();
+		}
+	}
+
+	FocusWindow();
+}
+
 UWidgetBlueprint* FWidgetBlueprintEditor::GetWidgetBlueprintObj() const
 {
 	return Cast<UWidgetBlueprint>(GetBlueprintObj());
@@ -965,70 +1195,120 @@ FWidgetReference FWidgetBlueprintEditor::GetReferenceFromPreview(UWidget* Previe
 
 TSharedPtr<ISequencer>& FWidgetBlueprintEditor::GetSequencer()
 {
-	if(!Sequencer.IsValid())
+	return bIsSequencerDrawerOpen ? GetDrawerSequencer() : GetTabSequencer();
+}
+
+TSharedPtr<ISequencer> FWidgetBlueprintEditor::CreateSequencerWidgetInternal()
+{
+	const float InTime = 0.f;
+	const float OutTime = 5.0f;
+
+	FSequencerViewParams ViewParams(TEXT("UMGSequencerSettings"));
 	{
-		const float InTime  = 0.f;
-		const float OutTime = 5.0f;
+		ViewParams.OnGetAddMenuContent = FOnGetAddMenuContent::CreateSP(this, &FWidgetBlueprintEditor::OnGetAnimationAddMenuContent);
+		ViewParams.OnBuildCustomContextMenuForGuid = FOnBuildCustomContextMenuForGuid::CreateSP(this, &FWidgetBlueprintEditor::OnBuildCustomContextMenuForGuid);
+	}
 
-		FSequencerViewParams ViewParams(TEXT("UMGSequencerSettings"));
-		{
-			ViewParams.OnGetAddMenuContent = FOnGetAddMenuContent::CreateSP(this, &FWidgetBlueprintEditor::OnGetAnimationAddMenuContent);
-		    ViewParams.OnBuildCustomContextMenuForGuid = FOnBuildCustomContextMenuForGuid::CreateSP(this, &FWidgetBlueprintEditor::OnBuildCustomContextMenuForGuid);
-		}
+	FSequencerInitParams SequencerInitParams;
+	{
+		UWidgetAnimation* NullAnimation = UWidgetAnimation::GetNullAnimation();
+		FFrameRate TickResolution = NullAnimation->MovieScene->GetTickResolution();
+		FFrameNumber StartFrame = (InTime * TickResolution).FloorToFrame();
+		FFrameNumber EndFrame = (OutTime * TickResolution).CeilToFrame();
+		NullAnimation->MovieScene->SetPlaybackRange(StartFrame, (EndFrame - StartFrame).Value);
+		FMovieSceneEditorData& EditorData = NullAnimation->MovieScene->GetEditorData();
+		EditorData.WorkStart = InTime;
+		EditorData.WorkEnd = OutTime;
 
-		FSequencerInitParams SequencerInitParams;
-		{
-			UWidgetAnimation* NullAnimation = UWidgetAnimation::GetNullAnimation();
-			FFrameRate TickResolution = NullAnimation->MovieScene->GetTickResolution();
-			FFrameNumber StartFrame = (InTime  * TickResolution).FloorToFrame();
-			FFrameNumber EndFrame   = (OutTime * TickResolution).CeilToFrame();
-			NullAnimation->MovieScene->SetPlaybackRange(StartFrame, (EndFrame-StartFrame).Value);
-			FMovieSceneEditorData& EditorData = NullAnimation->MovieScene->GetEditorData();
-			EditorData.WorkStart = InTime;
-			EditorData.WorkEnd   = OutTime;
-
-			SequencerInitParams.ViewParams = ViewParams;
-			SequencerInitParams.RootSequence = NullAnimation;
-			SequencerInitParams.bEditWithinLevelEditor = false;
-			SequencerInitParams.ToolkitHost = GetToolkitHost();
-			SequencerInitParams.PlaybackContext = TAttribute<UObject*>(this, &FWidgetBlueprintEditor::GetAnimationPlaybackContext);
-			SequencerInitParams.EventContexts = TAttribute<TArray<UObject*>>(this, &FWidgetBlueprintEditor::GetAnimationEventContexts);
-
-			SequencerInitParams.HostCapabilities.bSupportsCurveEditor = true;
-		};
-
-		Sequencer = FModuleManager::LoadModuleChecked<ISequencerModule>("Sequencer").CreateSequencer(SequencerInitParams);
-		// Never recompile the blueprint on evaluate as this can create an insidious loop
-		Sequencer->GetSequencerSettings()->SetCompileDirectorOnEvaluate(false);
-		Sequencer->OnMovieSceneDataChanged().AddSP( this, &FWidgetBlueprintEditor::OnMovieSceneDataChanged );
-		Sequencer->OnMovieSceneBindingsPasted().AddSP( this, &FWidgetBlueprintEditor::OnMovieSceneBindingsPasted );
-		// Change selected widgets in the sequencer tree view
-		Sequencer->GetSelectionChangedObjectGuids().AddSP(this, &FWidgetBlueprintEditor::SyncSelectedWidgetsWithSequencerSelection);
-		OnSelectedWidgetsChanged.AddSP(this, &FWidgetBlueprintEditor::SyncSequencerSelectionToSelectedWidgets);
-		
-		// Allow sequencer to test which bindings are selected
-		Sequencer->OnGetIsBindingVisible().BindRaw(this, &FWidgetBlueprintEditor::IsBindingSelected);
+		SequencerInitParams.ViewParams = ViewParams;
+		SequencerInitParams.RootSequence = NullAnimation;
+		SequencerInitParams.bEditWithinLevelEditor = false;
+		SequencerInitParams.ToolkitHost = GetToolkitHost();
+		SequencerInitParams.PlaybackContext = TAttribute<UObject*>(this, &FWidgetBlueprintEditor::GetAnimationPlaybackContext);
+		SequencerInitParams.EventContexts = TAttribute<TArray<UObject*>>(this, &FWidgetBlueprintEditor::GetAnimationEventContexts);
+
+		SequencerInitParams.HostCapabilities.bSupportsCurveEditor = true;
+	};
+
+	TSharedPtr<ISequencer> Sequencer = FModuleManager::LoadModuleChecked<ISequencerModule>("Sequencer").CreateSequencer(SequencerInitParams);
+	// Never recompile the blueprint on evaluate as this can create an insidious loop
+	Sequencer->GetSequencerSettings()->SetCompileDirectorOnEvaluate(false);
+	Sequencer->OnMovieSceneDataChanged().AddSP(this, &FWidgetBlueprintEditor::OnMovieSceneDataChanged);
+	Sequencer->OnMovieSceneBindingsPasted().AddSP(this, &FWidgetBlueprintEditor::OnMovieSceneBindingsPasted);
+	// Change selected widgets in the sequencer tree view
+	Sequencer->GetSelectionChangedObjectGuids().AddSP(this, &FWidgetBlueprintEditor::SyncSelectedWidgetsWithSequencerSelection);
+	OnSelectedWidgetsChanged.AddSP(this, &FWidgetBlueprintEditor::SyncSequencerSelectionToSelectedWidgets);
 
+	// Allow sequencer to test which bindings are selected
+	Sequencer->OnGetIsBindingVisible().BindRaw(this, &FWidgetBlueprintEditor::IsBindingSelected);
+	Sequencers.AddUnique(Sequencer);
+
+	return Sequencer;
+}
+
+TSharedPtr<ISequencer>& FWidgetBlueprintEditor::GetTabSequencer()
+{
+	if(!TabSequencer.IsValid())
+	{
+		TabSequencer = CreateSequencerWidgetInternal();
+
+		bIsSequencerDrawerOpen = false;
 		ChangeViewedAnimation(*UWidgetAnimation::GetNullAnimation());
 	}
 
-	return Sequencer;
+	return TabSequencer;
+}
+
+TSharedPtr<ISequencer>& FWidgetBlueprintEditor::GetDrawerSequencer()
+{
+	if(!DrawerSequencer.IsValid())
+	{
+		DrawerSequencer = CreateSequencerWidgetInternal();
+
+		bIsSequencerDrawerOpen = true;
+		ChangeViewedAnimation(*UWidgetAnimation::GetNullAnimation());
+	}
+
+	return DrawerSequencer;
+}
+
+void FWidgetBlueprintEditor::DockInLayoutClicked()
+{
+	GEditor->GetEditorSubsystem<UStatusBarSubsystem>()->ForceDismissDrawer();
+
+	const FName AnimationsTabName = FName(TEXT("Animations"));
+	if (TSharedPtr<SDockTab> ExistingTab = GetToolkitHost()->GetTabManager()->TryInvokeTab(AnimationsTabName))
+	{
+		ExistingTab->ActivateInParent(ETabActivationCause::SetDirectly);
+	}
 }
 
 void FWidgetBlueprintEditor::ChangeViewedAnimation( UWidgetAnimation& InAnimationToView )
 {
 	CurrentAnimation = &InAnimationToView;
 
-	if (Sequencer.IsValid())
+	if (TSharedPtr<ISequencer>& Sequencer = GetSequencer())
 	{
 		Sequencer->ResetToNewRootSequence(InAnimationToView);
+		if (&InAnimationToView == UWidgetAnimation::GetNullAnimation())
+		{
+			Sequencer->GetSequencerWidget()->SetEnabled(false);
+			Sequencer->SetAutoChangeMode(EAutoChangeMode::None);
+		}
+		else
+		{
+			Sequencer->GetSequencerWidget()->SetEnabled(true);
+		}
 	}
 
-	TSharedPtr<SOverlay> SequencerOverlayPin = SequencerOverlay.Pin();
-	if (SequencerOverlayPin.IsValid())
+	TWeakPtr<SOverlay> SequencerOverlay = bIsSequencerDrawerOpen ? DrawerSequencerOverlay : TabSequencerOverlay;
+	TWeakPtr<STextBlock> NoAnimationTextBlock = bIsSequencerDrawerOpen ? NoAnimationTextBlockDrawer : NoAnimationTextBlockTab;
+	if (SequencerOverlay.IsValid() && NoAnimationTextBlock.IsValid())
 	{
+		TSharedPtr<SOverlay> SequencerOverlayPin = SequencerOverlay.Pin();
 		TSharedPtr<STextBlock> NoAnimationTextBlockPin = NoAnimationTextBlock.Pin();
-		if( &InAnimationToView == UWidgetAnimation::GetNullAnimation())
+
+		if (&InAnimationToView == UWidgetAnimation::GetNullAnimation())
 		{
 			const FName CurveEditorTabName = FName(TEXT("SequencerGraphEditor"));
 			TSharedPtr<SDockTab> ExistingTab = GetToolkitHost()->GetTabManager()->FindExistingLiveTab(CurveEditorTabName);
@@ -1038,19 +1318,17 @@ void FWidgetBlueprintEditor::ChangeViewedAnimation( UWidgetAnimation& InAnimatio
 			}
 
 			// Disable sequencer from interaction
-			Sequencer->GetSequencerWidget()->SetEnabled(false);
-			Sequencer->SetAutoChangeMode(EAutoChangeMode::None);
 			NoAnimationTextBlockPin->SetVisibility(EVisibility::Visible);
-			SequencerOverlayPin->SetVisibility( EVisibility::HitTestInvisible );
+			SequencerOverlayPin->SetVisibility(EVisibility::HitTestInvisible);
 		}
 		else
 		{
 			// Allow sequencer to be interacted with
-			Sequencer->GetSequencerWidget()->SetEnabled(true);
 			NoAnimationTextBlockPin->SetVisibility(EVisibility::Collapsed);
-			SequencerOverlayPin->SetVisibility( EVisibility::SelfHitTestInvisible );
+			SequencerOverlayPin->SetVisibility(EVisibility::SelfHitTestInvisible);
 		}
 	}
+
 	InvalidatePreview();
 }
 
@@ -1072,6 +1350,59 @@ void FWidgetBlueprintEditor::Compile()
 	DestroyPreview();
 
 	FBlueprintEditor::Compile();
+
+	if (const UWidgetDesignerSettings* Settings = GetDefault<UWidgetDesignerSettings>())
+	{
+		// Check if we should create the compile tab
+		bool bShouldCreateCompileTab = false;
+		switch (Settings->CreateOnCompile)
+		{
+		case EDisplayOnCompile::DoC_ErrorsOrWarnings:
+			bShouldCreateCompileTab = CachedNumErrors > 0 || CachedNumWarnings > 0;
+			break;
+		case EDisplayOnCompile::DoC_ErrorsOnly:
+			bShouldCreateCompileTab = CachedNumErrors > 0;
+			break;
+		case EDisplayOnCompile::DoC_WarningsOnly:
+			bShouldCreateCompileTab = CachedNumWarnings > 0;
+			break;
+		case EDisplayOnCompile::DoC_Never:
+		default:
+			break;
+		}
+
+		if (bShouldCreateCompileTab)
+		{
+			GetToolkitHost()->GetTabManager()->TryInvokeTab(FBlueprintEditorTabs::CompilerResultsID);
+		}
+
+		// Check if we should dismiss the compile tab
+		bool bShouldDismissCompileTab = false;
+		switch (Settings->DismissOnCompile)
+		{
+		case EDisplayOnCompile::DoC_ErrorsOrWarnings:
+			bShouldDismissCompileTab = CachedNumErrors == 0 && CachedNumWarnings == 0;
+			break;
+		case EDisplayOnCompile::DoC_ErrorsOnly:
+			bShouldDismissCompileTab = CachedNumErrors == 0;
+			break;
+		case EDisplayOnCompile::DoC_WarningsOnly:
+			bShouldDismissCompileTab = CachedNumWarnings == 0;
+			break;
+		case EDisplayOnCompile::DoC_Never:
+		default:
+			break;
+		}
+
+		if (bShouldDismissCompileTab)
+		{
+			TSharedPtr<SDockTab> CompileResultsTab = GetToolkitHost()->GetTabManager()->FindExistingLiveTab(FBlueprintEditorTabs::CompilerResultsID);
+			if (CompileResultsTab)
+			{
+				CompileResultsTab->RequestCloseTab();
+			}
+		}
+	}
 }
 
 void FWidgetBlueprintEditor::DestroyPreview()
@@ -1170,11 +1501,13 @@ void FWidgetBlueprintEditor::UpdatePreview(UBlueprint* InBlueprint, bool bInForc
 	OnWidgetPreviewUpdated.Broadcast();
 
 	// We've changed the binding context so drastically that we should just clear all knowledge of our previous cached bindings
-
-	if (Sequencer.IsValid())
+	for (TWeakPtr<ISequencer>& SequencerPtr : Sequencers)
 	{
-		Sequencer->State.ClearObjectCaches(*Sequencer);
-		Sequencer->ForceEvaluate();
+		if (TSharedPtr<ISequencer> Sequencer = SequencerPtr.Pin())
+		{
+			Sequencer->State.ClearObjectCaches(*Sequencer);
+			Sequencer->ForceEvaluate();
+		}
 	}
 }
 
@@ -1351,9 +1684,10 @@ void FWidgetBlueprintEditor::OnGetAnimationAddMenuContentAllWidgets(FMenuBuilder
 		BindableObjects.Sort();
 	}
 
+	TSharedPtr<ISequencer>& ActiveSequencer = GetSequencer();
 	for (FObjectAndDisplayName& BindableObject : BindableObjects)
 	{
-		FGuid BoundObjectGuid = Sequencer->FindObjectId(*BindableObject.Object, Sequencer->GetFocusedTemplateID());
+		FGuid BoundObjectGuid = ActiveSequencer->FindObjectId(*BindableObject.Object, ActiveSequencer->GetFocusedTemplateID());
 		if (BoundObjectGuid.IsValid() == false)
 		{
 			FUIAction AddMenuAction(FExecuteAction::CreateSP(this, &FWidgetBlueprintEditor::AddObjectToAnimation, BindableObject.Object));
@@ -1364,20 +1698,21 @@ void FWidgetBlueprintEditor::OnGetAnimationAddMenuContentAllWidgets(FMenuBuilder
 
 void FWidgetBlueprintEditor::AddObjectToAnimation(UObject* ObjectToAnimate)
 {
-	UMovieScene* MovieScene = Sequencer->GetFocusedMovieSceneSequence()->GetMovieScene();
+	TSharedPtr<ISequencer>& ActiveSequencer = GetSequencer();
+	UMovieScene* MovieScene = ActiveSequencer->GetFocusedMovieSceneSequence()->GetMovieScene();
 	if (MovieScene->IsReadOnly())
 	{
 		return;
 	}
 
 	const FScopedTransaction Transaction( LOCTEXT( "AddWidgetToAnimation", "Add widget to animation" ) );
-	Sequencer->GetFocusedMovieSceneSequence()->Modify();
+	ActiveSequencer->GetFocusedMovieSceneSequence()->Modify();
 
 	// @todo Sequencer - Make this kind of adding more explicit, this current setup seem a bit brittle.
-	FGuid NewGuid = Sequencer->GetHandleToObject(ObjectToAnimate);
+	FGuid NewGuid = ActiveSequencer->GetHandleToObject(ObjectToAnimate);
 
 	TArray<UMovieSceneFolder*> SelectedParentFolders;
-	Sequencer->GetSelectedFolders(SelectedParentFolders);
+	ActiveSequencer->GetSelectedFolders(SelectedParentFolders);
 
 	if (SelectedParentFolders.Num() > 0)
 	{
@@ -1408,7 +1743,8 @@ void FWidgetBlueprintEditor::OnBuildCustomContextMenuForGuid(FMenuBuilder& MenuB
 				//need to make sure it's a widget, if not bound assume it is.
 				UWidget* BoundWidget = nullptr;
 				bool bNotBound = true;
-				for (TWeakObjectPtr<> WeakObjectPtr : GetSequencer()->FindObjectsInCurrentSequence(ObjectBinding))
+				TSharedPtr<ISequencer>& ActiveSequencer = GetSequencer();
+				for (TWeakObjectPtr<> WeakObjectPtr : ActiveSequencer->FindObjectsInCurrentSequence(ObjectBinding))
 				{
 					BoundWidget = Cast<UWidget>(WeakObjectPtr.Get());
 					bNotBound = false;
@@ -1530,7 +1866,8 @@ void FWidgetBlueprintEditor::AddWidgetsToTrack(const TArray<FWidgetReference> Wi
 {
 	const FScopedTransaction Transaction(LOCTEXT("AddSelectedWidgetsToTrack", "Add Widgets to Track"));
 
-	UWidgetAnimation* WidgetAnimation = Cast<UWidgetAnimation>(Sequencer->GetFocusedMovieSceneSequence());
+	TSharedPtr<ISequencer>& ActiveSequencer = GetSequencer();
+	UWidgetAnimation* WidgetAnimation = Cast<UWidgetAnimation>(ActiveSequencer->GetFocusedMovieSceneSequence());
 	UMovieScene* MovieScene = WidgetAnimation->GetMovieScene();
 
 	FText ExistingBindingName;
@@ -1540,7 +1877,7 @@ void FWidgetBlueprintEditor::AddWidgetsToTrack(const TArray<FWidgetReference> Wi
 		UWidget* PreviewWidget = Widget.GetPreview();
 
 		// If this widget is already bound to the animation we cannot add it to 2 separate bindings
-		FGuid SelectedWidgetId = Sequencer->FindObjectId(*PreviewWidget, MovieSceneSequenceID::Root);
+		FGuid SelectedWidgetId = ActiveSequencer->FindObjectId(*PreviewWidget, MovieSceneSequenceID::Root);
 		if (!SelectedWidgetId.IsValid())
 		{
 			WidgetsToAdd.Add(Widget);
@@ -1574,8 +1911,7 @@ void FWidgetBlueprintEditor::AddWidgetsToTrack(const TArray<FWidgetReference> Wi
 		}
 
 		UpdateTrackName(ObjectId);
-
-		Sequencer->NotifyMovieSceneDataChanged(EMovieSceneDataChangeType::MovieSceneStructureItemsChanged);
+		SyncSequencersMovieSceneData();
 	}
 }
 
@@ -1583,7 +1919,8 @@ void FWidgetBlueprintEditor::RemoveWidgetsFromTrack(const TArray<FWidgetReferenc
 {
 	const FScopedTransaction Transaction(LOCTEXT("RemoveWidgetsFromTrack", "Remove Widgets from Track"));
 
-	UWidgetAnimation* WidgetAnimation = Cast<UWidgetAnimation>(Sequencer->GetFocusedMovieSceneSequence());
+	TSharedPtr<ISequencer>& ActiveSequencer = GetSequencer();
+	UWidgetAnimation* WidgetAnimation = Cast<UWidgetAnimation>(ActiveSequencer->GetFocusedMovieSceneSequence());
 	UMovieScene* MovieScene = WidgetAnimation->GetMovieScene();
 
 	TArray<FWidgetReference> WidgetsToRemove;
@@ -1591,7 +1928,7 @@ void FWidgetBlueprintEditor::RemoveWidgetsFromTrack(const TArray<FWidgetReferenc
 	for (const FWidgetReference& Widget : Widgets)
 	{
 		UWidget* PreviewWidget = Widget.GetPreview();
-		FGuid WidgetId = Sequencer->FindObjectId(*PreviewWidget, MovieSceneSequenceID::Root);
+		FGuid WidgetId = ActiveSequencer->FindObjectId(*PreviewWidget, MovieSceneSequenceID::Root);
 		if (WidgetId.IsValid() && WidgetId == ObjectId)
 		{
 			WidgetsToRemove.Add(Widget);
@@ -1619,12 +1956,11 @@ void FWidgetBlueprintEditor::RemoveWidgetsFromTrack(const TArray<FWidgetReferenc
 			UWidget* PreviewWidget = Widget.GetPreview();
 			WidgetAnimation->RemoveBinding(*PreviewWidget);
 
-			Sequencer->PreAnimatedState.RestorePreAnimatedState(*PreviewWidget);
+			ActiveSequencer->PreAnimatedState.RestorePreAnimatedState(*PreviewWidget);
 		}
 
 		UpdateTrackName(ObjectId);
-
-		Sequencer->NotifyMovieSceneDataChanged(EMovieSceneDataChangeType::MovieSceneStructureItemsChanged);
+		SyncSequencersMovieSceneData();
 	}
 }
 
@@ -1632,7 +1968,8 @@ void FWidgetBlueprintEditor::RemoveAllWidgetsFromTrack(FGuid ObjectId)
 {
 	const FScopedTransaction Transaction(LOCTEXT("RemoveAllWidgetsFromTrack", "Remove All Widgets from Track"));
 
-	UWidgetAnimation* WidgetAnimation = Cast<UWidgetAnimation>(Sequencer->GetFocusedMovieSceneSequence());
+	TSharedPtr<ISequencer>& ActiveSequencer = GetSequencer();
+	UWidgetAnimation* WidgetAnimation = Cast<UWidgetAnimation>(ActiveSequencer->GetFocusedMovieSceneSequence());
 	UMovieScene* MovieScene = WidgetAnimation->GetMovieScene();
 
 	UUserWidget* PreviewRoot = GetPreview();
@@ -1642,11 +1979,11 @@ void FWidgetBlueprintEditor::RemoveAllWidgetsFromTrack(FGuid ObjectId)
 	MovieScene->Modify();
 
 	// Restore object animation state
-	for (TWeakObjectPtr<> WeakObject : Sequencer->FindBoundObjects(ObjectId, MovieSceneSequenceID::Root))
+	for (TWeakObjectPtr<> WeakObject : ActiveSequencer->FindBoundObjects(ObjectId, MovieSceneSequenceID::Root))
 	{
 		if (UObject* Obj = WeakObject.Get())
 		{
-			Sequencer->PreAnimatedState.RestorePreAnimatedState(*Obj);
+			ActiveSequencer->PreAnimatedState.RestorePreAnimatedState(*Obj);
 		}
 	}
 
@@ -1659,14 +1996,15 @@ void FWidgetBlueprintEditor::RemoveAllWidgetsFromTrack(FGuid ObjectId)
 		}
 	}
 
-	Sequencer->NotifyMovieSceneDataChanged(EMovieSceneDataChangeType::MovieSceneStructureItemsChanged);
+	SyncSequencersMovieSceneData();
 }
 
 void FWidgetBlueprintEditor::RemoveMissingWidgetsFromTrack(FGuid ObjectId)
 {
 	const FScopedTransaction Transaction(LOCTEXT("RemoveMissingWidgetsFromTrack", "Remove Missing Widgets from Track"));
 
-	UWidgetAnimation* WidgetAnimation = Cast<UWidgetAnimation>(Sequencer->GetFocusedMovieSceneSequence());
+	TSharedPtr<ISequencer>& ActiveSequencer = GetSequencer();
+	UWidgetAnimation* WidgetAnimation = Cast<UWidgetAnimation>(ActiveSequencer->GetFocusedMovieSceneSequence());
 	UMovieScene* MovieScene = WidgetAnimation->GetMovieScene();
 
 	UUserWidget* PreviewRoot = GetPreview();
@@ -1689,7 +2027,8 @@ void FWidgetBlueprintEditor::RemoveMissingWidgetsFromTrack(FGuid ObjectId)
 
 void FWidgetBlueprintEditor::ReplaceTrackWithWidgets(TArray<FWidgetReference> Widgets, FGuid ObjectId)
 {
-	UWidgetAnimation* WidgetAnimation = Cast<UWidgetAnimation>(Sequencer->GetFocusedMovieSceneSequence());
+	TSharedPtr<ISequencer>& ActiveSequencer = GetSequencer();
+	UWidgetAnimation* WidgetAnimation = Cast<UWidgetAnimation>(ActiveSequencer->GetFocusedMovieSceneSequence());
 	UMovieScene* MovieScene = WidgetAnimation->GetMovieScene();
 
 	// Filter out anything in the input array that is currently bound to another object in the animation
@@ -1697,7 +2036,7 @@ void FWidgetBlueprintEditor::ReplaceTrackWithWidgets(TArray<FWidgetReference> Wi
 	for (int32 Index = Widgets.Num()-1; Index >= 0; --Index)
 	{
 		UWidget* PreviewWidget = Widgets[Index].GetPreview();
-		FGuid WidgetId = Sequencer->FindObjectId(*PreviewWidget, MovieSceneSequenceID::Root);
+		FGuid WidgetId = ActiveSequencer->FindObjectId(*PreviewWidget, MovieSceneSequenceID::Root);
 		if (WidgetId.IsValid() && WidgetId != ObjectId)
 		{
 			Widgets.RemoveAt(Index, 1, false);
@@ -1734,21 +2073,22 @@ void FWidgetBlueprintEditor::ReplaceTrackWithWidgets(TArray<FWidgetReference> Wi
 	AddWidgetsToTrack(Widgets, ObjectId);
 
 	UpdateTrackName(ObjectId);
-
-	Sequencer->NotifyMovieSceneDataChanged(EMovieSceneDataChangeType::MovieSceneStructureItemsChanged);
+	SyncSequencersMovieSceneData();
 }
 
 void FWidgetBlueprintEditor::AddSlotTrack( UPanelSlot* Slot )
 {
-	GetSequencer()->GetHandleToObject( Slot );
+	TSharedPtr<ISequencer>& ActiveSequencer = GetSequencer();
+	ActiveSequencer->GetHandleToObject( Slot );
 }
 
 void FWidgetBlueprintEditor::AddMaterialTrack( UWidget* Widget, TArray<FProperty*> MaterialPropertyPath, FText MaterialPropertyDisplayName )
 {
-	FGuid WidgetHandle = Sequencer->GetHandleToObject( Widget );
+	TSharedPtr<ISequencer>& ActiveSequencer = GetSequencer();
+	FGuid WidgetHandle = ActiveSequencer->GetHandleToObject( Widget );
 	if ( WidgetHandle.IsValid() )
 	{
-		UMovieScene* MovieScene = Sequencer->GetFocusedMovieSceneSequence()->GetMovieScene();
+		UMovieScene* MovieScene = ActiveSequencer->GetFocusedMovieSceneSequence()->GetMovieScene();
 
 		if (MovieScene->IsReadOnly())
 		{
@@ -1771,7 +2111,7 @@ void FWidgetBlueprintEditor::AddMaterialTrack( UWidget* Widget, TArray<FProperty
 			NewTrack->SetBrushPropertyNamePath( MaterialPropertyNamePath );
 			NewTrack->SetDisplayName( FText::Format( LOCTEXT( "TrackDisplayNameFormat", "{0}"), MaterialPropertyDisplayName ) );
 
-			Sequencer->NotifyMovieSceneDataChanged(EMovieSceneDataChangeType::MovieSceneStructureItemAdded );
+			SyncSequencersMovieSceneData();
 		}
 	}
 }
@@ -1788,7 +2128,8 @@ void FWidgetBlueprintEditor::OnMovieSceneBindingsPasted(const TArray<FMovieScene
 		GetBindableObjects(GetPreview()->WidgetTree, BindableObjects);
 	}
 
-	UMovieSceneSequence* AnimationSequence = GetSequencer().Get()->GetFocusedMovieSceneSequence();
+	TSharedPtr<ISequencer>& ActiveSequencer = GetSequencer();
+	UMovieSceneSequence* AnimationSequence = ActiveSequencer->GetFocusedMovieSceneSequence();
 	UObject* BindingContext = GetAnimationPlaybackContext();
 
 	// First, rebind top level possessables (without parents) - match binding pasted's name with the bindable object name
@@ -1831,7 +2172,7 @@ void FWidgetBlueprintEditor::OnMovieSceneBindingsPasted(const TArray<FMovieScene
 					if (BindableObject.Object->GetFName().ToString() == BindingPasted.GetName())
 					{
 						// Create handle, to rebind correctly
-						Sequencer->GetHandleToObject(BindableObject.Object);
+						ActiveSequencer->GetHandleToObject(BindableObject.Object);
 						// Remove the existing binding, as it is now replaced by the that was just added by getting the handle
 						AnimationSequence->GetMovieScene()->RemovePossessable(BindingPasted.GetObjectGuid());
 						break;
@@ -1851,7 +2192,8 @@ void FWidgetBlueprintEditor::SyncSelectedWidgetsWithSequencerSelection(TArray<FG
 
 	TGuardValue<bool> Guard(bUpdatingExternalSelection, true);
 
-	UMovieSceneSequence* AnimationSequence = GetSequencer().Get()->GetFocusedMovieSceneSequence();
+	TSharedPtr<ISequencer>& ActiveSequencer = GetSequencer();
+	UMovieSceneSequence* AnimationSequence = ActiveSequencer->GetFocusedMovieSceneSequence();
 	UObject* BindingContext = GetAnimationPlaybackContext();
 	TSet<FWidgetReference> SequencerSelectedWidgets;
 	for (FGuid Guid : ObjectGuids)
@@ -1886,12 +2228,29 @@ void FWidgetBlueprintEditor::SyncSequencerSelectionToSelectedWidgets()
 
 	TGuardValue<bool> Guard(bUpdatingSequencerSelection, true);
 
-	if (GetSequencer()->GetSequencerSettings()->GetShowSelectedNodesOnly())
+	for (TWeakPtr<ISequencer> SequencerPtr : Sequencers)
 	{
-		GetSequencer()->RefreshTree();
+		if (TSharedPtr<ISequencer> Sequencer = SequencerPtr.Pin())
+		{
+			if (Sequencer->GetSequencerSettings()->GetShowSelectedNodesOnly())
+			{
+				Sequencer->RefreshTree();
+			}
+
+			Sequencer->ExternalSelectionHasChanged();
+		}
 	}
+}
 
-	GetSequencer()->ExternalSelectionHasChanged();
+void FWidgetBlueprintEditor::SyncSequencersMovieSceneData()
+{
+	for (TWeakPtr<ISequencer> SequencerPtr : Sequencers)
+	{
+		if (TSharedPtr<ISequencer> Sequencer = SequencerPtr.Pin())
+		{
+			Sequencer->NotifyMovieSceneDataChanged(EMovieSceneDataChangeType::MovieSceneStructureItemsChanged);
+		}
+	}
 }
 
 void FWidgetBlueprintEditor::UpdateTrackName(FGuid ObjectId)
@@ -1899,7 +2258,8 @@ void FWidgetBlueprintEditor::UpdateTrackName(FGuid ObjectId)
 	UUserWidget* PreviewRoot = GetPreview();
 	UObject* BindingContext = GetAnimationPlaybackContext();
 
-	UWidgetAnimation* WidgetAnimation = Cast<UWidgetAnimation>(Sequencer->GetFocusedMovieSceneSequence());
+	TSharedPtr<ISequencer>& ActiveSequencer = GetSequencer();
+	UWidgetAnimation* WidgetAnimation = Cast<UWidgetAnimation>(ActiveSequencer->GetFocusedMovieSceneSequence());
 	UMovieScene* MovieScene = WidgetAnimation->GetMovieScene();
 
 	const TArray<FWidgetAnimationBinding>& WidgetBindings = WidgetAnimation->GetBindings();
diff --git a/Engine/Source/Editor/UMGEditor/Public/Settings/WidgetDesignerSettings.h b/Engine/Source/Editor/UMGEditor/Public/Settings/WidgetDesignerSettings.h
--- a/Engine/Source/Editor/UMGEditor/Public/Settings/WidgetDesignerSettings.h
+++ b/Engine/Source/Editor/UMGEditor/Public/Settings/WidgetDesignerSettings.h
@@ -9,6 +9,15 @@
 #include "WidgetPaletteFavorites.h"
 #include "WidgetDesignerSettings.generated.h"
 
+UENUM()
+enum EDisplayOnCompile
+{
+	DoC_ErrorsOrWarnings,
+	DoC_ErrorsOnly,
+	DoC_WarningsOnly,
+	DoC_Never,
+};
+
 /**
  * Implements the settings for the Widget Blueprint Designer.
  */
@@ -62,6 +71,14 @@ public:
 	 */
 	UPROPERTY(EditAnywhere, config, Category = Interaction)
 	bool bRespectLocks;
+
+	/** Setting to automatically create compile tab based on compile results */
+	UPROPERTY(EditAnywhere, config, Category = Interaction)
+	TEnumAsByte<EDisplayOnCompile> CreateOnCompile;
+
+	/** Setting to automatically dismiss compile tab based on compile results */
+	UPROPERTY(EditAnywhere, config, Category = Interaction)
+	TEnumAsByte<EDisplayOnCompile> DismissOnCompile;
 	
 	/**
 	 * List of Favorites widgets used to populate the Favorites Palette
diff --git a/Engine/Source/Editor/UMGEditor/Public/WidgetBlueprintEditor.h b/Engine/Source/Editor/UMGEditor/Public/WidgetBlueprintEditor.h
--- a/Engine/Source/Editor/UMGEditor/Public/WidgetBlueprintEditor.h
+++ b/Engine/Source/Editor/UMGEditor/Public/WidgetBlueprintEditor.h
@@ -7,6 +7,7 @@
 #include "Widgets/SWidget.h"
 #include "Framework/Commands/UICommandList.h"
 #include "Framework/MultiBox/MultiBoxExtender.h"
+#include "Settings/WidgetDesignerSettings.h"
 #include "AssetData.h"
 #include "PreviewScene.h"
 #include "GraphEditor.h"
@@ -46,6 +47,9 @@ public:
 	/** Called after the widget preview has been updated */
 	DECLARE_MULTICAST_DELEGATE(FOnWidgetPreviewUpdated)
 
+	/** Called when animation list changes */
+	DECLARE_MULTICAST_DELEGATE(FOnWidgetAnimationsUpdated)
+
 	DECLARE_EVENT(FWidgetBlueprintEditor, FOnEnterWidgetDesigner)
 
 public:
@@ -91,9 +95,18 @@ public:
 	/** Creates a widget reference using the preview.  Which is used to lookup the stable template pointer. */
 	FWidgetReference GetReferenceFromPreview(UWidget* PreviewWidget);
 
-	/** @return The sequencer used to create widget animations */
+	/** @return The currently active sequencer used to create widget animations */
 	TSharedPtr<ISequencer>& GetSequencer();
 
+	/** @return The sequencer used to create widget animations */
+	TSharedPtr<ISequencer>& GetTabSequencer();
+
+	/** @return The sequencer used to create widget animations in drawer */
+	TSharedPtr<ISequencer>& GetDrawerSequencer();
+
+	/** Callback to dock anim tab into layout */
+	void DockInLayoutClicked();
+
 	/** Changes the currently viewed animation in Sequencer to the new one*/
 	void ChangeViewedAnimation( UWidgetAnimation& InAnimationToView );
 
@@ -149,7 +162,25 @@ public:
 	FOnWidgetBlueprintTransaction& GetOnWidgetBlueprintTransaction() { return OnWidgetBlueprintTransaction; }
 
 	/** Creates a sequencer widget */
-	TSharedRef<SWidget> CreateSequencerWidget();
+	TSharedRef<SWidget> CreateSequencerTabWidget();
+
+	/** Creates a sequencer widget, for anim drawer */
+	TSharedRef<SWidget> CreateSequencerDrawerWidget();
+
+	/** Gets sequencer widget, for anim drawer */
+	TSharedRef<SWidget> OnGetWidgetAnimSequencer();
+
+	/** Toggles Anim Drawer, only used for keyboard shortcut */
+	void ToggleAnimDrawer();
+
+	/** Callback for when widget animation list has been updated */
+	void NotifyWidgetAnimListChanged();
+
+	/** Callback for anim drawer opening */
+	void OnWidgetAnimSequencerOpened(FName StatusBarWithDrawerName);
+
+	/** Callback for anim drawer closing */
+	void OnWidgetAnimSequencerDismissed(const TSharedPtr<SWidget>& NewlyFocusedWidget);
 
 	/**
 	 * The widget we're now hovering over in any particular context, allows multiple views to 
@@ -196,6 +227,9 @@ public:
 	/** Notification for when the preview widget has been updated */
 	FOnWidgetPreviewUpdated OnWidgetPreviewUpdated;
 
+	/** Notification for when the preview widget has been updated */
+	FOnWidgetAnimationsUpdated OnWidgetAnimationsUpdated;
+
 	/** Fires after the mode change to Designer*/
 	FOnEnterWidgetDesigner OnEnterWidgetDesigner;
 
@@ -218,8 +252,18 @@ protected:
 	bool IsImageUsedAsThumbnail();
 	bool IsPreviewWidgetInitialized();
 
+	void CustomizeWidgetCompileOptions(FMenuBuilder& InMenuBuilder);
+
+	static void AddCreateCompileTabSubMenu(FMenuBuilder& InMenuBuilder);
+	static void AddDismissCompileTabSubMenu(FMenuBuilder& InMenuBuilder);
+	static void SetDismissOnCompileSetting(EDisplayOnCompile InDismissOnCompile);
+	static void SetCreateOnCompileSetting(EDisplayOnCompile InCreateOnCompile);
+	static bool IsDismissOnCompileSet(EDisplayOnCompile InDismissOnCompile);
+	static bool IsCreateOnCompileSet(EDisplayOnCompile InCreateOnCompile);
+
 	void OpenCreateNativeBaseClassDialog();
 	void OnCreateNativeBaseClassSuccessfully(const FString& InClassName, const FString& InClassPath, const FString& InModuleName);
+	TSharedPtr<ISequencer> CreateSequencerWidgetInternal();
 
 	// Begin FBlueprintEditor
 	virtual void RegisterApplicationModes(const TArray<UBlueprint*>& InBlueprints, bool bShouldOpenInDefaultsMode, bool bNewlyCreated = false) override;
@@ -311,6 +355,9 @@ private:
 	/** Tell sequencer the selected widgets changed */
 	void SyncSequencerSelectionToSelectedWidgets();
 
+	/** Tell sequencers movie data changed */
+	void SyncSequencersMovieSceneData();
+
 	/** Get the animation playback context */
 	UObject* GetAnimationPlaybackContext() const { return GetPreview(); }
 
@@ -324,14 +371,29 @@ private:
 	/** The preview scene that owns the preview GUI */
 	FPreviewScene PreviewScene;
 
-	/** Sequencer for creating and previewing widget animations */
-	TSharedPtr<ISequencer> Sequencer;
+	/** List of created sequencers */
+	TArray<TWeakPtr<ISequencer>> Sequencers;
+
+	/** Sequencer for creating and previewing widget animations in tabs */
+	TSharedPtr<ISequencer> TabSequencer;
 
-	/** Overlay used to display UI on top of sequencer */
-	TWeakPtr<SOverlay> SequencerOverlay;
+	/** Overlay used to display UI on top of tab sequencer */
+	TWeakPtr<SOverlay> TabSequencerOverlay;
+
+	/** Sequencer for creating and previewing widget animations in drawer */
+	TSharedPtr<ISequencer> DrawerSequencer;
+
+	/** Overlay used to display UI on top of drawer sequencer */
+	TWeakPtr<SOverlay> DrawerSequencerOverlay;
+
+	/** Widget created for anim drawer */
+	TSharedPtr<SWidget> AnimDrawerWidget;
 
 	/** A text block which is displayed in the overlay when no animation is selected. */
-	TWeakPtr<STextBlock> NoAnimationTextBlock;
+	TWeakPtr<STextBlock> NoAnimationTextBlockTab;
+
+	/** A text block which is displayed in the overlay when no animation is selected in drawer. */
+	TWeakPtr<STextBlock> NoAnimationTextBlockDrawer;
 
 	/** The Blueprint associated with the current preview */
 	UWidgetBlueprint* PreviewBlueprint;
@@ -388,6 +450,9 @@ private:
 
 	FDelegateHandle SequencerAddTrackExtenderHandle;
 
+	/** True if sequencer drawer is open */
+	bool bIsSequencerDrawerOpen;
+
 	/** When true the animation data in the generated class should be replaced with the current animation data. */
 	bool bRefreshGeneratedClassAnimations;
 
diff --git a/Engine/Source/Editor/UMGEditor/UMGEditor.Build.cs b/Engine/Source/Editor/UMGEditor/UMGEditor.Build.cs
--- a/Engine/Source/Editor/UMGEditor/UMGEditor.Build.cs
+++ b/Engine/Source/Editor/UMGEditor/UMGEditor.Build.cs
@@ -61,6 +61,7 @@ public class UMGEditor : ModuleRules
 				"Slate",
 				"SlateCore",
 				"SlateRHIRenderer",
+				"StatusBar",
 				"MessageLog",
 				"MovieScene",
 				"MovieSceneTools",
diff --git a/Engine/Source/Editor/UnrealEd/Private/Toolkits/AssetEditorToolkit.cpp b/Engine/Source/Editor/UnrealEd/Private/Toolkits/AssetEditorToolkit.cpp
--- a/Engine/Source/Editor/UnrealEd/Private/Toolkits/AssetEditorToolkit.cpp
+++ b/Engine/Source/Editor/UnrealEd/Private/Toolkits/AssetEditorToolkit.cpp
@@ -1132,10 +1132,19 @@ void FAssetEditorToolkit::RegenerateMenusAndToolbars()
 }
 
 
+void FAssetEditorToolkit::RegisterDrawer(FStatusBarDrawer&& Drawer, int32 SlotIndex)
+{
+	TSharedPtr< class SStandaloneAssetEditorToolkitHost > HostWidget = StandaloneHost.Pin();
+	if (HostWidget.IsValid())
+	{
+		HostWidget->RegisterDrawer(MoveTemp(Drawer), SlotIndex);
+	}
+}
+
 void FAssetEditorToolkit::RestoreFromLayout(const TSharedRef<FTabManager::FLayout>& NewLayout)
 {
 	TSharedPtr< class SStandaloneAssetEditorToolkitHost > HostWidget = StandaloneHost.Pin();
-	if (HostWidget.Get() != NULL)
+	if (HostWidget.IsValid())
 	{
 		// Save the old layout
 		FLayoutSaveRestore::SaveToConfig(GEditorLayoutIni, TabManager->PersistLayout());
diff --git a/Engine/Source/Editor/UnrealEd/Private/Toolkits/SStandaloneAssetEditorToolkitHost.cpp b/Engine/Source/Editor/UnrealEd/Private/Toolkits/SStandaloneAssetEditorToolkitHost.cpp
--- a/Engine/Source/Editor/UnrealEd/Private/Toolkits/SStandaloneAssetEditorToolkitHost.cpp
+++ b/Engine/Source/Editor/UnrealEd/Private/Toolkits/SStandaloneAssetEditorToolkitHost.cpp
@@ -165,6 +165,7 @@ void SStandaloneAssetEditorToolkitHost::RestoreFromLayout( const TSharedRef<FTab
 
 	TSharedPtr<SWindow> ParentWindow = FSlateApplication::Get().FindWidgetWindow( HostTab );
 	TSharedPtr<SWidget> RestoredUI = MyTabManager->RestoreFrom( NewLayout, ParentWindow );
+	StatusBarWidget = GEditor->GetEditorSubsystem<UStatusBarSubsystem>()->MakeStatusBarWidget(GetMenuName(), HostTab);
 
 	checkf(RestoredUI.IsValid(), TEXT("The layout must have a primary dock area") );
 
@@ -181,11 +182,11 @@ void SStandaloneAssetEditorToolkitHost::RestoreFromLayout( const TSharedRef<FTab
 		[
 			RestoredUI.ToSharedRef()
 		]
-		+ SVerticalBox::Slot()
+	+ SVerticalBox::Slot()
 		.Padding(0.0f, 2.0f, 0.0f, 0.0f)
 		.AutoHeight()
 		[
-			GEditor->GetEditorSubsystem<UStatusBarSubsystem>()->MakeStatusBarWidget(GetMenuName(), HostTab)
+			StatusBarWidget.ToSharedRef()
 		]
 	];
 }
@@ -243,6 +244,14 @@ void SStandaloneAssetEditorToolkitHost::SetToolbar(TSharedPtr<SWidget> Toolbar)
 	}
 }
 
+void SStandaloneAssetEditorToolkitHost::RegisterDrawer(FStatusBarDrawer&& Drawer, int32 SlotIndex)
+{
+	if (StatusBarWidget.IsValid())
+	{
+		GEditor->GetEditorSubsystem<UStatusBarSubsystem>()->RegisterDrawer(GetMenuName(), MoveTemp(Drawer), SlotIndex);
+	}
+}
+
 FEditorModeTools& SStandaloneAssetEditorToolkitHost::GetEditorModeManager() const
 {
 	check(HostedAssetEditorToolkit.IsValid());
diff --git a/Engine/Source/Editor/UnrealEd/Private/Toolkits/SStandaloneAssetEditorToolkitHost.h b/Engine/Source/Editor/UnrealEd/Private/Toolkits/SStandaloneAssetEditorToolkitHost.h
--- a/Engine/Source/Editor/UnrealEd/Private/Toolkits/SStandaloneAssetEditorToolkitHost.h
+++ b/Engine/Source/Editor/UnrealEd/Private/Toolkits/SStandaloneAssetEditorToolkitHost.h
@@ -93,6 +93,11 @@ public:
 	 */
 	void SetToolbar(TSharedPtr<SWidget> Toolbar);
 
+	/**
+	 * Registers a drawer for the asset editor status bar
+	 */
+	void RegisterDrawer(struct FStatusBarDrawer&& Drawer, int32 SlotIndex = INDEX_NONE);
+
 	virtual FEditorModeTools& GetEditorModeManager() const override;
 
 private:
@@ -134,6 +139,8 @@ private:
 	/** A delegate which is called any time the LevelEditor's active viewport changes. */
 	FOnActiveViewportChanged OnActiveViewportChangedDelegate;
 
+	/** Reference to bottom status bar, containing dockable tabs */
+	TSharedPtr<SWidget> StatusBarWidget;
 
 	SVerticalBox::FSlot* ToolbarSlot;
 };
diff --git a/Engine/Source/Editor/UnrealEd/Public/Toolkits/AssetEditorToolkit.h b/Engine/Source/Editor/UnrealEd/Public/Toolkits/AssetEditorToolkit.h
--- a/Engine/Source/Editor/UnrealEd/Public/Toolkits/AssetEditorToolkit.h
+++ b/Engine/Source/Editor/UnrealEd/Public/Toolkits/AssetEditorToolkit.h
@@ -177,6 +177,11 @@ public:
 	 */
 	static void SetPreviousWorldCentricToolkitHostForNewAssetEditor(TSharedRef< IToolkitHost > ToolkitHost);
 
+	/**
+	 * Registers a drawer for the asset editor status bar
+	 */
+	void RegisterDrawer(struct FStatusBarDrawer&& Drawer, int32 SlotIndex = INDEX_NONE);
+
 	/** Applies the passed in layout (or the saved user-modified version if available).  Must be called after InitAssetEditor. */
 	void RestoreFromLayout(const TSharedRef<FTabManager::FLayout>& NewLayout);
 
diff --git a/Engine/Source/Runtime/SlateCore/Public/Widgets/SWindow.h b/Engine/Source/Runtime/SlateCore/Public/Widgets/SWindow.h
--- a/Engine/Source/Runtime/SlateCore/Public/Widgets/SWindow.h
+++ b/Engine/Source/Runtime/SlateCore/Public/Widgets/SWindow.h
@@ -790,6 +790,14 @@ public:
 		WidgetToFocusOnActivate = InWidget;
 	}
 
+	/**
+	 * Returns widget last focused on deactivate
+	 */
+	TWeakPtr<SWidget> GetWidgetFocusedOnDeactivate()
+	{
+		return WidgetFocusedOnDeactivate;
+	}
+
 	UE_DEPRECATED(4.16, "ActivateWhenFirstShown() is deprecated. Please use ActivationPolicy() instead.")
 	bool ActivateWhenFirstShown() const
 	{
