UMG: Move WidgetBlueprintApplicationModes from private to public to enable user defined Mode for the Widget Blueprint Editor.
Add WidgetBlueprintGeneratedClassExtension to let the user grab the Editor from a extended menu item.
Add a OnRegisterTabsForEditor to let the user modify the menu but only in some ApplicationMode.
Add Extension for the Widget Blueprint Editor,  that let the user add extra info in the Blueprint and let the user add extra info to the generated class. The generated info is saved on the WidgetBlueprintGeneratedClass (info is on the Class and not on the UserWidget).

diff --git a/Engine/Source/Editor/UMGEditor/Private/BlueprintModes/WidgetBlueprintApplicationMode.h b/Engine/Source/Editor/UMGEditor/Private/BlueprintModes/WidgetBlueprintApplicationMode.h
deleted file mode 100644
diff --git a/Engine/Source/Editor/UMGEditor/Private/BlueprintModes/WidgetBlueprintApplicationModes.cpp b/Engine/Source/Editor/UMGEditor/Private/BlueprintModes/WidgetBlueprintApplicationModes.cpp
--- a/Engine/Source/Editor/UMGEditor/Private/BlueprintModes/WidgetBlueprintApplicationModes.cpp
+++ b/Engine/Source/Editor/UMGEditor/Private/BlueprintModes/WidgetBlueprintApplicationModes.cpp
@@ -8,17 +8,13 @@ const FName FWidgetBlueprintApplicationModes::GraphMode("GraphName");
 
 FText FWidgetBlueprintApplicationModes::GetLocalizedMode(const FName InMode)
 {
-	static TMap< FName, FText > LocModes;
-
-	if ( LocModes.Num() == 0 )
+	if (InMode == FWidgetBlueprintApplicationModes::DesignerMode)
 	{
-		LocModes.Add(DesignerMode, NSLOCTEXT("WidgetBlueprintModes", "DesignerMode", "Designer"));
-		LocModes.Add(GraphMode, NSLOCTEXT("WidgetBlueprintModes", "GraphMode", "Graph"));
+		return NSLOCTEXT("WidgetBlueprintModes", "DesignerMode", "Designer");
 	}
-
-	check(InMode != NAME_None);
-	const FText* OutDesc = LocModes.Find(InMode);
-	check(OutDesc);
-
-	return *OutDesc;
+	else if (InMode == FWidgetBlueprintApplicationModes::GraphMode)
+	{
+		return NSLOCTEXT("WidgetBlueprintModes", "GraphMode", "Graph");
+	}
+	return FText::GetEmpty();
 }
diff --git a/Engine/Source/Editor/UMGEditor/Private/BlueprintModes/WidgetBlueprintApplicationModes.h b/Engine/Source/Editor/UMGEditor/Private/BlueprintModes/WidgetBlueprintApplicationModes.h
deleted file mode 100644
diff --git a/Engine/Source/Editor/UMGEditor/Private/BlueprintModes/WidgetDesignerApplicationMode.cpp b/Engine/Source/Editor/UMGEditor/Private/BlueprintModes/WidgetDesignerApplicationMode.cpp
--- a/Engine/Source/Editor/UMGEditor/Private/BlueprintModes/WidgetDesignerApplicationMode.cpp
+++ b/Engine/Source/Editor/UMGEditor/Private/BlueprintModes/WidgetDesignerApplicationMode.cpp
@@ -106,8 +106,8 @@ FWidgetDesignerApplicationMode::FWidgetDesignerApplicationMode(TSharedPtr<FWidge
 	TabFactories.RegisterFactory(MakeShareable(new FCompilerResultsSummoner(InWidgetEditor)));
 	TabFactories.RegisterFactory(MakeShareable(new FNavigationTabSummoner(InWidgetEditor)));
 
-	// setup toolbar - clear existing toolbar extender from the BP mode
-	//@TODO: Keep this in sync with BlueprintEditorModes.cpp
+	IUMGEditorModule& EditorModule = FModuleManager::GetModuleChecked<IUMGEditorModule>("UMGEditor");
+	EditorModule.OnRegisterTabsForEditor().Broadcast(*this, TabFactories);
 
 	//Make sure we start with our existing list of extenders instead of creating a new one
 	IUMGEditorModule& UMGEditorModule = FModuleManager::LoadModuleChecked<IUMGEditorModule>("UMGEditor");
diff --git a/Engine/Source/Editor/UMGEditor/Private/BlueprintModes/WidgetGraphApplicationMode.cpp b/Engine/Source/Editor/UMGEditor/Private/BlueprintModes/WidgetGraphApplicationMode.cpp
--- a/Engine/Source/Editor/UMGEditor/Private/BlueprintModes/WidgetGraphApplicationMode.cpp
+++ b/Engine/Source/Editor/UMGEditor/Private/BlueprintModes/WidgetGraphApplicationMode.cpp
@@ -5,6 +5,7 @@
 #include "BlueprintEditorTabs.h"
 #include "SBlueprintEditorToolbar.h"
 
+#include "UMGEditorModule.h"
 #include "WidgetBlueprintEditorToolbar.h"
 #include "BlueprintModes/WidgetBlueprintApplicationModes.h"
 #include "ToolMenus.h"
@@ -68,8 +69,9 @@ FWidgetGraphApplicationMode::FWidgetGraphApplicationMode(TSharedPtr<FWidgetBluep
 			)
 		);
 	
-	// setup toolbar
-	//@TODO: Keep this in sync with AnimBlueprintMode.cpp
+	IUMGEditorModule& EditorModule = FModuleManager::GetModuleChecked<IUMGEditorModule>("UMGEditor");
+	EditorModule.OnRegisterTabsForEditor().Broadcast(*this, TabFactories);
+
 	ToolbarExtender = MakeShareable(new FExtender);
 	InWidgetEditor->GetWidgetToolbarBuilder()->AddWidgetBlueprintEditorModesToolbar(ToolbarExtender);
 
diff --git a/Engine/Source/Editor/UMGEditor/Private/UMGEditorModule.cpp b/Engine/Source/Editor/UMGEditor/Private/UMGEditorModule.cpp
--- a/Engine/Source/Editor/UMGEditor/Private/UMGEditorModule.cpp
+++ b/Engine/Source/Editor/UMGEditor/Private/UMGEditorModule.cpp
@@ -218,6 +218,12 @@ public:
 		return &WidgetBlueprintCompiler;
 	}
 
+	/** Register common tabs */
+	virtual FOnRegisterTabs& OnRegisterTabsForEditor() override
+	{
+		return RegisterTabsForEditor;
+	}
+
 private:
 	void RegisterAssetTypeAction(IAssetTools& AssetTools, TSharedRef<IAssetTypeActions> Action)
 	{
@@ -272,6 +278,9 @@ private:
 	/** Compiler customization for Widgets */
 	FWidgetBlueprintCompiler WidgetBlueprintCompiler;
 
+	/** */
+	FOnRegisterTabs RegisterTabsForEditor;
+
 	bool bThumbnailRenderersRegistered;
 	bool bOnPostEngineInitHandled;
 };
diff --git a/Engine/Source/Editor/UMGEditor/Private/WidgetBlueprintCompiler.cpp b/Engine/Source/Editor/UMGEditor/Private/WidgetBlueprintCompiler.cpp
--- a/Engine/Source/Editor/UMGEditor/Private/WidgetBlueprintCompiler.cpp
+++ b/Engine/Source/Editor/UMGEditor/Private/WidgetBlueprintCompiler.cpp
@@ -20,6 +20,7 @@
 #include "IUMGModule.h"
 #include "UMGEditorProjectSettings.h"
 #include "WidgetCompilerRule.h"
+#include "WidgetBlueprintExtension.h"
 #include "Editor/WidgetCompilerLog.h"
 #include "Editor.h"
 
@@ -99,10 +100,18 @@ FWidgetBlueprintCompilerContext::FWidgetBlueprintCompilerContext(UWidgetBlueprin
 	, OldWidgetTree(nullptr)
 	, WidgetSchema(nullptr)
 {
+	UWidgetBlueprintExtension::ForEachExtension(WidgetBlueprint(), [this](UWidgetBlueprintExtension* InExtension)
+		{
+			InExtension->BeginCompilation(*this);
+		});
 }
 
 FWidgetBlueprintCompilerContext::~FWidgetBlueprintCompilerContext()
 {
+	UWidgetBlueprintExtension::ForEachExtension(WidgetBlueprint(), [](UWidgetBlueprintExtension* InExtension)
+		{
+			InExtension->EndCompilation();
+		});
 }
 
 UEdGraphSchema_K2* FWidgetBlueprintCompilerContext::CreateSchema()
@@ -178,6 +187,11 @@ void FWidgetBlueprintCompilerContext::CreateFunctionList()
 			}
 		}
 	}
+
+	UWidgetBlueprintExtension::ForEachExtension(WidgetBlueprint(), [this](UWidgetBlueprintExtension* InExtension)
+		{
+			InExtension->CreateFunctionList();
+		});
 }
 
 void FWidgetBlueprintCompilerContext::ValidateWidgetNames()
@@ -294,8 +308,16 @@ void FWidgetBlueprintCompilerContext::CleanAndSanitizeClass(UBlueprintGeneratedC
 		RenameObjectToTransientPackage(Animation, FName(), false);
 	}
 	NewWidgetBlueprintClass->Animations.Empty();
-
 	NewWidgetBlueprintClass->Bindings.Empty();
+	NewWidgetBlueprintClass->Extensions.Empty();
+
+	if (UWidgetBlueprintGeneratedClass* WidgetClassToClean = Cast<UWidgetBlueprintGeneratedClass>(ClassToClean))
+	{
+		UWidgetBlueprintExtension::ForEachExtension(WidgetBlueprint(), [WidgetClassToClean, InOutOldCDO](UWidgetBlueprintExtension* InExtension)
+			{
+				InExtension->CleanAndSanitizeClass(WidgetClassToClean, InOutOldCDO);
+			});
+	}
 }
 
 void FWidgetBlueprintCompilerContext::SaveSubObjectsFromCleanAndSanitizeClass(FSubobjectCollection& SubObjectsToSave, UBlueprintGeneratedClass* ClassToClean)
@@ -461,6 +483,11 @@ void FWidgetBlueprintCompilerContext::CreateClassVariablesFromBlueprint()
 			WidgetAnimToMemberVariableMap.Add(Animation, AnimationProperty);
 		}
 	}
+
+	UWidgetBlueprintExtension::ForEachExtension(WidgetBlueprint(), [](UWidgetBlueprintExtension* InExtension)
+		{
+			InExtension->CreateClassVariablesFromBlueprint();
+		});
 }
 
 void FWidgetBlueprintCompilerContext::CopyTermDefaultsToDefaultObject(UObject* DefaultObject)
@@ -528,6 +555,11 @@ void FWidgetBlueprintCompilerContext::CopyTermDefaultsToDefaultObject(UObject* D
 			MessageLog.Warning(*LOCTEXT("NonTickableButTickFound", "This widget has a blueprint implemented Tick event but the widget is set to never tick.  This tick event will never be called.").ToString());
 		}
 	}
+
+	UWidgetBlueprintExtension::ForEachExtension(WidgetBlueprint(), [](UWidgetBlueprintExtension* InExtension)
+		{
+			InExtension->CreateClassVariablesFromBlueprint();
+		});
 }
 
 void FWidgetBlueprintCompilerContext::SanitizeBindings(UBlueprintGeneratedClass* Class)
@@ -859,6 +891,11 @@ void FWidgetBlueprintCompilerContext::FinishCompilingClass(UClass* Class)
 	}
 
 	Super::FinishCompilingClass(Class);
+
+	UWidgetBlueprintExtension::ForEachExtension(WidgetBlueprint(), [BPGClass](UWidgetBlueprintExtension* InExtension)
+		{
+			InExtension->FinishCompilingClass(BPGClass);
+		});
 }
 
 
@@ -968,10 +1005,27 @@ void FWidgetBlueprintCompilerContext::VerifyEventReplysAreNotEmpty(FKismetFuncti
 
 bool FWidgetBlueprintCompilerContext::ValidateGeneratedClass(UBlueprintGeneratedClass* Class)
 {
-	bool SuperResult = Super::ValidateGeneratedClass(Class);
-	bool Result = UWidgetBlueprint::ValidateGeneratedClass(Class);
+	const bool bSuperResult = Super::ValidateGeneratedClass(Class);
+	const bool bResult = UWidgetBlueprint::ValidateGeneratedClass(Class);
 
-	return SuperResult && Result;
+	UWidgetBlueprintGeneratedClass* WidgetClass = Cast<UWidgetBlueprintGeneratedClass>(Class);
+	bool bExtension = WidgetClass != nullptr;
+	if (bExtension)
+	{
+		UWidgetBlueprintExtension::ForEachExtension(WidgetBlueprint(), [&bExtension, WidgetClass](UWidgetBlueprintExtension* InExtension)
+			{
+				bExtension = InExtension->ValidateGeneratedClass(WidgetClass) && bExtension;
+			});
+	}
+
+	return bSuperResult && bResult && bExtension;
+}
+
+void FWidgetBlueprintCompilerContext::AddExtension(UWidgetBlueprintGeneratedClass* Class, UWidgetBlueprintGeneratedClassExtension* Extension)
+{
+	check(Class);
+	check(Extension);
+	Class->Extensions.Add(Extension);
 }
 
 #undef LOCTEXT_NAMESPACE
diff --git a/Engine/Source/Editor/UMGEditor/Private/WidgetBlueprintEditor.cpp b/Engine/Source/Editor/UMGEditor/Private/WidgetBlueprintEditor.cpp
--- a/Engine/Source/Editor/UMGEditor/Private/WidgetBlueprintEditor.cpp
+++ b/Engine/Source/Editor/UMGEditor/Private/WidgetBlueprintEditor.cpp
@@ -15,6 +15,7 @@
 #include "WidgetBlueprint.h"
 #include "StatusBarSubsystem.h"
 #include "Editor.h"
+#include "WidgetBlueprintToolMenuContext.h"
 
 #if WITH_EDITOR
 	#include "EditorStyleSet.h"
@@ -189,7 +190,7 @@ void FWidgetBlueprintEditor::InitWidgetBlueprintEditor(const EToolkitMode::Type
 
 void FWidgetBlueprintEditor::InitalizeExtenders()
 {
-	FBlueprintEditor::InitalizeExtenders();
+	Super::InitalizeExtenders();
 
 	IUMGEditorModule& UMGEditorModule = FModuleManager::LoadModuleChecked<IUMGEditorModule>("UMGEditor");
 	AddMenuExtender(UMGEditorModule.GetMenuExtensibilityManager()->GetAllExtenders(GetToolkitCommands(), GetEditingObjects()));
@@ -380,8 +381,6 @@ void FWidgetBlueprintEditor::TakeSnapshot()
 			}
 		}
 	}
-
-
 }
 
 void FWidgetBlueprintEditor::CaptureThumbnail()
@@ -428,6 +427,35 @@ bool FWidgetBlueprintEditor::IsPreviewWidgetInitialized()
 	return GetPreview() != nullptr;
 }
 
+FName FWidgetBlueprintEditor::GetToolkitFName() const
+{
+	return FName("WidgetBlueprintEditor");
+}
+
+FText FWidgetBlueprintEditor::GetBaseToolkitName() const
+{
+	return LOCTEXT("AppLabel", "Widget Editor");
+}
+
+FString FWidgetBlueprintEditor::GetWorldCentricTabPrefix() const
+{
+	return LOCTEXT("WorldCentricTabPrefix", "Widget Editor ").ToString();
+}
+
+FLinearColor FWidgetBlueprintEditor::GetWorldCentricTabColorScale() const
+{
+	return FLinearColor(0.3f, 0.25f, 0.35f, 0.5f);
+}
+
+void FWidgetBlueprintEditor::InitToolMenuContext(FToolMenuContext& MenuContext)
+{
+	Super::InitToolMenuContext(MenuContext);
+
+	UWidgetBlueprintToolMenuContext* Context = NewObject<UWidgetBlueprintToolMenuContext>();
+	Context->WidgetBlueprintEditor = SharedThis(this);
+	MenuContext.AddObject(Context);
+}
+
 void FWidgetBlueprintEditor::SetCreateOnCompileSetting(EDisplayOnCompile InCreateOnCompile)
 {
 	UWidgetDesignerSettings* Settings = GetMutableDefault<UWidgetDesignerSettings>();
@@ -477,7 +505,7 @@ void FWidgetBlueprintEditor::OnCreateNativeBaseClassSuccessfully(const FString&
 
 void FWidgetBlueprintEditor::RegisterApplicationModes(const TArray<UBlueprint*>& InBlueprints, bool bShouldOpenInDefaultsMode, bool bNewlyCreated/* = false*/)
 {
-	//FBlueprintEditor::RegisterApplicationModes(InBlueprints, bShouldOpenInDefaultsMode);
+	//Super::RegisterApplicationModes(InBlueprints, bShouldOpenInDefaultsMode);
 
 	if ( InBlueprints.Num() == 1 )
 	{
@@ -659,7 +687,7 @@ void FWidgetBlueprintEditor::OnBlueprintChangedImpl(UBlueprint* InBlueprint, boo
 {
 	DestroyPreview();
 
-	FBlueprintEditor::OnBlueprintChangedImpl(InBlueprint, bIsJustBeingCompiled);
+	Super::OnBlueprintChangedImpl(InBlueprint, bIsJustBeingCompiled);
 
 	if ( InBlueprint )
 	{
@@ -820,7 +848,7 @@ void FWidgetBlueprintEditor::DuplicateSelectedWidgets()
 
 void FWidgetBlueprintEditor::Tick(float DeltaTime)
 {
-	FBlueprintEditor::Tick(DeltaTime);
+	Super::Tick(DeltaTime);
 
 	// Tick the preview scene world.
 	if ( !GIntraFrameDebuggingGameThread )
@@ -936,7 +964,7 @@ static bool MigratePropertyValue(UObject* SourceObject, UObject* DestinationObje
 
 void FWidgetBlueprintEditor::AddReferencedObjects( FReferenceCollector& Collector )
 {
-	FBlueprintEditor::AddReferencedObjects( Collector );
+	Super::AddReferencedObjects( Collector );
 
 	UUserWidget* Preview = GetPreview();
 	Collector.AddReferencedObject( Preview );
@@ -978,14 +1006,14 @@ void FWidgetBlueprintEditor::MigrateFromChain(FEditPropertyChain* PropertyThatCh
 
 void FWidgetBlueprintEditor::PostUndo(bool bSuccessful)
 {
-	FBlueprintEditor::PostUndo(bSuccessful);
+	Super::PostUndo(bSuccessful);
 
 	OnWidgetBlueprintTransaction.Broadcast();
 }
 
 void FWidgetBlueprintEditor::PostRedo(bool bSuccessful)
 {
-	FBlueprintEditor::PostRedo(bSuccessful);
+	Super::PostRedo(bSuccessful);
 
 	OnWidgetBlueprintTransaction.Broadcast();
 }
@@ -1513,7 +1541,7 @@ void FWidgetBlueprintEditor::UpdatePreview(UBlueprint* InBlueprint, bool bInForc
 
 FGraphAppearanceInfo FWidgetBlueprintEditor::GetGraphAppearance(UEdGraph* InGraph) const
 {
-	FGraphAppearanceInfo AppearanceInfo = FBlueprintEditor::GetGraphAppearance(InGraph);
+	FGraphAppearanceInfo AppearanceInfo = Super::GetGraphAppearance(InGraph);
 
 	if ( GetBlueprintObj()->IsA(UWidgetBlueprint::StaticClass()) )
 	{
diff --git a/Engine/Source/Editor/UMGEditor/Private/WidgetBlueprintExtension.cpp b/Engine/Source/Editor/UMGEditor/Private/WidgetBlueprintExtension.cpp
new file mode 100644
--- /dev/null
+++ b/Engine/Source/Editor/UMGEditor/Private/WidgetBlueprintExtension.cpp
@@ -0,0 +1,60 @@
+// Copyright Epic Games, Inc. All Rights Reserved.
+
+#include "WidgetBlueprintExtension.h"
+#include "Engine/Blueprint.h"
+#include "WidgetBlueprint.h"
+
+
+UWidgetBlueprintExtension* UWidgetBlueprintExtension::RequestExtension(UWidgetBlueprint* InBlueprint, TSubclassOf<UWidgetBlueprintExtension> InExtensionType)
+{
+	
+	checkf(!InBlueprint->bBeingCompiled, TEXT("Do not use RequestExtension when a blueprint is being compiled."));
+
+	// Look for an existing extension
+	if (UWidgetBlueprintExtension* ExistingExtension = GetExtension(InBlueprint, InExtensionType))
+	{
+		return ExistingExtension;
+	}
+
+	// Not found, create one
+	UWidgetBlueprintExtension* NewExtension = NewObject<UWidgetBlueprintExtension>(InBlueprint, InExtensionType.Get());
+	InBlueprint->Extensions.Add(NewExtension);
+	return NewExtension;
+}
+
+
+UWidgetBlueprintExtension* UWidgetBlueprintExtension::GetExtension(UWidgetBlueprint* InBlueprint, TSubclassOf<UWidgetBlueprintExtension> InExtensionType)
+{
+	// Look for an existing extension
+	for (UBlueprintExtension* Extension : InBlueprint->Extensions)
+	{
+		if (Extension && Extension->GetClass() == InExtensionType)
+		{
+			return CastChecked<UWidgetBlueprintExtension>(Extension);
+		}
+	}
+
+	return nullptr;
+}
+
+
+TArray<UWidgetBlueprintExtension*> UWidgetBlueprintExtension::GetExtensions(UWidgetBlueprint* InBlueprint)
+{
+	TArray<UWidgetBlueprintExtension*> Extensions;
+
+	for (UBlueprintExtension* Extension : InBlueprint->Extensions)
+	{
+		if (Extension && Extension->GetClass()->IsChildOf(UWidgetBlueprintExtension::StaticClass()))
+		{
+			Extensions.Add(CastChecked<UWidgetBlueprintExtension>(Extension));
+		}
+	}
+
+	return Extensions;
+}
+
+
+UWidgetBlueprint* UWidgetBlueprintExtension::GetWidgetBlueprint() const
+{
+	return CastChecked<UWidgetBlueprint>(GetOuter());
+}
diff --git a/Engine/Source/Editor/UMGEditor/Public/BlueprintModes/WidgetBlueprintApplicationMode.h b/Engine/Source/Editor/UMGEditor/Public/BlueprintModes/WidgetBlueprintApplicationMode.h
new file mode 100644
--- /dev/null
+++ b/Engine/Source/Editor/UMGEditor/Public/BlueprintModes/WidgetBlueprintApplicationMode.h
@@ -0,0 +1,29 @@
+// Copyright Epic Games, Inc. All Rights Reserved.
+
+#pragma once
+
+#include "CoreMinimal.h"
+#include "WorkflowOrientedApp/WorkflowTabManager.h"
+#include "WidgetBlueprintEditor.h"
+#include "BlueprintEditorModes.h"
+
+class UWidgetBlueprint;
+
+/////////////////////////////////////////////////////
+// FWidgetBlueprintApplicationMode
+
+class UMGEDITOR_API FWidgetBlueprintApplicationMode : public FBlueprintEditorApplicationMode
+{
+public:
+	FWidgetBlueprintApplicationMode(TSharedPtr<class FWidgetBlueprintEditor> InWidgetEditor, FName InModeName);
+
+public:
+	TSharedPtr<class FWidgetBlueprintEditor> GetBlueprintEditor() const;
+	UWidgetBlueprint* GetBlueprint() const;
+
+protected:
+	TWeakPtr<class FWidgetBlueprintEditor> MyWidgetBlueprintEditor;
+
+	// Set of spawnable tabs in the mode
+	FWorkflowAllowedTabSet TabFactories;
+};
diff --git a/Engine/Source/Editor/UMGEditor/Public/BlueprintModes/WidgetBlueprintApplicationModes.h b/Engine/Source/Editor/UMGEditor/Public/BlueprintModes/WidgetBlueprintApplicationModes.h
new file mode 100644
--- /dev/null
+++ b/Engine/Source/Editor/UMGEditor/Public/BlueprintModes/WidgetBlueprintApplicationModes.h
@@ -0,0 +1,21 @@
+// Copyright Epic Games, Inc. All Rights Reserved.
+
+#pragma once
+
+#include "CoreMinimal.h"
+
+/////////////////////////////////////////////////////
+// FWidgetBlueprintApplicationModes
+
+// This is the list of IDs for widget blueprint editor modes
+struct UMGEDITOR_API FWidgetBlueprintApplicationModes
+{
+	// Mode constants
+	static const FName DesignerMode;
+	static const FName GraphMode;
+
+	static FText GetLocalizedMode(const FName InMode);
+
+private:
+	FWidgetBlueprintApplicationModes() = delete;
+};
diff --git a/Engine/Source/Editor/UMGEditor/Public/UMGEditorModule.h b/Engine/Source/Editor/UMGEditor/Public/UMGEditorModule.h
--- a/Engine/Source/Editor/UMGEditor/Public/UMGEditorModule.h
+++ b/Engine/Source/Editor/UMGEditor/Public/UMGEditorModule.h
@@ -11,10 +11,15 @@ extern const FName UMGEditorAppIdentifier;
 
 class FUMGEditor;
 class FWidgetBlueprintCompiler;
+class FWidgetBlueprintApplicationMode;
+class FWorkflowAllowedTabSet;
 
 /** The public interface of the UMG editor module. */
 class IUMGEditorModule : public IModuleInterface, public IHasMenuExtensibility, public IHasToolBarExtensibility, public IHasDesignerExtensibility
 {
 public:
 	virtual FWidgetBlueprintCompiler* GetRegisteredCompiler() = 0;
+
+	DECLARE_EVENT_TwoParams(IUMGEditorModule, FOnRegisterTabs, const FWidgetBlueprintApplicationMode&, FWorkflowAllowedTabSet&);
+	virtual FOnRegisterTabs& OnRegisterTabsForEditor() = 0;
 };
diff --git a/Engine/Source/Editor/UMGEditor/Public/WidgetBlueprintCompiler.h b/Engine/Source/Editor/UMGEditor/Public/WidgetBlueprintCompiler.h
--- a/Engine/Source/Editor/UMGEditor/Public/WidgetBlueprintCompiler.h
+++ b/Engine/Source/Editor/UMGEditor/Public/WidgetBlueprintCompiler.h
@@ -76,6 +76,9 @@ protected:
 
 	void VerifyEventReplysAreNotEmpty(FKismetFunctionContext& Context);
 
+public:
+	void AddExtension(UWidgetBlueprintGeneratedClass* Class, UWidgetBlueprintGeneratedClassExtension* Extension);
+
 protected:
 	void FixAbandonedWidgetTree(UWidgetBlueprint* WidgetBP);
 
diff --git a/Engine/Source/Editor/UMGEditor/Public/WidgetBlueprintEditor.h b/Engine/Source/Editor/UMGEditor/Public/WidgetBlueprintEditor.h
--- a/Engine/Source/Editor/UMGEditor/Public/WidgetBlueprintEditor.h
+++ b/Engine/Source/Editor/UMGEditor/Public/WidgetBlueprintEditor.h
@@ -37,6 +37,9 @@ struct FNamedSlotSelection
  */
 class UMGEDITOR_API FWidgetBlueprintEditor : public FBlueprintEditor
 {
+private:
+	using Super = FBlueprintEditor;
+
 public:
 	DECLARE_MULTICAST_DELEGATE_OneParam(FOnHoveredWidgetSet, const FWidgetReference&)
 	DECLARE_MULTICAST_DELEGATE(FOnHoveredWidgetCleared);
@@ -59,14 +62,24 @@ public:
 
 	void InitWidgetBlueprintEditor(const EToolkitMode::Type Mode, const TSharedPtr< IToolkitHost >& InitToolkitHost, const TArray<UBlueprint*>& InBlueprints, bool bShouldOpenInDefaultsMode);
 
-	/** FBlueprintEditor interface */
+	//~ Begin FBlueprintEditor interface
 	virtual void Tick(float DeltaTime) override;
 	virtual void PostUndo(bool bSuccessful) override;
 	virtual void PostRedo(bool bSuccessful) override;
 	virtual void Compile() override;
+	//~ End Begin FBlueprintEditor interface
 
-	/** FGCObjectInterface */
+	//~ Begin FGCObjectInterface interface
 	virtual void AddReferencedObjects( FReferenceCollector& Collector ) override;
+	//~ end FGCObjectInterface interface
+
+	//~ Begin IToolkit interface
+	virtual FName GetToolkitFName() const override;
+	virtual FText GetBaseToolkitName() const override;
+	virtual FString GetWorldCentricTabPrefix() const override;
+	virtual FLinearColor GetWorldCentricTabColorScale() const override;
+	virtual void InitToolMenuContext(FToolMenuContext& MenuContext) override;
+	//~ End IToolkit interface
 
 	/** @return The widget blueprint currently being edited in this editor */
 	class UWidgetBlueprint* GetWidgetBlueprintObj() const;
@@ -265,11 +278,11 @@ protected:
 	void OnCreateNativeBaseClassSuccessfully(const FString& InClassName, const FString& InClassPath, const FString& InModuleName);
 	TSharedPtr<ISequencer> CreateSequencerWidgetInternal();
 
-	// Begin FBlueprintEditor
+	//~ Begin FBlueprintEditor interface
 	virtual void RegisterApplicationModes(const TArray<UBlueprint*>& InBlueprints, bool bShouldOpenInDefaultsMode, bool bNewlyCreated = false) override;
 	virtual FGraphAppearanceInfo GetGraphAppearance(class UEdGraph* InGraph) const override;
 	virtual TSubclassOf<UEdGraphSchema> GetDefaultSchemaClass() const override;
-	// End FBlueprintEditor
+	//~ End FBlueprintEditor interface
 
 private:
 	bool CanDeleteSelectedWidgets();
diff --git a/Engine/Source/Editor/UMGEditor/Public/WidgetBlueprintExtension.h b/Engine/Source/Editor/UMGEditor/Public/WidgetBlueprintExtension.h
new file mode 100644
--- /dev/null
+++ b/Engine/Source/Editor/UMGEditor/Public/WidgetBlueprintExtension.h
@@ -0,0 +1,131 @@
+// Copyright Epic Games, Inc. All Rights Reserved.
+
+#pragma once
+
+#include "CoreMinimal.h"
+#include "UObject/ObjectMacros.h"
+#include "Blueprint/BlueprintExtension.h"
+#include "WidgetBlueprint.h"
+
+class FSubobjectCollection;
+class FWidgetBlueprintCompilerContext;
+class UWidgetBlueprintGeneratedClass;
+
+#include "WidgetBlueprintExtension.generated.h"
+
+/** Extension that allows per-system data to be held on the widget blueprint, and per-system logic to be executed during compilation */
+UCLASS()
+class UMGEDITOR_API UWidgetBlueprintExtension : public UBlueprintExtension
+{
+	GENERATED_BODY()
+
+public:
+	/**
+	 * Request an WidgetBlueprintExtension for an WidgetBlueprint.
+	 * @note It is illegal to perform this operation once compilation has commenced, use GetExtension instead.
+	 */
+	template<typename ExtensionType>
+	static ExtensionType* RequestExtension(UWidgetBlueprint* InWidgetBlueprint)
+	{
+		return CastChecked<ExtensionType>(RequestExtension(InWidgetBlueprint, ExtensionType::StaticClass()));
+	}
+
+	/**
+	 * Request an WidgetBlueprintExtension for an WidgetBlueprint.
+	 * @note It is illegal to perform this operation once compilation has commenced, use GetExtension instead.
+	 */
+	static UWidgetBlueprintExtension* RequestExtension(UWidgetBlueprint* InWidgetBlueprint, TSubclassOf<UWidgetBlueprintExtension> InExtensionType);
+
+	/** Get an already-requested extension for an WidgetBlueprint. */
+	template<typename ExtensionType>
+    static ExtensionType* GetExtension(UWidgetBlueprint* InWidgetBlueprint)
+	{
+		return Cast<ExtensionType>(GetExtension(InWidgetBlueprint, ExtensionType::StaticClass()));
+	}
+
+	/** Get an already-requested extension for an WidgetBlueprint. */
+	static UWidgetBlueprintExtension* GetExtension(UWidgetBlueprint* InWidgetBlueprint, TSubclassOf<UWidgetBlueprintExtension> InExtensionType);
+
+	/** Get all subsystems currently present on an WidgetBlueprint */
+	static TArray<UWidgetBlueprintExtension*> GetExtensions(UWidgetBlueprint* InWidgetBlueprint);
+
+	/** Iterate over all registered WidgetBlueprintExtensions in an WidgetBlueprint */
+	template<typename Predicate>
+	static void ForEachExtension(UWidgetBlueprint* InWidgetBlueprint, Predicate Pred)
+	{
+		for (UBlueprintExtension* BlueprintExtension : InWidgetBlueprint->Extensions)
+		{
+			if (UWidgetBlueprintExtension* WidgetBlueprintExtension = Cast<UWidgetBlueprintExtension>(BlueprintExtension))
+			{
+				Pred(WidgetBlueprintExtension);
+			}
+		}
+	}
+
+	/** Get the WidgetBlueprint that hosts this extension */
+	UWidgetBlueprint* GetWidgetBlueprint() const;
+
+protected:
+	/** 
+	 * Override point called when a compiler context is created for the WidgetBlueprint
+	 * @param	InCreationContext	The compiler context for the current compilation
+	 */
+	virtual void HandleBeginCompilation(FWidgetBlueprintCompilerContext& InCreationContext) {}
+
+	virtual void HandleCreateFunctionList() {}
+	virtual void HandleCleanAndSanitizeClass(UWidgetBlueprintGeneratedClass* ClassToClean, UObject* OldCDO) {}
+	virtual void HandleCreateClassVariablesFromBlueprint() {}
+	virtual void HandleCopyTermDefaultsToDefaultObject(UObject* DefaultObject) {}
+	virtual void HandleFinishCompilingClass(UWidgetBlueprintGeneratedClass* Class) {}
+	virtual bool HandleValidateGeneratedClass(UWidgetBlueprintGeneratedClass* Class) { return true; }
+
+	/**
+	 * Override point called when a compiler context is destroyed for the WidgetBlueprint.
+	 * Can be used to clean up resources.
+	 */
+	virtual void HandleEndCompilation() {}
+
+
+private:
+	friend FWidgetBlueprintCompilerContext;
+
+	void BeginCompilation(FWidgetBlueprintCompilerContext& InCreationContext)
+	{
+		HandleBeginCompilation(InCreationContext);
+	}
+
+	void CreateFunctionList()
+	{
+		HandleCreateFunctionList();
+	}
+
+	void CleanAndSanitizeClass(UWidgetBlueprintGeneratedClass* ClassToClean, UObject* OldCDO)
+	{
+		HandleCleanAndSanitizeClass(ClassToClean, OldCDO);
+	}
+
+	void CreateClassVariablesFromBlueprint()
+	{
+		HandleCreateClassVariablesFromBlueprint();
+	}
+
+	void CopyTermDefaultsToDefaultObject(UObject* DefaultObject)
+	{
+		HandleCopyTermDefaultsToDefaultObject(DefaultObject);
+	}
+
+	void FinishCompilingClass(UWidgetBlueprintGeneratedClass* Class)
+	{
+		HandleFinishCompilingClass(Class);
+	}
+
+	bool ValidateGeneratedClass(UWidgetBlueprintGeneratedClass* Class)
+	{
+		return HandleValidateGeneratedClass(Class);
+	}
+
+	void EndCompilation()
+	{
+		HandleEndCompilation();
+	}
+};
diff --git a/Engine/Source/Editor/UMGEditor/Public/WidgetBlueprintToolMenuContext.h b/Engine/Source/Editor/UMGEditor/Public/WidgetBlueprintToolMenuContext.h
new file mode 100644
--- /dev/null
+++ b/Engine/Source/Editor/UMGEditor/Public/WidgetBlueprintToolMenuContext.h
@@ -0,0 +1,18 @@
+// Copyright Epic Games, Inc. All Rights Reserved.
+
+#pragma once
+
+#include "CoreMinimal.h"
+
+#include "WidgetBlueprintToolMenuContext.generated.h"
+
+class FWidgetBlueprintEditor;
+
+UCLASS()
+class UMGEDITOR_API UWidgetBlueprintToolMenuContext : public UObject
+{
+	GENERATED_BODY()
+
+public:
+	TWeakPtr<FWidgetBlueprintEditor> WidgetBlueprintEditor;
+};
diff --git a/Engine/Source/Runtime/UMG/Private/UserWidget.cpp b/Engine/Source/Runtime/UMG/Private/UserWidget.cpp
--- a/Engine/Source/Runtime/UMG/Private/UserWidget.cpp
+++ b/Engine/Source/Runtime/UMG/Private/UserWidget.cpp
@@ -14,6 +14,7 @@
 #include "Blueprint/WidgetBlueprintGeneratedClass.h"
 #include "Animation/UMGSequencePlayer.h"
 #include "Animation/UMGSequenceTickManager.h"
+#include "Extensions/WidgetBlueprintGeneratedClassExtension.h"
 #include "UObject/UnrealType.h"
 #include "Blueprint/WidgetNavigation.h"
 #include "Animation/WidgetAnimation.h"
@@ -1343,22 +1344,55 @@ void UUserWidget::BindToAnimationEvent(UWidgetAnimation* InAnimation, FWidgetAni
 
 void UUserWidget::NativeOnInitialized()
 {
+	if (UWidgetBlueprintGeneratedClass* Class = GetWidgetTreeOwningClass())
+	{
+		Class->ForEachExtension([this](UWidgetBlueprintGeneratedClassExtension* Extension)
+			{
+				Extension->Initialize(this);
+			});
+	}
+
 	OnInitialized();
 }
 
 void UUserWidget::NativePreConstruct()
 {
-	PreConstruct(IsDesignTime());
+	const bool bIsDesignTime = IsDesignTime();
+	if (UWidgetBlueprintGeneratedClass* Class = GetWidgetTreeOwningClass())
+	{
+		Class->ForEachExtension([this, bIsDesignTime](UWidgetBlueprintGeneratedClassExtension* Extension)
+			{
+				Extension->PreConstruct(this, bIsDesignTime);
+			});
+	}
+
+	PreConstruct(bIsDesignTime);
 }
 
 void UUserWidget::NativeConstruct()
 {
+	if (UWidgetBlueprintGeneratedClass* Class = GetWidgetTreeOwningClass())
+	{
+		Class->ForEachExtension([this](UWidgetBlueprintGeneratedClassExtension* Extension)
+			{
+				Extension->Construct(this);
+			});
+	}
+
 	Construct();
 	UpdateCanTick();
 }
 
 void UUserWidget::NativeDestruct()
 {
+	if (UWidgetBlueprintGeneratedClass* Class = GetWidgetTreeOwningClass())
+	{
+		Class->ForEachExtension([this](UWidgetBlueprintGeneratedClassExtension* Extension)
+		{
+			Extension->Destruct(this);
+		});
+	}
+
 	StopListeningForAllInputActions();
 	Destruct();
 }
diff --git a/Engine/Source/Runtime/UMG/Private/WidgetBlueprintGeneratedClass.cpp b/Engine/Source/Runtime/UMG/Private/WidgetBlueprintGeneratedClass.cpp
--- a/Engine/Source/Runtime/UMG/Private/WidgetBlueprintGeneratedClass.cpp
+++ b/Engine/Source/Runtime/UMG/Private/WidgetBlueprintGeneratedClass.cpp
@@ -7,6 +7,7 @@
 #include "Animation/WidgetAnimation.h"
 #include "Serialization/TextReferenceCollector.h"
 #include "Engine/UserInterfaceSettings.h"
+#include "Extensions/WidgetBlueprintGeneratedClassExtension.h"
 #include "UMGPrivate.h"
 #include "Interfaces/ITargetPlatform.h"
 #include "UObject/EditorObjectVersion.h"
@@ -17,6 +18,7 @@
 #include "UObject/UObjectHash.h"
 #include "UObject/UObjectIterator.h"
 
+
 #if WITH_EDITOR
 #include "Engine/Blueprint.h"
 #endif
@@ -114,11 +116,12 @@ namespace
 /////////////////////////////////////////////////////
 // UWidgetBlueprintGeneratedClass
 
-UWidgetBlueprintGeneratedClass::UWidgetBlueprintGeneratedClass(const FObjectInitializer& ObjectInitializer)
-	: Super(ObjectInitializer)
+UWidgetBlueprintGeneratedClass::UWidgetBlueprintGeneratedClass()
 {
 #if WITH_EDITORONLY_DATA
-	{ static const FAutoRegisterTextReferenceCollectorCallback AutomaticRegistrationOfTextReferenceCollector(UWidgetBlueprintGeneratedClass::StaticClass(), &CollectWidgetBlueprintGeneratedClassTextReferences); }
+	{
+		static const FAutoRegisterTextReferenceCollectorCallback AutomaticRegistrationOfTextReferenceCollector(UWidgetBlueprintGeneratedClass::StaticClass(), &CollectWidgetBlueprintGeneratedClassTextReferences);
+	}
 	bCanCallPreConstruct = true;
 #endif
 }
@@ -239,11 +242,8 @@ void UWidgetBlueprintGeneratedClass::InitializeWidgetStatic(UUserWidget* UserWid
 			Widget->WidgetGeneratedBy = InClass->ClassGeneratedBy;
 #endif
 
-			// TODO UMG Make this an FName
-			FString VariableName = Widget->GetName();
-
 			// Find property with the same name as the template and assign the new widget to it.
-			FObjectPropertyBase* Prop = FindFProperty<FObjectPropertyBase>(WidgetBlueprintClass, *VariableName);
+			FObjectPropertyBase* Prop = FindFProperty<FObjectPropertyBase>(WidgetBlueprintClass, Widget->GetFName());
 			if (Prop)
 			{
 				Prop->SetObjectPropertyValue_InContainer(UserWidget, Widget);
@@ -269,8 +269,6 @@ void UWidgetBlueprintGeneratedClass::InitializeWidgetStatic(UUserWidget* UserWid
 		{
 			UBlueprintGeneratedClass::BindDynamicDelegates(InClass, UserWidget);
 		}
-
-		//TODO UMG Add OnWidgetInitialized?
 	}
 }
 
@@ -439,4 +437,31 @@ UWidgetBlueprintGeneratedClass* UWidgetBlueprintGeneratedClass::FindWidgetTreeOw
 	return nullptr;
 }
 
+
+UWidgetBlueprintGeneratedClassExtension* UWidgetBlueprintGeneratedClass::GetExtension(TSubclassOf<UWidgetBlueprintGeneratedClassExtension> InExtensionType)
+{
+	for (UWidgetBlueprintGeneratedClassExtension* Extension : Extensions)
+	{
+		if (Extension->IsA(InExtensionType))
+		{
+			return Extension;
+		}
+	}
+	return nullptr;
+}
+
+
+TArray<UWidgetBlueprintGeneratedClassExtension*> UWidgetBlueprintGeneratedClass::GetExtensions(TSubclassOf<UWidgetBlueprintGeneratedClassExtension> InExtensionType)
+{
+	TArray<UWidgetBlueprintGeneratedClassExtension*> Result;
+	for (UWidgetBlueprintGeneratedClassExtension* Extension : Extensions)
+	{
+		if (Extension->IsA(InExtensionType))
+		{
+			Result.Add(Extension);
+		}
+	}
+	return Result;
+}
+
 #undef LOCTEXT_NAMESPACE
diff --git a/Engine/Source/Runtime/UMG/Public/Blueprint/WidgetBlueprintGeneratedClass.h b/Engine/Source/Runtime/UMG/Public/Blueprint/WidgetBlueprintGeneratedClass.h
--- a/Engine/Source/Runtime/UMG/Public/Blueprint/WidgetBlueprintGeneratedClass.h
+++ b/Engine/Source/Runtime/UMG/Public/Blueprint/WidgetBlueprintGeneratedClass.h
@@ -11,6 +11,7 @@
 
 class UUserWidget;
 class UWidgetAnimation;
+class UWidgetBlueprintGeneratedClassExtension;
 class UWidgetTree;
 
 UENUM()
@@ -56,7 +57,11 @@ struct FDelegateRuntimeBinding
 UCLASS()
 class UMG_API UWidgetBlueprintGeneratedClass : public UBlueprintGeneratedClass
 {
-	GENERATED_UCLASS_BODY()
+	GENERATED_BODY()
+	friend class FWidgetBlueprintCompilerContext;
+
+public:
+	UWidgetBlueprintGeneratedClass();
 
 private:
 
@@ -64,6 +69,10 @@ private:
 	UPROPERTY()
 	TObjectPtr<UWidgetTree> WidgetTree;
 
+	/** The extension that are considered static to the class */
+	UPROPERTY()
+	TArray<TObjectPtr<UWidgetBlueprintGeneratedClassExtension>> Extensions;
+
 #if WITH_EDITORONLY_DATA
 
 public:
@@ -113,8 +122,6 @@ public:
 	 */
 	void InitializeWidget(UUserWidget* UserWidget) const;
 
-	static void InitializeBindingsStatic(UUserWidget* UserWidget, const TArray< FDelegateRuntimeBinding >& InBindings);
-
 	static void InitializeWidgetStatic(UUserWidget* UserWidget
 		, const UClass* InClass
 		, UWidgetTree* InWidgetTree
@@ -127,6 +134,30 @@ public:
 	void SetClassRequiresNativeTick(bool InClassRequiresNativeTick);
 #endif
 
+	/** Find the first extension of the requested type. */
+	template<typename ExtensionType>
+	ExtensionType* GetExtension()
+	{
+		return Cast<ExtensionType>(GetExtension(ExtensionType::StaticClass()));
+	}
+
+	/** Find the first extension of the requested type. */
+	UWidgetBlueprintGeneratedClassExtension* GetExtension(TSubclassOf<UWidgetBlueprintGeneratedClassExtension> InExtensionType);
+
+	/** Find the extensions of the requested type. */
+	TArray<UWidgetBlueprintGeneratedClassExtension*> GetExtensions(TSubclassOf<UWidgetBlueprintGeneratedClassExtension> InExtensionType);
+
+	template<typename Predicate>
+	void ForEachExtension(Predicate Pred) const
+	{
+		for (UWidgetBlueprintGeneratedClassExtension* Extension : Extensions)
+		{
+			check(Extension);
+			Pred(Extension);
+		}
+	}
+
 private:
+	static void InitializeBindingsStatic(UUserWidget* UserWidget, const TArray< FDelegateRuntimeBinding >& InBindings);
 	static void BindAnimations(UUserWidget* Instance, const TArray< UWidgetAnimation* >& InAnimations);
 };
diff --git a/Engine/Source/Runtime/UMG/Public/Extensions/WidgetBlueprintGeneratedClassExtension.h b/Engine/Source/Runtime/UMG/Public/Extensions/WidgetBlueprintGeneratedClassExtension.h
new file mode 100644
--- /dev/null
+++ b/Engine/Source/Runtime/UMG/Public/Extensions/WidgetBlueprintGeneratedClassExtension.h
@@ -0,0 +1,28 @@
+// Copyright Epic Games, Inc. All Rights Reserved.
+
+#pragma once
+
+#include "CoreMinimal.h"
+#include "Components/Widget.h"
+
+#include "WidgetBlueprintGeneratedClassExtension.generated.h"
+
+
+/**
+ * WidgetExtension is the base class for components that define reusable behavior that can be added to different types of Widgets.
+ */
+UCLASS(Abstract, DefaultToInstanced)
+class UMG_API UWidgetBlueprintGeneratedClassExtension : public UObject
+{
+	GENERATED_BODY()
+	
+public:
+	/** Extend the UUserWidget::Initialize function */
+	virtual void Initialize(UUserWidget* UserWidget) { }
+	/** Extend the UUserWidget::PreConstruct function */
+	virtual void PreConstruct(UUserWidget* UserWidget, bool IsDesignTime) { }
+	/** Extend the UUserWidget::Construct function */
+	virtual void Construct(UUserWidget* UserWidget) { }
+	/** Extend the UUserWidget::Destruct function */
+	virtual void Destruct(UUserWidget* UserWidget) { }
+};
