Fix zoom not working well at small sizes

diff --git a/Engine/Source/Editor/UMGEditor/Private/Designer/SDesignSurface.cpp b/Engine/Source/Editor/UMGEditor/Private/Designer/SDesignSurface.cpp
--- a/Engine/Source/Editor/UMGEditor/Private/Designer/SDesignSurface.cpp
+++ b/Engine/Source/Editor/UMGEditor/Private/Designer/SDesignSurface.cpp
@@ -398,15 +398,7 @@ void SDesignSurface::ChangeZoomLevel(int32 ZoomLevelDelta, const FVector2D& Widg
 
 		// Re-center the screen so that it feels like zooming around the cursor.
 		{
-			FSlateRect GraphBounds = ComputeSensibleBounds();
-
-			// Make sure we are not zooming into/out into emptiness; otherwise the user will get lost..
-			const FVector2D ClampedPointToMaintainGraphSpace(
-				FMath::Clamp(PointToMaintainGraphSpace.X, GraphBounds.Left, GraphBounds.Right),
-				FMath::Clamp(PointToMaintainGraphSpace.Y, GraphBounds.Top, GraphBounds.Bottom)
-				);
-
-			const FVector2D NewViewOffset = ClampedPointToMaintainGraphSpace - WidgetSpaceZoomOrigin / GetZoomAmount();
+			const FVector2D NewViewOffset = PointToMaintainGraphSpace - WidgetSpaceZoomOrigin / GetZoomAmount();
 
 			// If we're panning while zooming we need to update the viewoffset start.
 			ViewOffsetStart += (NewViewOffset - ViewOffset);
