Remove Callbacks to rebuild UMG Detail View as they are already handled at the Editor Utility Widget Blueprint level.

diff --git a/Engine/Source/Editor/Blutility/Private/EditorUtilityWidgetBlueprint.cpp b/Engine/Source/Editor/Blutility/Private/EditorUtilityWidgetBlueprint.cpp
--- a/Engine/Source/Editor/Blutility/Private/EditorUtilityWidgetBlueprint.cpp
+++ b/Engine/Source/Editor/Blutility/Private/EditorUtilityWidgetBlueprint.cpp
@@ -110,14 +110,16 @@ void UEditorUtilityWidgetBlueprint::ChangeTabWorld(UWorld* World, EMapChangeType
 {
 	if (MapChangeType == EMapChangeType::TearDownWorld)
 	{
-		if (CreatedTab.IsValid())
-		{
-			CreatedTab.Pin()->SetContent(SNullWidget::NullWidget);
-		}
-		if (CreatedUMGWidget)
+		// We need to Delete the UMG widget if we are tearing down the World it was built with.
+		if (CreatedUMGWidget && World == CreatedUMGWidget->GetWorld())
 		{
+			if (CreatedTab.IsValid())
+			{
+				CreatedTab.Pin()->SetContent(SNullWidget::NullWidget);
+			}
+			
 			CreatedUMGWidget->Rename(nullptr, GetTransientPackage());
-			CreatedUMGWidget = nullptr;
+			CreatedUMGWidget = nullptr;			
 		}
 	}
 	else if (MapChangeType != EMapChangeType::SaveMap)
diff --git a/Engine/Source/Editor/UMGEditor/Private/Components/PropertyViewBase.cpp b/Engine/Source/Editor/UMGEditor/Private/Components/PropertyViewBase.cpp
--- a/Engine/Source/Editor/UMGEditor/Private/Components/PropertyViewBase.cpp
+++ b/Engine/Source/Editor/UMGEditor/Private/Components/PropertyViewBase.cpp
@@ -18,14 +18,6 @@
 void UPropertyViewBase::ReleaseSlateResources(bool bReleaseChildren)
 {
 	Super::ReleaseSlateResources(bReleaseChildren);
-
-	FCoreUObjectDelegates::OnAssetLoaded.Remove(AssetLoadedHandle);
-	AssetLoadedHandle.Reset();
-	FCoreUObjectDelegates::PostLoadMapWithWorld.Remove(PostLoadMapHandle);
-	PostLoadMapHandle.Reset();
-	FEditorDelegates::MapChange.Remove(MapChangeHandle);
-	MapChangeHandle.Reset();
-
 	DisplayedWidget.Reset();
 }
 
@@ -40,19 +32,6 @@ TSharedRef<SWidget> UPropertyViewBase::RebuildWidget()
 	
 	BuildContentWidget();
 
-	if (!AssetLoadedHandle.IsValid())
-	{
-		AssetLoadedHandle = FCoreUObjectDelegates::OnAssetLoaded.AddUObject(this, &UPropertyViewBase::InternalOnAssetLoaded);
-	}
-	if (!PostLoadMapHandle.IsValid())
-	{
-		PostLoadMapHandle = FCoreUObjectDelegates::PostLoadMapWithWorld.AddUObject(this, &UPropertyViewBase::InternalPostLoadMapWithWorld);
-	}
-	if (!MapChangeHandle.IsValid())
-	{
-		MapChangeHandle = FEditorDelegates::MapChange.AddUObject(this, &UPropertyViewBase::InternalOnMapChange);
-	}
-
 	return DisplayedWidget.ToSharedRef();
 }
 
@@ -79,27 +58,6 @@ void UPropertyViewBase::OnPropertyChangedBroadcast(FName PropertyName)
 }
 
 
-void UPropertyViewBase::InternalOnAssetLoaded(UObject* AssetLoaded)
-{
-	if(Object.ToSoftObjectPath().GetAssetPathName() == FSoftObjectPath(AssetLoaded).GetAssetPathName())
-	{
-		BuildContentWidget();
-	}
-}
-
-
-void UPropertyViewBase::InternalPostLoadMapWithWorld(UWorld* AssetLoaded)
-{
-	InternalOnMapChange(0);
-}
-
-
-void UPropertyViewBase::InternalOnMapChange(uint32)
-{
-	BuildContentWidget();
-}
-
-
 void UPropertyViewBase::PostLoad()
 {
 	Super::PostLoad();
diff --git a/Engine/Source/Editor/UMGEditor/Public/Components/PropertyViewBase.h b/Engine/Source/Editor/UMGEditor/Public/Components/PropertyViewBase.h
--- a/Engine/Source/Editor/UMGEditor/Public/Components/PropertyViewBase.h
+++ b/Engine/Source/Editor/UMGEditor/Public/Components/PropertyViewBase.h
@@ -54,9 +54,6 @@ protected:
 	void OnPropertyChangedBroadcast(FName PropertyName);
 
 private:
-	void InternalOnAssetLoaded(UObject* LoadedAsset);
-	void InternalOnMapChange(uint32);
-	void InternalPostLoadMapWithWorld(UWorld* LoadedWorld);
 
 	//~ UWidget interface
 public:
@@ -75,7 +72,4 @@ public:
 
 private:
 	TSharedPtr<SBorder> DisplayedWidget;
-	FDelegateHandle AssetLoadedHandle;
-	FDelegateHandle PostLoadMapHandle;
-	FDelegateHandle MapChangeHandle;
 };
