Add drag drop widget animations

diff --git a/Engine/Source/Editor/UMGEditor/Private/TabFactory/AnimationTabSummoner.cpp b/Engine/Source/Editor/UMGEditor/Private/TabFactory/AnimationTabSummoner.cpp
--- a/Engine/Source/Editor/UMGEditor/Private/TabFactory/AnimationTabSummoner.cpp
+++ b/Engine/Source/Editor/UMGEditor/Private/TabFactory/AnimationTabSummoner.cpp
@@ -23,6 +23,8 @@
 #include "UMGEditorActions.h"
 #include "UMGStyle.h"
 #include "Widgets/Input/SSearchBox.h"
+#include "Input/DragAndDrop.h"
+#include "DragAndDrop/DecoratedDragDropOp.h"
 #include "Widgets/Text/SInlineEditableTextBlock.h"
 #include "Framework/Commands/GenericCommands.h"
 #include "ScopedTransaction.h"
@@ -133,6 +135,29 @@ struct FWidgetAnimationListItem
 	bool bNewAnimation;
 };
 
+/**
+ * This drag drop operation allows us to move around animations in the widget tree
+ */
+class FWidgetAnimationDragDropOp : public FDecoratedDragDropOp
+{
+public:
+	DRAG_DROP_OPERATOR_TYPE(FWidgetAnimationDragDropOp, FDecoratedDragDropOp)
+
+	/** The template to create an instance */
+	TSharedPtr<FWidgetAnimationListItem> ListItem;
+
+	/** Constructs the drag drop operation */
+	static TSharedRef<FWidgetAnimationDragDropOp> New(const TSharedPtr<FWidgetAnimationListItem>& InListItem, FText InDragText)
+	{
+		TSharedRef<FWidgetAnimationDragDropOp> Operation = MakeShared<FWidgetAnimationDragDropOp>();
+		Operation->ListItem = InListItem;
+		Operation->DefaultHoverText = InDragText;
+		Operation->CurrentHoverText = InDragText;
+		Operation->Construct();
+
+		return Operation;
+	}
+};
 
 typedef SListView<TSharedPtr<FWidgetAnimationListItem> > SWidgetAnimationListView;
 
@@ -149,7 +174,10 @@ public:
 
 		STableRow<TSharedPtr<FWidgetAnimationListItem>>::Construct(
 			STableRow<TSharedPtr<FWidgetAnimationListItem>>::FArguments()
-			.Padding( FMargin( 3.0f, 2.0f) )
+			.Padding(FMargin(3.0f, 2.0f))
+			.OnDragDetected(this, &SWidgetAnimationListItem::OnDragDetected)
+			.OnCanAcceptDrop(this, &SWidgetAnimationListItem::OnCanAcceptDrop)
+			.OnAcceptDrop(this, &SWidgetAnimationListItem::OnAcceptDrop)
 			.Content()
 			[
 				SAssignNew(InlineTextBlock, SInlineEditableTextBlock)
@@ -250,6 +278,68 @@ private:
 			}
 		}
 	}
+
+	/** Called whenever a drag is detected by the tree view. */
+	FReply OnDragDetected(const FGeometry& InGeometry, const FPointerEvent& InPointerEvent)
+	{
+		TSharedPtr<FWidgetAnimationListItem> ListItemPinned = ListItem.Pin();
+		if (ListItemPinned.IsValid())
+		{
+			FText DefaultText = LOCTEXT("DefaultDragDropFormat", "Move 1 item(s)");
+			return FReply::Handled().BeginDragDrop(FWidgetAnimationDragDropOp::New(ListItemPinned, DefaultText));
+		}
+		return FReply::Unhandled();
+	}
+
+	/** Called to determine whether a current drag operation is valid for this row. */
+	TOptional<EItemDropZone> OnCanAcceptDrop(const FDragDropEvent& DragDropEvent, EItemDropZone InItemDropZone, TSharedPtr<FWidgetAnimationListItem> InListItem)
+	{
+		TSharedPtr<FWidgetAnimationDragDropOp> DragDropOp = DragDropEvent.GetOperationAs<FWidgetAnimationDragDropOp>();
+		if (DragDropOp.IsValid())
+		{
+			if (InItemDropZone == EItemDropZone::OntoItem)
+			{
+				DragDropOp->CurrentIconBrush = FEditorStyle::GetBrush(TEXT("Graph.ConnectorFeedback.Error"));
+			}
+			else
+			{
+				DragDropOp->CurrentIconBrush = FEditorStyle::GetBrush(TEXT("Graph.ConnectorFeedback.Ok"));
+			}
+			return InItemDropZone;
+		}
+		return TOptional<EItemDropZone>();
+	}
+
+	/** Called to complete a drag and drop onto this drop. */
+	FReply OnAcceptDrop(const FDragDropEvent& DragDropEvent, EItemDropZone InItemDropZone, TSharedPtr<FWidgetAnimationListItem> InListItem)
+	{
+		TSharedPtr<FWidgetAnimationDragDropOp> DragDropOp = DragDropEvent.GetOperationAs<FWidgetAnimationDragDropOp>();
+		if (DragDropOp.IsValid() 
+			&& DragDropOp->ListItem.IsValid() && DragDropOp->ListItem->Animation
+			&& InListItem.IsValid() && InListItem->Animation
+			&& DragDropOp->ListItem->Animation != InListItem->Animation)
+		{
+			if (UWidgetBlueprint* Blueprint = BlueprintEditor.Pin()->GetWidgetBlueprintObj())
+			{
+				Blueprint->Modify();
+				Blueprint->Animations.Remove(DragDropOp->ListItem->Animation);
+
+				int32 RelativeNewIndex = Blueprint->Animations.IndexOfByKey(InListItem->Animation);
+				RelativeNewIndex += InItemDropZone == EItemDropZone::BelowItem ? 1 : 0;
+
+				Blueprint->Animations.Insert(DragDropOp->ListItem->Animation, RelativeNewIndex);
+				FBlueprintEditorUtils::MarkBlueprintAsStructurallyModified(Blueprint);
+
+				if (TSharedPtr<FWidgetBlueprintEditor> WidgetBlueprintEditorPin = BlueprintEditor.Pin())
+				{
+					WidgetBlueprintEditorPin->NotifyWidgetAnimListChanged();
+				}
+				return FReply::Handled();
+			}
+		}
+		return FReply::Unhandled();
+	}
+
 private:
 	TWeakPtr<FWidgetAnimationListItem> ListItem;
 	TWeakPtr<FWidgetBlueprintEditor> BlueprintEditor;
