Integrate fix from UE 5.0.
- Remove Callbacks to rebuild UMG Detail View as they are already handled at the Editor Utility Widget Blueprint level.
- Fix a crash when we force delete a Editor Utility Widget when it is running

diff --git a/Engine/Source/Editor/UMGEditor/Public/Components/PropertyViewBase.h b/Engine/Source/Editor/UMGEditor/Public/Components/PropertyViewBase.h
--- a/Engine/Source/Editor/UMGEditor/Public/Components/PropertyViewBase.h
+++ b/Engine/Source/Editor/UMGEditor/Public/Components/PropertyViewBase.h
@@ -53,8 +53,6 @@ protected:
 	TSharedPtr<SBorder> GetDisplayWidget() const { return DisplayedWidget; }
 	void OnPropertyChangedBroadcast(FName PropertyName);
 
-private:
-
 	//~ UWidget interface
 public:
 	virtual void ReleaseSlateResources(bool bReleaseChildren) override;
