Add "Expand All", "Collapse All" to Widget Hierarchy context menu to match behavior we already have in property views

diff --git a/Engine/Source/Editor/UMGEditor/Private/Hierarchy/SHierarchyView.cpp b/Engine/Source/Editor/UMGEditor/Private/Hierarchy/SHierarchyView.cpp
--- a/Engine/Source/Editor/UMGEditor/Private/Hierarchy/SHierarchyView.cpp
+++ b/Engine/Source/Editor/UMGEditor/Private/Hierarchy/SHierarchyView.cpp
@@ -238,6 +238,22 @@ TSharedPtr<SWidget> SHierarchyView::WidgetHierarchy_OnContextMenuOpening()
 
 	FWidgetBlueprintEditorUtils::CreateWidgetContextMenu(MenuBuilder, BlueprintEditor.Pin().ToSharedRef(), FVector2D(0, 0));
 
+	MenuBuilder.BeginSection("Expansion", LOCTEXT("Expansion", "Expansion"));
+
+	MenuBuilder.AddMenuEntry(
+		LOCTEXT( "CollapseAll_Label", "Collapse All" ),
+		LOCTEXT( "CollapseAll_Tooltip", "Collapses this item and all children" ),
+		FSlateIcon(),
+		FUIAction(FExecuteAction::CreateSP(this, &SHierarchyView::SetItemExpansionRecursive_SelectedItems, false))
+	);
+	
+	MenuBuilder.AddMenuEntry(
+		LOCTEXT( "ExpandAll_Label", "Expand All" ),
+		LOCTEXT( "ExpandAll_Tooltip", "Expands this item and all children" ),
+		FSlateIcon(),
+		FUIAction(FExecuteAction::CreateSP(this, &SHierarchyView::SetItemExpansionRecursive_SelectedItems, true))
+	);
+
 	return MenuBuilder.MakeWidget();
 }
 
@@ -472,6 +488,14 @@ void SHierarchyView::SetItemExpansionRecursive(TSharedPtr<FHierarchyModel> Model
 	}
 }
 
+void SHierarchyView::SetItemExpansionRecursive_SelectedItems(const bool bInExpansionState)
+{
+	for (const TSharedPtr<FHierarchyModel>& Item : WidgetTreeView->GetSelectedItems())
+	{
+		SetItemExpansionRecursive(Item, bInExpansionState);
+	}
+}
+
 //@TODO UMG Drop widgets onto the tree, when nothing is present, if there is a root node present, what happens then, let the root node attempt to place it?
 
 #undef LOCTEXT_NAMESPACE
diff --git a/Engine/Source/Editor/UMGEditor/Private/Hierarchy/SHierarchyView.h b/Engine/Source/Editor/UMGEditor/Private/Hierarchy/SHierarchyView.h
--- a/Engine/Source/Editor/UMGEditor/Private/Hierarchy/SHierarchyView.h
+++ b/Engine/Source/Editor/UMGEditor/Private/Hierarchy/SHierarchyView.h
@@ -122,6 +122,8 @@ private:
 
 	/** Handler for recursively expanding/collapsing items */
 	void SetItemExpansionRecursive(TSharedPtr<FHierarchyModel> Model, bool bInExpansionState);
+	
+	void SetItemExpansionRecursive_SelectedItems(const bool bInExpansionState);
 
 private:
 
