PR #8650: Allow selecting Abstract Classes as base for UserWidget (Contributed by [REDACTED_PERSON])

diff --git a/Engine/Source/Editor/UMGEditor/Private/WidgetBlueprintFactory.cpp b/Engine/Source/Editor/UMGEditor/Private/WidgetBlueprintFactory.cpp
--- a/Engine/Source/Editor/UMGEditor/Private/WidgetBlueprintFactory.cpp
+++ b/Engine/Source/Editor/UMGEditor/Private/WidgetBlueprintFactory.cpp
@@ -73,7 +73,7 @@ bool UWidgetBlueprintFactory::ConfigureProperties()
 		Options.ClassFilters.Add(Filter.ToSharedRef());
 		Options.ExtraPickerCommonClasses.Add(UUserWidget::StaticClass());
 
-		Filter->DisallowedClassFlags = CLASS_Abstract | CLASS_Deprecated | CLASS_NewerVersionExists;
+		Filter->DisallowedClassFlags = CLASS_Deprecated | CLASS_NewerVersionExists;
 		Filter->AllowedChildrenOfClasses.Add(UUserWidget::StaticClass());
 
 		const FText TitleText = LOCTEXT("CreateWidgetBlueprint", "Pick Root Widget for New Widget Blueprint");
