Remove the creation of a Canvas at the root of every Widget Blueprint created.

diff --git a/Engine/Source/Editor/UMGEditor/Private/Settings/UMGEditorProjectSettings.cpp b/Engine/Source/Editor/UMGEditor/Private/Settings/UMGEditorProjectSettings.cpp
--- a/Engine/Source/Editor/UMGEditor/Private/Settings/UMGEditorProjectSettings.cpp
+++ b/Engine/Source/Editor/UMGEditor/Private/Settings/UMGEditorProjectSettings.cpp
@@ -17,7 +17,7 @@ UUMGEditorProjectSettings::UUMGEditorProjectSettings()
 	bShowWidgetsFromDeveloperContent = true;
 
 	bUseWidgetTemplateSelector = false;
-	DefaultRootWidget = UCanvasPanel::StaticClass();
+	DefaultRootWidget = nullptr;
 	DefaultWidgetParentClass = UUserWidget::StaticClass();
 }
 
