#UE Make FDirectoryWidgetCompilerOptions::IgnoredWidgets editor-only to prevent them from cooking

diff --git a/Engine/Source/Editor/UMGEditor/Public/UMGEditorProjectSettings.h b/Engine/Source/Editor/UMGEditor/Public/UMGEditorProjectSettings.h
--- a/Engine/Source/Editor/UMGEditor/Public/UMGEditorProjectSettings.h
+++ b/Engine/Source/Editor/UMGEditor/Public/UMGEditorProjectSettings.h
@@ -103,9 +103,11 @@ public:
 	UPROPERTY(EditAnywhere, Category = Compiler, meta = (ContentDir))
 	FDirectoryPath Directory;
 
+#if WITH_EDITORONLY_DATA
 	/** These widgets are ignored, and they will use the next most applicable directory to determine their rules. */
 	UPROPERTY(EditAnywhere, Category = Compiler)
 	TArray<TSoftObjectPtr<UWidgetBlueprint>> IgnoredWidgets;
+#endif // WITH_EDITORONLY_DATA
 
 	/** The directory specific compiler options for these widgets. */
 	UPROPERTY(EditAnywhere, Category = Compiler)
