Fix Anim Drawer dismissing changes and tab if sequencer not selected initially.

diff --git a/Engine/Source/Editor/UMGEditor/Private/WidgetBlueprintEditor.cpp b/Engine/Source/Editor/UMGEditor/Private/WidgetBlueprintEditor.cpp
--- a/Engine/Source/Editor/UMGEditor/Private/WidgetBlueprintEditor.cpp
+++ b/Engine/Source/Editor/UMGEditor/Private/WidgetBlueprintEditor.cpp
@@ -1140,6 +1140,11 @@ void FWidgetBlueprintEditor::OnWidgetAnimSequencerOpened(FName StatusBarWithDraw
 			Sequencer->RefreshTree();
 		}
 	}
+
+	if (DrawerSequencer)
+	{
+		FSlateApplication::Get().SetUserFocus(FSlateApplication::Get().GetUserIndexForKeyboard(), DrawerSequencer->GetSequencerWidget());
+	}
 }
 
 void FWidgetBlueprintEditor::OnWidgetAnimSequencerDismissed(const TSharedPtr<SWidget>& NewlyFocusedWidget)
@@ -1163,7 +1168,7 @@ void FWidgetBlueprintEditor::OnWidgetAnimSequencerDismissed(const TSharedPtr<SWi
 		}
 	}
 
-	FocusWindow();
+	SetKeyboardFocus();
 }
 
 UWidgetBlueprint* FWidgetBlueprintEditor::GetWidgetBlueprintObj() const
