Added extension point for UMG Bind menu to enable MVVM work.

diff --git a/Engine/Source/Editor/UMGEditor/Private/UMGEditorModule.cpp b/Engine/Source/Editor/UMGEditor/Private/UMGEditorModule.cpp
--- a/Engine/Source/Editor/UMGEditor/Private/UMGEditorModule.cpp
+++ b/Engine/Source/Editor/UMGEditor/Private/UMGEditorModule.cpp
@@ -1,43 +1,36 @@
 // Copyright Epic Games, Inc. All Rights Reserved.
 
 #include "UMGEditorModule.h"
-#include "Modules/ModuleManager.h"
-#include "UObject/UObjectHash.h"
-#include "Editor.h"
 
-#include "Blueprint/UserWidget.h"
-#include "Blueprint/WidgetBlueprintGeneratedClass.h"
-#include "Settings/WidgetDesignerSettings.h"
-#include "WidgetBlueprint.h"
+#include "Animation/MarginTrackEditor.h"
+#include "Animation/MovieSceneSequenceEditor_WidgetAnimation.h"
+#include "Animation/Sequencer2DTransformTrackEditor.h"
 #include "Animation/WidgetAnimation.h"
-
+#include "Animation/WidgetMaterialTrackEditor.h"
 #include "AssetToolsModule.h"
-#include "IAssetTypeActions.h"
 #include "AssetTypeActions_SlateVectorArtData.h"
 #include "AssetTypeActions_WidgetBlueprint.h"
 #include "AssetTypeActions_WidgetBlueprintGeneratedClass.h"
-#include "KismetCompilerModule.h"
-#include "WidgetBlueprintCompiler.h"
-
+#include "Blueprint/UserWidget.h"
+#include "Blueprint/WidgetBlueprintGeneratedClass.h"
+#include "BlueprintEditorModule.h"
+#include "Designer/DesignerCommands.h"
+#include "DynamicEntryBoxDetails.h"
+#include "Editor.h"
+#include "IAssetTypeActions.h"
 #include "ISequencerModule.h"
-#include "Animation/MarginTrackEditor.h"
-#include "Animation/Sequencer2DTransformTrackEditor.h"
-#include "Animation/WidgetMaterialTrackEditor.h"
-#include "Animation/MovieSceneSequenceEditor_WidgetAnimation.h"
+#include "ISettingsModule.h"
 #include "IUMGModule.h"
-#include "Designer/DesignerCommands.h"
+#include "KismetCompilerModule.h"
+#include "ListViewBaseDetails.h"
+#include "Modules/ModuleManager.h"
 #include "Navigation/SWidgetDesignerNavigation.h"
-
-#include "ClassIconFinder.h"
-
-#include "UMGEditorProjectSettings.h"
-#include "ISettingsModule.h"
-#include "SequencerSettings.h"
-
-#include "BlueprintEditorModule.h"
 #include "PropertyEditorModule.h"
-#include "DynamicEntryBoxDetails.h"
-#include "ListViewBaseDetails.h"
+#include "SequencerSettings.h"
+#include "UMGEditorProjectSettings.h"
+#include "UObject/UObjectHash.h"
+#include "WidgetBlueprint.h"
+#include "WidgetBlueprintCompiler.h"
 #include "WidgetBlueprintThumbnailRenderer.h"
 #include "WidgetThumbnailCustomization.h"
 
@@ -80,6 +73,8 @@
 
 		DesignerExtensibilityManager->AddDesignerExtensionFactory(SWidgetDesignerNavigation::MakeDesignerExtension());
 
+		PropertyBindingExtensibilityManager = MakeShared<FPropertyBindingExtensibilityManager>();
+
 		// Register widget blueprint compiler we do this no matter what.
 		IKismetCompilerInterface& KismetCompilerModule = FModuleManager::LoadModuleChecked<IKismetCompilerInterface>("KismetCompiler");
 		KismetCompilerModule.GetCompilers().Add(&WidgetBlueprintCompiler);
@@ -171,6 +166,7 @@
 	virtual TSharedPtr<FExtensibilityManager> GetMenuExtensibilityManager() override { return MenuExtensibilityManager; }
 	virtual TSharedPtr<FExtensibilityManager> GetToolBarExtensibilityManager() override { return ToolBarExtensibilityManager; }
 	virtual TSharedPtr<FDesignerExtensibilityManager> GetDesignerExtensibilityManager() override { return DesignerExtensibilityManager; }
+	virtual TSharedPtr<FPropertyBindingExtensibilityManager> GetPropertyBindingExtensibilityManager() override { return PropertyBindingExtensibilityManager; }
 
 	/** Register settings objects. */
 	void RegisterSettings()
@@ -264,6 +260,7 @@
 	TSharedPtr<FExtensibilityManager> MenuExtensibilityManager;
 	TSharedPtr<FExtensibilityManager> ToolBarExtensibilityManager;
 	TSharedPtr<FDesignerExtensibilityManager> DesignerExtensibilityManager;
+	TSharedPtr<FPropertyBindingExtensibilityManager> PropertyBindingExtensibilityManager;
 
 	FDelegateHandle SequenceEditorHandle;
 	FDelegateHandle MarginTrackEditorCreateTrackEditorHandle;
diff --git a/Engine/Source/Editor/UMGEditor/Private/WidgetBlueprintExtension.cpp b/Engine/Source/Editor/UMGEditor/Private/WidgetBlueprintExtension.cpp
--- a/Engine/Source/Editor/UMGEditor/Private/WidgetBlueprintExtension.cpp
+++ b/Engine/Source/Editor/UMGEditor/Private/WidgetBlueprintExtension.cpp
@@ -23,7 +23,7 @@ UWidgetBlueprintExtension* UWidgetBlueprintExtension::RequestExtension(UWidgetBl
 }
 
 
-UWidgetBlueprintExtension* UWidgetBlueprintExtension::GetExtension(UWidgetBlueprint* InBlueprint, TSubclassOf<UWidgetBlueprintExtension> InExtensionType)
+UWidgetBlueprintExtension* UWidgetBlueprintExtension::GetExtension(const UWidgetBlueprint* InBlueprint, TSubclassOf<UWidgetBlueprintExtension> InExtensionType)
 {
 	// Look for an existing extension
 	for (UBlueprintExtension* Extension : InBlueprint->Extensions)
@@ -38,7 +38,7 @@ UWidgetBlueprintExtension* UWidgetBlueprintExtension::GetExtension(UWidgetBluepr
 }
 
 
-TArray<UWidgetBlueprintExtension*> UWidgetBlueprintExtension::GetExtensions(UWidgetBlueprint* InBlueprint)
+TArray<UWidgetBlueprintExtension*> UWidgetBlueprintExtension::GetExtensions(const UWidgetBlueprint* InBlueprint)
 {
 	TArray<UWidgetBlueprintExtension*> Extensions;
 
diff --git a/Engine/Source/Editor/UMGEditor/Public/IHasPropertyBindingExtensibility.h b/Engine/Source/Editor/UMGEditor/Public/IHasPropertyBindingExtensibility.h
new file mode 100644
--- /dev/null
+++ b/Engine/Source/Editor/UMGEditor/Public/IHasPropertyBindingExtensibility.h
@@ -0,0 +1,51 @@
+// Copyright Epic Games, Inc. All Rights Reserved.
+
+#pragma once
+
+#include "Components/Widget.h"
+#include "Containers/Array.h"
+#include "Misc/Optional.h"
+#include "Templates/SharedPointer.h"
+#include "WidgetBlueprint.h"
+
+class UMGEDITOR_API IPropertyBindingExtension
+{
+public:
+	/** Does this extension want to extend this property in the widget? */
+	virtual bool CanExtend(const UWidgetBlueprint* WidgetBlueprint, const UWidget* Widget, const FProperty* Property) const = 0;
+	virtual TOptional<FName> GetCurrentValue(const UWidgetBlueprint* WidgetBlueprint, const UWidget* Widget, const FProperty* Property) const = 0;
+	virtual TSharedPtr<FExtender> CreateMenuExtender(const UWidgetBlueprint* WidgetBlueprint, const UWidget* Widget, const FProperty* Property) = 0;
+};
+
+/**
+ * Bindings menu extensibility manager holds a list of registered bindings menu extensions.
+ */
+class UMGEDITOR_API FPropertyBindingExtensibilityManager
+{
+public:
+	void AddExtension(const TSharedRef<IPropertyBindingExtension>& Extension)
+	{
+		Extensions.AddUnique(Extension);
+	}
+
+	void RemoveExtension(const TSharedRef<IPropertyBindingExtension>& Extension)
+	{
+		Extensions.RemoveSingle(Extension);
+	}
+
+	const TArray<TSharedPtr<IPropertyBindingExtension>>& GetExtensions() const
+	{
+		return Extensions;
+	}
+
+private:
+	TArray<TSharedPtr<IPropertyBindingExtension>> Extensions;
+};
+
+/** Indicates that a class has a bindings menu that is extensible */
+class IHasPropertyBindingExtensibility
+{
+public:
+	virtual TSharedPtr<FPropertyBindingExtensibilityManager> GetPropertyBindingExtensibilityManager() = 0;
+};
+
diff --git a/Engine/Source/Editor/UMGEditor/Public/UMGEditorModule.h b/Engine/Source/Editor/UMGEditor/Public/UMGEditorModule.h
--- a/Engine/Source/Editor/UMGEditor/Public/UMGEditorModule.h
+++ b/Engine/Source/Editor/UMGEditor/Public/UMGEditorModule.h
@@ -2,23 +2,29 @@
 
 #pragma once
 
-#include "CoreMinimal.h"
 #include "Modules/ModuleInterface.h"
 #include "Toolkits/AssetEditorToolkit.h"
 #include "IHasDesignerExtensibility.h"
+#include "IHasPropertyBindingExtensibility.h"
 
 extern const FName UMGEditorAppIdentifier;
 
 class FUMGEditor;
-class FWidgetBlueprintCompiler;
 class FWidgetBlueprintApplicationMode;
 class FWorkflowAllowedTabSet;
 
+DECLARE_DELEGATE_RetVal_ThreeParams(TSharedPtr<FExtender>, FOnExtendBindingsMenuForProperty, const class UWidgetBlueprint*, const class UWidget*, const class FProperty*);
+
 /** The public interface of the UMG editor module. */
-class IUMGEditorModule : public IModuleInterface, public IHasMenuExtensibility, public IHasToolBarExtensibility, public IHasDesignerExtensibility
+class IUMGEditorModule : 
+	public IModuleInterface, 
+	public IHasMenuExtensibility, 
+	public IHasToolBarExtensibility, 
+	public IHasDesignerExtensibility,
+	public IHasPropertyBindingExtensibility
 {
 public:
-	virtual FWidgetBlueprintCompiler* GetRegisteredCompiler() = 0;
+	virtual class FWidgetBlueprintCompiler* GetRegisteredCompiler() = 0;
 
 	DECLARE_EVENT_TwoParams(IUMGEditorModule, FOnRegisterTabs, const FWidgetBlueprintApplicationMode&, FWorkflowAllowedTabSet&);
 	virtual FOnRegisterTabs& OnRegisterTabsForEditor() = 0;
diff --git a/Engine/Source/Editor/UMGEditor/Public/WidgetBlueprintExtension.h b/Engine/Source/Editor/UMGEditor/Public/WidgetBlueprintExtension.h
--- a/Engine/Source/Editor/UMGEditor/Public/WidgetBlueprintExtension.h
+++ b/Engine/Source/Editor/UMGEditor/Public/WidgetBlueprintExtension.h
@@ -38,20 +38,20 @@ public:
 
 	/** Get an already-requested extension for an WidgetBlueprint. */
 	template<typename ExtensionType>
-    static ExtensionType* GetExtension(UWidgetBlueprint* InWidgetBlueprint)
+    static ExtensionType* GetExtension(const UWidgetBlueprint* InWidgetBlueprint)
 	{
 		return Cast<ExtensionType>(GetExtension(InWidgetBlueprint, ExtensionType::StaticClass()));
 	}
 
 	/** Get an already-requested extension for an WidgetBlueprint. */
-	static UWidgetBlueprintExtension* GetExtension(UWidgetBlueprint* InWidgetBlueprint, TSubclassOf<UWidgetBlueprintExtension> InExtensionType);
+	static UWidgetBlueprintExtension* GetExtension(const UWidgetBlueprint* InWidgetBlueprint, TSubclassOf<UWidgetBlueprintExtension> InExtensionType);
 
 	/** Get all subsystems currently present on an WidgetBlueprint */
-	static TArray<UWidgetBlueprintExtension*> GetExtensions(UWidgetBlueprint* InWidgetBlueprint);
+	static TArray<UWidgetBlueprintExtension*> GetExtensions(const UWidgetBlueprint* InWidgetBlueprint);
 
 	/** Iterate over all registered WidgetBlueprintExtensions in an WidgetBlueprint */
 	template<typename Predicate>
-	static void ForEachExtension(UWidgetBlueprint* InWidgetBlueprint, Predicate Pred)
+	static void ForEachExtension(const UWidgetBlueprint* InWidgetBlueprint, Predicate Pred)
 	{
 		for (UBlueprintExtension* BlueprintExtension : InWidgetBlueprint->Extensions)
 		{
