Removed unnecessary loads of FPropertyEditorModule from SWidgetDetailsView.cpp.

diff --git a/Engine/Source/Editor/UMGEditor/Private/Details/SWidgetDetailsView.cpp b/Engine/Source/Editor/UMGEditor/Private/Details/SWidgetDetailsView.cpp
--- a/Engine/Source/Editor/UMGEditor/Private/Details/SWidgetDetailsView.cpp
+++ b/Engine/Source/Editor/UMGEditor/Private/Details/SWidgetDetailsView.cpp
@@ -18,7 +18,6 @@
 #include "PropertyEditorModule.h"
 #include "Customizations/SlateBrushCustomization.h"
 #include "Customizations/SlateFontInfoCustomization.h"
-
 #include "Customizations/WidgetTypeCustomization.h"
 #include "Customizations/WidgetNavigationCustomization.h"
 #include "Customizations/CanvasSlotCustomization.h"
@@ -192,9 +191,6 @@ SWidgetDetailsView::~SWidgetDetailsView()
 	}
 
 	// Unregister the property type layouts
-	static FName PropertyEditor("PropertyEditor");
-	FPropertyEditorModule& PropertyModule = FModuleManager::GetModuleChecked<FPropertyEditorModule>(PropertyEditor);
-
 	PropertyView->UnregisterInstancedCustomPropertyTypeLayout(TEXT("Widget"));
 	PropertyView->UnregisterInstancedCustomPropertyTypeLayout(TEXT("WidgetNavigation"));
 	PropertyView->UnregisterInstancedCustomPropertyTypeLayout(TEXT("PanelSlot"));
@@ -209,10 +205,6 @@ SWidgetDetailsView::~SWidgetDetailsView()
 void SWidgetDetailsView::RegisterCustomizations()
 {
 	PropertyView->RegisterInstancedCustomPropertyLayout(UWidget::StaticClass(), FOnGetDetailCustomizationInstance::CreateStatic(&FBlueprintWidgetCustomization::MakeInstance, BlueprintEditor.Pin().ToSharedRef(), BlueprintEditor.Pin()->GetBlueprintObj()));
-
-	static FName PropertyEditor("PropertyEditor");
-	FPropertyEditorModule& PropertyModule = FModuleManager::GetModuleChecked<FPropertyEditorModule>(PropertyEditor);
-
 	PropertyView->RegisterInstancedCustomPropertyTypeLayout(TEXT("Widget"), FOnGetPropertyTypeCustomizationInstance::CreateStatic(&FWidgetTypeCustomization::MakeInstance, BlueprintEditor.Pin().ToSharedRef()), nullptr);
 	PropertyView->RegisterInstancedCustomPropertyTypeLayout(TEXT("WidgetNavigation"), FOnGetPropertyTypeCustomizationInstance::CreateStatic(&FWidgetNavigationCustomization::MakeInstance, BlueprintEditor.Pin().ToSharedRef()));
 	PropertyView->RegisterInstancedCustomPropertyTypeLayout(TEXT("PanelSlot"), FOnGetPropertyTypeCustomizationInstance::CreateStatic(&FCanvasSlotCustomization::MakeInstance, BlueprintEditor.Pin()->GetBlueprintObj()));
