Undid sorting of includes in UMGEditorModule. Removed unnecessary delegate declaration.

diff --git a/Engine/Source/Editor/UMGEditor/Private/UMGEditorModule.cpp b/Engine/Source/Editor/UMGEditor/Private/UMGEditorModule.cpp
--- a/Engine/Source/Editor/UMGEditor/Private/UMGEditorModule.cpp
+++ b/Engine/Source/Editor/UMGEditor/Private/UMGEditorModule.cpp
@@ -1,36 +1,43 @@
 // Copyright Epic Games, Inc. All Rights Reserved.
 
 #include "UMGEditorModule.h"
+#include "Modules/ModuleManager.h"
+#include "UObject/UObjectHash.h"
+#include "Editor.h"
 
-#include "Animation/MarginTrackEditor.h"
-#include "Animation/MovieSceneSequenceEditor_WidgetAnimation.h"
-#include "Animation/Sequencer2DTransformTrackEditor.h"
+#include "Blueprint/UserWidget.h"
+#include "Blueprint/WidgetBlueprintGeneratedClass.h"
+#include "Settings/WidgetDesignerSettings.h"
+#include "WidgetBlueprint.h"
 #include "Animation/WidgetAnimation.h"
-#include "Animation/WidgetMaterialTrackEditor.h"
+
 #include "AssetToolsModule.h"
+#include "IAssetTypeActions.h"
 #include "AssetTypeActions_SlateVectorArtData.h"
 #include "AssetTypeActions_WidgetBlueprint.h"
 #include "AssetTypeActions_WidgetBlueprintGeneratedClass.h"
-#include "Blueprint/UserWidget.h"
-#include "Blueprint/WidgetBlueprintGeneratedClass.h"
-#include "BlueprintEditorModule.h"
-#include "Designer/DesignerCommands.h"
-#include "DynamicEntryBoxDetails.h"
-#include "Editor.h"
-#include "IAssetTypeActions.h"
+#include "KismetCompilerModule.h"
+#include "WidgetBlueprintCompiler.h"
+
 #include "ISequencerModule.h"
-#include "ISettingsModule.h"
+#include "Animation/MarginTrackEditor.h"
+#include "Animation/Sequencer2DTransformTrackEditor.h"
+#include "Animation/WidgetMaterialTrackEditor.h"
+#include "Animation/MovieSceneSequenceEditor_WidgetAnimation.h"
 #include "IUMGModule.h"
-#include "KismetCompilerModule.h"
-#include "ListViewBaseDetails.h"
-#include "Modules/ModuleManager.h"
+#include "Designer/DesignerCommands.h"
 #include "Navigation/SWidgetDesignerNavigation.h"
-#include "PropertyEditorModule.h"
-#include "SequencerSettings.h"
+
+#include "ClassIconFinder.h"
+
 #include "UMGEditorProjectSettings.h"
-#include "UObject/UObjectHash.h"
-#include "WidgetBlueprint.h"
-#include "WidgetBlueprintCompiler.h"
+#include "ISettingsModule.h"
+#include "SequencerSettings.h"
+
+#include "BlueprintEditorModule.h"
+#include "PropertyEditorModule.h"
+#include "DynamicEntryBoxDetails.h"
+#include "ListViewBaseDetails.h"
 #include "WidgetBlueprintThumbnailRenderer.h"
 #include "WidgetThumbnailCustomization.h"
 
diff --git a/Engine/Source/Editor/UMGEditor/Public/UMGEditorModule.h b/Engine/Source/Editor/UMGEditor/Public/UMGEditorModule.h
--- a/Engine/Source/Editor/UMGEditor/Public/UMGEditorModule.h
+++ b/Engine/Source/Editor/UMGEditor/Public/UMGEditorModule.h
@@ -13,8 +13,6 @@ class FUMGEditor;
 class FWidgetBlueprintApplicationMode;
 class FWorkflowAllowedTabSet;
 
-DECLARE_DELEGATE_RetVal_ThreeParams(TSharedPtr<FExtender>, FOnExtendBindingsMenuForProperty, const class UWidgetBlueprint*, const class UWidget*, const class FProperty*);
-
 /** The public interface of the UMG editor module. */
 class IUMGEditorModule : 
 	public IModuleInterface, 
