UMG Fill Size property is no longer constrained to be only 45 pixels wide, but rather expands to fit its contents, allowing up to 6 digits to be displayed.

diff --git a/Engine/Source/Editor/UMGEditor/Private/Customizations/SlateChildSizeCustomization.cpp b/Engine/Source/Editor/UMGEditor/Private/Customizations/SlateChildSizeCustomization.cpp
--- a/Engine/Source/Editor/UMGEditor/Private/Customizations/SlateChildSizeCustomization.cpp
+++ b/Engine/Source/Editor/UMGEditor/Private/Customizations/SlateChildSizeCustomization.cpp
@@ -16,21 +16,19 @@ void FSlateChildSizeCustomization::CustomizeHeader(TSharedRef<IPropertyHandle> P
 	TSharedPtr<IPropertyHandle> ValueHandle = PropertyHandle->GetChildHandle(GET_MEMBER_NAME_CHECKED(FSlateChildSize, Value));
 	TSharedPtr<IPropertyHandle> RuleHandle = PropertyHandle->GetChildHandle(GET_MEMBER_NAME_CHECKED(FSlateChildSize, SizeRule));
 
-	const FMargin OuterPadding(2, 0);
-	const FMargin ContentPadding(4, 2);
-
-	if ( !( ValueHandle.IsValid() || RuleHandle.IsValid() ) )
+	if (!ValueHandle.IsValid() || !RuleHandle.IsValid())
 	{
 		return;
 	}
 
+	const FMargin OuterPadding(2, 0);
+
 	HeaderRow
 	.NameContent()
 	[
 		PropertyHandle->CreatePropertyNameWidget()
 	]
 	.ValueContent()
-	.MaxDesiredWidth(TOptional<float>())
 	[
 		SNew(SHorizontalBox)
 		+ SHorizontalBox::Slot()
@@ -51,16 +49,12 @@ void FSlateChildSizeCustomization::CustomizeHeader(TSharedRef<IPropertyHandle> P
 		.AutoWidth()
 		.Padding(OuterPadding)
 		[
-			SNew(SBox)
-			.WidthOverride(45)
-			[
-				SNew( SNumericEntryBox<float> )
-				.LabelVAlign(VAlign_Center)
-				.Visibility(this, &FSlateChildSizeCustomization::GetValueVisiblity, RuleHandle)
-				.Value(this, &FSlateChildSizeCustomization::GetValue, ValueHandle)
-				.OnValueCommitted(this, &FSlateChildSizeCustomization::HandleValueComitted, ValueHandle)
-				.UndeterminedString( LOCTEXT("MultipleValues", "Multiple Values") )
-			]
+			SNew( SNumericEntryBox<float> )
+			.LabelVAlign(VAlign_Center)
+			.Visibility(this, &FSlateChildSizeCustomization::GetValueVisiblity, RuleHandle)
+			.Value(this, &FSlateChildSizeCustomization::GetValue, ValueHandle)
+			.OnValueCommitted(this, &FSlateChildSizeCustomization::HandleValueComitted, ValueHandle)
+			.UndeterminedString( LOCTEXT("MultipleValues", "Multiple Values") )
 		]
 	];
 }
