Fix to a crash when using drag and drop from the Viewport Windows to the Hierarchy panel.

#rnx

diff --git a/Engine/Source/Editor/UMGEditor/Private/Designer/SDesignerView.cpp b/Engine/Source/Editor/UMGEditor/Private/Designer/SDesignerView.cpp
--- a/Engine/Source/Editor/UMGEditor/Private/Designer/SDesignerView.cpp
+++ b/Engine/Source/Editor/UMGEditor/Private/Designer/SDesignerView.cpp
@@ -36,9 +36,10 @@
 
 #include "Kismet2/BlueprintEditorUtils.h"
 
+#include "DragAndDrop/AssetDragDropOp.h"
 #include "DragAndDrop/DecoratedDragDropOp.h"
 #include "DragDrop/WidgetTemplateDragDropOp.h"
-#include "DragAndDrop/AssetDragDropOp.h"
+#include "DragDrop/SelectedWidgetDragDropOp.h"
 
 #include "Templates/WidgetTemplateBlueprintClass.h"
 #include "Templates/WidgetTemplateImageClass.h"
@@ -187,115 +188,6 @@ public:
 	}
 };
 
-
-class FSelectedWidgetDragDropOp : public FDecoratedDragDropOp
-{
-public:
-	DRAG_DROP_OPERATOR_TYPE(FSelectedWidgetDragDropOp, FDecoratedDragDropOp)
-
-	virtual ~FSelectedWidgetDragDropOp();
-
-	struct FDraggingWidgetReference
-	{
-		FWidgetReference Widget;
-
-		FVector2D DraggedOffset;
-	};
-
-	struct FItem
-	{
-		/** The slot properties for the old slot the widget was in, is used to attempt to reapply the same layout information */
-		TMap<FName, FString> ExportedSlotProperties;
-
-		/** The widget being dragged */
-		UWidget* Template;
-
-		/** The preview widget being dragged */
-		UWidget* Preview;
-
-		/** Can the drag drop change the widget's parent? */
-		bool bStayingInParent;
-
-		/** The original parent of the widget. */
-		FWidgetReference ParentWidget;
-
-		/** The offset of the original click location, as a percentage of the widget's size. */
-		FVector2D DraggedOffset;
-	};
-
-	TArray<FItem> DraggedWidgets;
-
-	bool bShowingMessage;
-
-	IUMGDesigner* Designer;
-
-	static TSharedRef<FSelectedWidgetDragDropOp> New(TSharedPtr<FWidgetBlueprintEditor> Editor, IUMGDesigner* InDesigner, const TArray<FDraggingWidgetReference>& InWidgets);
-};
-
-FSelectedWidgetDragDropOp::~FSelectedWidgetDragDropOp()
-{
-	if ( bShowingMessage )
-	{
-		Designer->PopDesignerMessage();
-	}
-}
-
-TSharedRef<FSelectedWidgetDragDropOp> FSelectedWidgetDragDropOp::New(TSharedPtr<FWidgetBlueprintEditor> Editor, IUMGDesigner* InDesigner, const TArray<FDraggingWidgetReference>& InWidgets)
-{
-	TSharedRef<FSelectedWidgetDragDropOp> Operation = MakeShareable(new FSelectedWidgetDragDropOp());
-	Operation->bShowingMessage = false;
-	Operation->Designer = InDesigner;
-
-	for (const FDraggingWidgetReference& InDraggedWidget : InWidgets)
-	{
-		FItem DraggedWidget;
-		DraggedWidget.bStayingInParent = false;
-
-		if (UPanelWidget* PanelTemplate = InDraggedWidget.Widget.GetTemplate()->GetParent())
-		{
-			DraggedWidget.ParentWidget = Editor->GetReferenceFromTemplate(PanelTemplate);
-			DraggedWidget.bStayingInParent = PanelTemplate->LockToPanelOnDrag() || GetDefault<UWidgetDesignerSettings>()->bLockToPanelOnDragByDefault;
-
-			if ( DraggedWidget.bStayingInParent )
-			{
-				Operation->bShowingMessage = true;
-			}
-		}
-
-		// Cache the preview and template, it's not safe to query the preview/template while dragging the widget as it no longer
-		// exists in the tree.
-		DraggedWidget.Preview = InDraggedWidget.Widget.GetPreview();
-		DraggedWidget.Template = InDraggedWidget.Widget.GetTemplate();
-
-		DraggedWidget.DraggedOffset = InDraggedWidget.DraggedOffset;
-
-		FWidgetBlueprintEditorUtils::ExportPropertiesToText(InDraggedWidget.Widget.GetTemplate()->Slot, DraggedWidget.ExportedSlotProperties);
-
-		Operation->DraggedWidgets.Add(DraggedWidget);
-	}
-
-	// Set the display text based on whether we're dragging a single or multiple widgets
-	if (InWidgets.Num() == 1)
-	{
-		FText DisplayText = InWidgets[0].Widget.GetTemplate()->GetLabelText();
-
-		Operation->DefaultHoverText = DisplayText;
-		Operation->CurrentHoverText = DisplayText;
-	}
-	else
-	{
-		Operation->CurrentHoverText = Operation->DefaultHoverText = LOCTEXT("DragMultipleWidgets", "Multiple Widgets");
-	}
-
-	if ( Operation->bShowingMessage )
-	{
-		InDesigner->PushDesignerMessage(LOCTEXT("PressAltToMoveFromParent", "Press [Alt] to move the widget out of the current parent"));
-	}
-
-	Operation->Construct();
-	return Operation;
-}
-
 //////////////////////////////////////////////////////////////////////////
 
 UWidget* SDesignerView::GetWidgetInDesignScopeFromSlateWidget(TSharedRef<SWidget>& InWidget)
diff --git a/Engine/Source/Editor/UMGEditor/Private/DragDrop/SelectedWidgetDragDropOp.cpp b/Engine/Source/Editor/UMGEditor/Private/DragDrop/SelectedWidgetDragDropOp.cpp
new file mode 100644
--- /dev/null
+++ b/Engine/Source/Editor/UMGEditor/Private/DragDrop/SelectedWidgetDragDropOp.cpp
@@ -0,0 +1,75 @@
+// Copyright Epic Games, Inc. All Rights Reserved.
+
+#include "DragDrop/SelectedWidgetDragDropOp.h"
+#include "WidgetTemplate.h"
+#include "WidgetBlueprintEditor.h"
+#include "Components/PanelWidget.h"
+#include "WidgetBlueprintEditorUtils.h"
+
+#define LOCTEXT_NAMESPACE "UMG"
+
+FSelectedWidgetDragDropOp::~FSelectedWidgetDragDropOp()
+{
+	if (bShowingMessage)
+	{
+		Designer->PopDesignerMessage();
+	}
+}
+
+TSharedRef<FSelectedWidgetDragDropOp> FSelectedWidgetDragDropOp::New(TSharedPtr<FWidgetBlueprintEditor> Editor, IUMGDesigner* InDesigner, const TArray<FDraggingWidgetReference>& InWidgets)
+{
+	TSharedRef<FSelectedWidgetDragDropOp> Operation = MakeShareable(new FSelectedWidgetDragDropOp());
+	Operation->bShowingMessage = false;
+	Operation->Designer = InDesigner;
+
+	for (const FDraggingWidgetReference& InDraggedWidget : InWidgets)
+	{
+		FItem DraggedWidget;
+		DraggedWidget.bStayingInParent = false;
+
+		if (UPanelWidget* PanelTemplate = InDraggedWidget.Widget.GetTemplate()->GetParent())
+		{
+			DraggedWidget.ParentWidget = Editor->GetReferenceFromTemplate(PanelTemplate);
+			DraggedWidget.bStayingInParent = PanelTemplate->LockToPanelOnDrag() || GetDefault<UWidgetDesignerSettings>()->bLockToPanelOnDragByDefault;
+
+			if (DraggedWidget.bStayingInParent)
+			{
+				Operation->bShowingMessage = true;
+			}
+		}
+
+		// Cache the preview and template, it's not safe to query the preview/template while dragging the widget as it no longer
+		// exists in the tree.
+		DraggedWidget.Preview = InDraggedWidget.Widget.GetPreview();
+		DraggedWidget.Template = InDraggedWidget.Widget.GetTemplate();
+
+		DraggedWidget.DraggedOffset = InDraggedWidget.DraggedOffset;
+
+		FWidgetBlueprintEditorUtils::ExportPropertiesToText(InDraggedWidget.Widget.GetTemplate()->Slot, DraggedWidget.ExportedSlotProperties);
+
+		Operation->DraggedWidgets.Add(DraggedWidget);
+	}
+
+	// Set the display text based on whether we're dragging a single or multiple widgets
+	if (InWidgets.Num() == 1)
+	{
+		FText DisplayText = InWidgets[0].Widget.GetTemplate()->GetLabelText();
+
+		Operation->DefaultHoverText = DisplayText;
+		Operation->CurrentHoverText = DisplayText;
+	}
+	else
+	{
+		Operation->CurrentHoverText = Operation->DefaultHoverText = LOCTEXT("DragMultipleWidgets", "Multiple Widgets");
+	}
+
+	if (Operation->bShowingMessage)
+	{
+		InDesigner->PushDesignerMessage(LOCTEXT("PressAltToMoveFromParent", "Press [Alt] to move the widget out of the current parent"));
+	}
+
+	Operation->Construct();
+	return Operation;
+}
+
+#undef LOCTEXT_NAMESPACE
diff --git a/Engine/Source/Editor/UMGEditor/Private/DragDrop/SelectedWidgetDragDropOp.h b/Engine/Source/Editor/UMGEditor/Private/DragDrop/SelectedWidgetDragDropOp.h
new file mode 100644
--- /dev/null
+++ b/Engine/Source/Editor/UMGEditor/Private/DragDrop/SelectedWidgetDragDropOp.h
@@ -0,0 +1,60 @@
+// Copyright Epic Games, Inc. All Rights Reserved.
+
+#pragma once
+
+#include "CoreMinimal.h"
+#include "IUMGDesigner.h"
+#include "Input/DragAndDrop.h"
+#include "DragAndDrop/DecoratedDragDropOp.h"
+#include "WidgetReference.h"
+
+class FWidgetTemplate;
+class FWidgetBlueprintEditor;
+class UWidget;
+
+/**
+ * This drag drop operation allows widget Selected in the viewport to be dragged and dropped into the designer.
+ */
+class FSelectedWidgetDragDropOp : public FDecoratedDragDropOp
+{
+public:
+	DRAG_DROP_OPERATOR_TYPE(FSelectedWidgetDragDropOp, FDecoratedDragDropOp);
+
+	virtual ~FSelectedWidgetDragDropOp();
+
+	struct FDraggingWidgetReference
+	{
+		FWidgetReference Widget;
+
+		FVector2D DraggedOffset;
+	};
+
+	struct FItem
+	{
+		/** The slot properties for the old slot the widget was in, is used to attempt to reapply the same layout information */
+		TMap<FName, FString> ExportedSlotProperties;
+
+		/** The widget being dragged */
+		UWidget* Template;
+
+		/** The preview widget being dragged */
+		UWidget* Preview;
+
+		/** Can the drag drop change the widget's parent? */
+		bool bStayingInParent;
+
+		/** The original parent of the widget. */
+		FWidgetReference ParentWidget;
+
+		/** The offset of the original click location, as a percentage of the widget's size. */
+		FVector2D DraggedOffset;
+	};
+
+	TArray<FItem> DraggedWidgets;
+
+	bool bShowingMessage;
+
+	IUMGDesigner* Designer;
+
+	static TSharedRef<FSelectedWidgetDragDropOp> New(TSharedPtr<FWidgetBlueprintEditor> Editor, IUMGDesigner* InDesigner, const TArray<FDraggingWidgetReference>& InWidgets);
+};
diff --git a/Engine/Source/Editor/UMGEditor/Private/Hierarchy/SHierarchyViewItem.cpp b/Engine/Source/Editor/UMGEditor/Private/Hierarchy/SHierarchyViewItem.cpp
--- a/Engine/Source/Editor/UMGEditor/Private/Hierarchy/SHierarchyViewItem.cpp
+++ b/Engine/Source/Editor/UMGEditor/Private/Hierarchy/SHierarchyViewItem.cpp
@@ -24,10 +24,9 @@
 #include "DragAndDrop/AssetDragDropOp.h"
 #include "DragAndDrop/ClassDragDropOp.h"
 #include "DragDrop/WidgetTemplateDragDropOp.h"
-#include "WidgetTemplate.h"
-
-
+#include "DragDrop/SelectedWidgetDragDropOp.h"
 
+#include "WidgetTemplate.h"
 
 #include "Widgets/Text/SInlineEditableTextBlock.h"
 
@@ -159,6 +158,12 @@ TOptional<EItemDropZone> ProcessHierarchyDragDrop(const FDragDropEvent& DragDrop
 		}
 	}
 
+	// We do not support to dragging a Widget from the Viewport to the Hierarchy panel
+	if (TSharedPtr<FSelectedWidgetDragDropOp> SelectedWidgetDragDropOp = DragDropEvent.GetOperationAs<FSelectedWidgetDragDropOp>())
+	{
+		return TOptional<EItemDropZone>();
+	}
+
 	if ( TargetTemplate && ( DropZone == EItemDropZone::AboveItem || DropZone == EItemDropZone::BelowItem ) )
 	{
 		if ( UPanelWidget* TargetParentTemplate = TargetTemplate->GetParent() )
@@ -208,8 +213,10 @@ TOptional<EItemDropZone> ProcessHierarchyDragDrop(const FDragDropEvent& DragDrop
 
 				Blueprint->WidgetTree->SetFlags(RF_Transactional);
 				Blueprint->WidgetTree->Modify();
-
-				Blueprint->WidgetTree->RootWidget = FWidgetBlueprintEditorUtils::GetWidgetTemplateFromDragDrop(Blueprint, Blueprint->WidgetTree, DragDropOp);
+				if (UWidget* Widget = FWidgetBlueprintEditorUtils::GetWidgetTemplateFromDragDrop(Blueprint, Blueprint->WidgetTree, DragDropOp))
+				{
+					Blueprint->WidgetTree->RootWidget = Widget;
+				}
 				FBlueprintEditorUtils::MarkBlueprintAsStructurallyModified(Blueprint);
 			}
 
@@ -241,20 +248,21 @@ TOptional<EItemDropZone> ProcessHierarchyDragDrop(const FDragDropEvent& DragDrop
 					
 					Parent->SetFlags(RF_Transactional);
 					Parent->Modify();
-					UWidget* Widget = FWidgetBlueprintEditorUtils::GetWidgetTemplateFromDragDrop(Blueprint, Blueprint->WidgetTree, DragDropOp);
-
-					UPanelSlot* NewSlot = nullptr;
-					if (Index.IsSet())
-					{
-						NewSlot = Parent->InsertChildAt(Index.GetValue(), Widget);
-					}
-					else
+					if (UWidget* Widget = FWidgetBlueprintEditorUtils::GetWidgetTemplateFromDragDrop(Blueprint, Blueprint->WidgetTree, DragDropOp))
 					{
-						NewSlot = Parent->AddChild(Widget);
-					}
-					check(NewSlot);
+						UPanelSlot* NewSlot = nullptr;
+						if (Index.IsSet())
+						{
+							NewSlot = Parent->InsertChildAt(Index.GetValue(), Widget);
+						}
+						else
+						{
+							NewSlot = Parent->AddChild(Widget);
+						}
+						check(NewSlot);
 
-					FBlueprintEditorUtils::MarkBlueprintAsStructurallyModified(Blueprint);
+						FBlueprintEditorUtils::MarkBlueprintAsStructurallyModified(Blueprint);
+					}
 				}
 
 				if (DecoratedDragDropOp.IsValid())
@@ -749,11 +757,12 @@ TOptional<EItemDropZone> FHierarchyRoot::HandleCanAcceptDrop(const FDragDropEven
 	TSharedPtr<FDragDropOperation> DragDropOp = DragDropEvent.GetOperation();
 	if (DragDropOp.IsValid())
 	{
-		UWidgetBlueprint* Blueprint = BlueprintEditor.Pin()->GetWidgetBlueprintObj();
-		if(Blueprint)
+		if(UWidgetBlueprint* Blueprint = BlueprintEditor.Pin()->GetWidgetBlueprintObj())
 		{
-			UWidget* Widget = FWidgetBlueprintEditorUtils::GetWidgetTemplateFromDragDrop(Blueprint, Blueprint->WidgetTree, DragDropOp);
-			bIsFreeFromCircularReferences = !HasCircularReferences(Blueprint, Widget, DragDropOp);
+			if (UWidget* Widget = FWidgetBlueprintEditorUtils::GetWidgetTemplateFromDragDrop(Blueprint, Blueprint->WidgetTree, DragDropOp))
+			{
+				bIsFreeFromCircularReferences = !HasCircularReferences(Blueprint, Widget, DragDropOp);
+			}
 		}
 	}
 
@@ -908,14 +917,15 @@ TOptional<EItemDropZone> FNamedSlotModel::HandleCanAcceptDrop(const FDragDropEve
 				return TOptional<EItemDropZone>();
 			}
 
-			UWidget* Widget = FWidgetBlueprintEditorUtils::GetWidgetTemplateFromDragDrop(Blueprint, Blueprint->WidgetTree, DragDropOp);
-			bool bIsFreeFromCircularReferences = !HasCircularReferences(Blueprint, Widget, DragDropOp);
-
-			if (DecoratedDragDropOp.IsValid())
+			if (UWidget* Widget = FWidgetBlueprintEditorUtils::GetWidgetTemplateFromDragDrop(Blueprint, Blueprint->WidgetTree, DragDropOp))
 			{
-				DecoratedDragDropOp->CurrentIconBrush = FEditorStyle::GetBrush(TEXT("Graph.ConnectorFeedback.OK"));
+				bool bIsFreeFromCircularReferences = !HasCircularReferences(Blueprint, Widget, DragDropOp);
+				if (DecoratedDragDropOp.IsValid())
+				{
+					DecoratedDragDropOp->CurrentIconBrush = FEditorStyle::GetBrush(TEXT("Graph.ConnectorFeedback.OK"));
+				}
+				return bIsFreeFromCircularReferences ? EItemDropZone::OntoItem : TOptional<EItemDropZone>();
 			}
-			return bIsFreeFromCircularReferences ? EItemDropZone::OntoItem : TOptional<EItemDropZone>();
 		}
 	}
 
@@ -982,11 +992,15 @@ FReply FNamedSlotModel::HandleAcceptDrop(FDragDropEvent const& DragDropEvent, EI
 		Blueprint->WidgetTree->SetFlags(RF_Transactional);
 		Blueprint->WidgetTree->Modify();
 
-		UWidget* DroppingWidget = FWidgetBlueprintEditorUtils::GetWidgetTemplateFromDragDrop(Blueprint, Blueprint->WidgetTree, DragDropOp);
-
-		DoDrop(SlotHostWidget, DroppingWidget);
-
-		return FReply::Handled();
+		if (UWidget* DroppingWidget = FWidgetBlueprintEditorUtils::GetWidgetTemplateFromDragDrop(Blueprint, Blueprint->WidgetTree, DragDropOp))
+		{
+			DoDrop(SlotHostWidget, DroppingWidget);
+			return FReply::Handled();
+		}
+		else
+		{
+			return FReply::Unhandled();
+		}
 	}
 
 	TSharedPtr<FHierarchyWidgetDragDropOp> HierarchyDragDropOp = DragDropEvent.GetOperationAs<FHierarchyWidgetDragDropOp>();
@@ -1150,8 +1164,10 @@ TOptional<EItemDropZone> FHierarchyWidget::HandleCanAcceptDrop(const FDragDropEv
 		UWidgetBlueprint* Blueprint = BlueprintEditor.Pin()->GetWidgetBlueprintObj();
 		if (Blueprint)
 		{
-			UWidget* Widget = FWidgetBlueprintEditorUtils::GetWidgetTemplateFromDragDrop(Blueprint, Blueprint->WidgetTree, DragDropOp);
-			bIsFreeFromCircularReferences = !HasCircularReferences(Blueprint, Widget, DragDropOp);
+			if (UWidget* Widget = FWidgetBlueprintEditorUtils::GetWidgetTemplateFromDragDrop(Blueprint, Blueprint->WidgetTree, DragDropOp))
+			{
+				bIsFreeFromCircularReferences = !HasCircularReferences(Blueprint, Widget, DragDropOp);
+			}
 		}
 	}
 	bool bIsDrop = false;
