Fix unrecognized tab when using sequencer curves in widget editor

diff --git a/Engine/Source/Editor/UMGEditor/Private/BlueprintModes/WidgetDesignerApplicationMode.cpp b/Engine/Source/Editor/UMGEditor/Private/BlueprintModes/WidgetDesignerApplicationMode.cpp
--- a/Engine/Source/Editor/UMGEditor/Private/BlueprintModes/WidgetDesignerApplicationMode.cpp
+++ b/Engine/Source/Editor/UMGEditor/Private/BlueprintModes/WidgetDesignerApplicationMode.cpp
@@ -6,6 +6,7 @@
 #include "SBlueprintEditorToolbar.h"
 #include "BlueprintEditorSharedTabFactories.h"
 
+#include "HAL/PlatformApplicationMisc.h"
 #include "WidgetBlueprintEditorToolbar.h"
 #include "UMGEditorModule.h"
 #include "StatusBarSubsystem.h"
@@ -33,7 +34,18 @@ FWidgetDesignerApplicationMode::FWidgetDesignerApplicationMode(TSharedPtr<FWidge
 	// Override the default created category here since "Designer Editor" sounds awkward
 	WorkspaceMenuCategory = FWorkspaceItem::NewGroup(LOCTEXT("WorkspaceMenu_WidgetDesigner", "Widget Designer"));
 
-	TabLayout = FTabManager::NewLayout("WidgetBlueprintEditor_Designer_Layout_v4_555")
+	FDisplayMetrics DisplayMetrics;
+	FSlateApplication::Get().GetDisplayMetrics(DisplayMetrics);
+
+	const float DPIScale = FPlatformApplicationMisc::GetDPIScaleFactorAtPoint(DisplayMetrics.PrimaryDisplayWorkAreaRect.Left, DisplayMetrics.PrimaryDisplayWorkAreaRect.Top);
+
+	const float CenterScale = 0.4f;
+	const FVector2D DisplaySize(
+		DisplayMetrics.PrimaryDisplayWorkAreaRect.Right - DisplayMetrics.PrimaryDisplayWorkAreaRect.Left,
+		DisplayMetrics.PrimaryDisplayWorkAreaRect.Bottom - DisplayMetrics.PrimaryDisplayWorkAreaRect.Top);
+	const FVector2D WindowSize = (CenterScale * DisplaySize) / DPIScale;
+
+	TabLayout = FTabManager::NewLayout("WidgetBlueprintEditor_Designer_Layout_v4_7")
 	->AddArea
 	(
 		FTabManager::NewPrimaryArea()
@@ -93,6 +105,18 @@ FWidgetDesignerApplicationMode::FWidgetDesignerApplicationMode(TSharedPtr<FWidge
 				->SetForegroundTab(FAnimationTabSummoner::TabID)
 			)
 		)
+	)
+	->AddArea
+	(
+		// Sequencer popup
+		FTabManager::NewArea(WindowSize)
+		->SetOrientation(Orient_Vertical)
+		->Split
+		(
+			FTabManager::NewStack()
+			->SetSizeCoefficient(1.0f)
+			->AddTab("SequencerGraphEditor", ETabState::ClosedTab)
+		)
 	);
 
 	// Add Tab Spawners
