Widget Tool Palette Prototype.
Add Marquee & create widget tools.
Add framework for widget-based tools.

diff --git a/Engine/Plugins/Experimental/WidgetEditorToolPalette/Config/BaseWidgetEditorToolPalette.ini b/Engine/Plugins/Experimental/WidgetEditorToolPalette/Config/BaseWidgetEditorToolPalette.ini
new file mode 100644
--- /dev/null
+++ b/Engine/Plugins/Experimental/WidgetEditorToolPalette/Config/BaseWidgetEditorToolPalette.ini
@@ -0,0 +1,6 @@
+[/Script/WidgetEditorToolPalette.CreateWidgetToolSettings]
+CreateWidgetStacks=(DisplayName="Common",WidgetToolInfos=((WidgetClass=Class'"/Script/UMG.Border"',WidgetHotkey=(Key=One,bAlt=True)),(WidgetClass=Class'"/Script/UMG.Button"',WidgetHotkey=(Key=Two,bAlt=True)),(WidgetClass=Class'"/Script/UMG.CheckBox"',WidgetHotkey=(Key=Three,bAlt=True)),(WidgetClass=Class'"/Script/UMG.Image"',WidgetHotkey=(Key=Four,bAlt=True)),(WidgetClass=Class'"/Script/UMG.NamedSlot"',WidgetHotkey=(Key=Five,bAlt=True)),(WidgetClass=Class'"/Script/UMG.ProgressBar"',WidgetHotkey=(Key=Six,bAlt=True)),(WidgetClass=Class'"/Script/UMG.RichTextBlock"',WidgetHotkey=(Key=Seven,bAlt=True)),(WidgetClass=Class'"/Script/UMG.Slider"',WidgetHotkey=(Key=Eight,bAlt=True)),(WidgetClass=Class'"/Script/UMG.TextBlock"',WidgetHotkey=(Key=Nine,bAlt=True))))
+CreateWidgetStacks=(DisplayName="Input",WidgetToolInfos=((WidgetClass=Class'"/Script/UMG.ComboBoxString"',WidgetHotkey=(Key=I,bAlt=True)),(WidgetClass=Class'"/Script/UMG.EditableText"',WidgetHotkey=(Key=I,bAlt=True)),(WidgetClass=Class'"/Script/UMG.MultiLineEditableText"',WidgetHotkey=(Key=I,bAlt=True)),(WidgetClass=Class'"/Script/UMG.SpinBox"',WidgetHotkey=(Key=I,bAlt=True)),(WidgetClass=Class'"/Script/UMG.EditableTextBox"',WidgetHotkey=(Key=I,bAlt=True)),(WidgetClass=Class'"/Script/UMG.MultiLineEditableTextBox"',WidgetHotkey=(Key=I,bAlt=True))))
+CreateWidgetStacks=(DisplayName="Lists",WidgetToolInfos=((WidgetClass=Class'"/Script/UMG.ListView"',WidgetHotkey=(Key=L,bAlt=True)),(WidgetClass=Class'"/Script/UMG.TileView"',WidgetHotkey=(Key=L,bAlt=True)),(WidgetClass=Class'"/Script/UMG.TreeView"',WidgetHotkey=(Key=L,bAlt=True))))
+CreateWidgetStacks=(DisplayName="Panel",WidgetToolInfos=((WidgetClass=Class'"/Script/UMG.CanvasPanel"',WidgetHotkey=(Key=C,bAlt=True)),(WidgetClass=Class'"/Script/UMG.GridPanel"',WidgetHotkey=(Key=G,bAlt=True)),(WidgetClass=Class'"/Script/UMG.HorizontalBox"',WidgetHotkey=(Key=H,bAlt=True)),(WidgetClass=Class'"/Script/UMG.Overlay"',WidgetHotkey=(Key=O,bAlt=True)),(WidgetClass=Class'"/Script/UMG.SafeZone"',WidgetHotkey=(Key=F,bAlt=True)),(WidgetClass=Class'"/Script/UMG.ScaleBox"',WidgetHotkey=(Key=A,bAlt=True)),(WidgetClass=Class'"/Script/UMG.ScrollBox"',WidgetHotkey=(Key=S,bAlt=True)),(WidgetClass=Class'"/Script/UMG.SizeBox"',WidgetHotkey=(Key=Z,bAlt=True)),(WidgetClass=Class'"/Script/UMG.UniformGridPanel"',WidgetHotkey=(Key=U,bAlt=True)),(WidgetClass=Class'"/Script/UMG.VerticalBox"',WidgetHotkey=(Key=V,bAlt=True)),(WidgetClass=Class'"/Script/UMG.WidgetSwitcher"',WidgetHotkey=(Key=W,bAlt=True)),(WidgetClass=Class'"/Script/UMG.WrapBox"',WidgetHotkey=(Key=R,bAlt=True))))
+CreateWidgetStacks=(DisplayName="Primitive",WidgetToolInfos=((WidgetClass=Class'"/Script/UMG.CircularThrobber"',WidgetHotkey=(Key=P,bAlt=True)),(WidgetClass=Class'"/Script/UMG.MenuAnchor"',WidgetHotkey=(Key=P,bAlt=True)),(WidgetClass=Class'"/Script/UMG.NativeWidgetHost"',WidgetHotkey=(Key=P,bAlt=True)),(WidgetClass=Class'"/Script/UMG.Spacer"',WidgetHotkey=(Key=P,bAlt=True)),(WidgetClass=Class'"/Script/UMG.Throbber"',WidgetHotkey=(Key=P,bAlt=True))))
\ No newline at end of file
diff --git a/Engine/Plugins/Experimental/WidgetEditorToolPalette/Source/WidgetEditorToolPalette/Private/DefaultTools/CreateWidgetTool.cpp b/Engine/Plugins/Experimental/WidgetEditorToolPalette/Source/WidgetEditorToolPalette/Private/DefaultTools/CreateWidgetTool.cpp
new file mode 100644
--- /dev/null
+++ b/Engine/Plugins/Experimental/WidgetEditorToolPalette/Source/WidgetEditorToolPalette/Private/DefaultTools/CreateWidgetTool.cpp
@@ -0,0 +1,110 @@
+// Copyright Epic Games, Inc. All Rights Reserved.
+
+#include "DefaultTools/CreateWidgetTool.h"
+#include "InteractiveToolManager.h"
+#include "WidgetModeManager.h"
+#include "ToolBuilderUtil.h"
+#include "WidgetBlueprintEditor.h"
+#include "Tools/AssetEditorContextInterface.h"
+#include "ContextObjectStore.h"
+#include "Engine/World.h"
+
+#include "WidgetBlueprint.h"
+#include "WidgetTemplate.h"
+#include "Kismet2/BlueprintEditorUtils.h"
+#include "WidgetBlueprintEditorUtils.h"
+#include "Blueprint/WidgetTree.h"
+#include "Components/Button.h"
+#include "Components/Widget.h"
+#include "Templates/WidgetTemplateClass.h"
+#include "DragDrop/WidgetTemplateDragDropOp.h"
+#include "Designer/SDesignerView.h"
+
+// localization namespace
+#define LOCTEXT_NAMESPACE "UCreateWidgetTool"
+
+/*
+ * ToolBuilder
+ */
+
+UCreateWidgetToolBuilder::UCreateWidgetToolBuilder()
+	: WidgetClass(nullptr)
+{
+}
+
+UInteractiveTool* UCreateWidgetToolBuilder::BuildTool(const FToolBuilderState & SceneState) const
+{
+	UCreateWidgetTool* NewTool = NewObject<UCreateWidgetTool>(SceneState.ToolManager);
+	NewTool->WidgetClass = WidgetClass;
+
+	if (IAssetEditorContextInterface* AssetEditorContext = SceneState.ToolManager->GetContextObjectStore()->FindContext<IAssetEditorContextInterface>())
+	{
+		if (const IToolkitHost* ToolkitHost = AssetEditorContext->GetToolkitHost())
+		{
+			// @TODO: DarenC - try to avoid this cast
+			FWidgetModeManager& WidgetModeManager = static_cast<FWidgetModeManager&>(ToolkitHost->GetEditorModeManager());
+			NewTool->SetOwningWidget(WidgetModeManager.GetManagedWidget());
+			NewTool->SetOwningToolkit(WidgetModeManager.OwningToolkit.Pin());
+		}
+	}
+
+	return NewTool;
+}
+
+/*
+ * Tool
+ */
+
+UCreateWidgetToolProperties::UCreateWidgetToolProperties()
+{
+	// PlaceOnObjects = true;
+}
+
+UCreateWidgetTool::UCreateWidgetTool()
+{
+}
+
+void UCreateWidgetTool::SetOwningToolkit(TSharedPtr<IToolkit> InOwningToolkit)
+{
+	OwningToolkit = StaticCastSharedPtr<FWidgetBlueprintEditor>(InOwningToolkit);
+}
+
+void UCreateWidgetTool::SetOwningWidget(TSharedPtr<SWidget> InOwningWidget)
+{
+	OwningWidget = StaticCastSharedPtr<SDesignerView>(InOwningWidget);
+}
+
+void UCreateWidgetTool::Setup()
+{
+	USingleClickTool::Setup();
+
+	Properties = NewObject<UCreateWidgetToolProperties>(this);
+	AddToolPropertySource(Properties);
+}
+
+bool UCreateWidgetTool::OnMouseButtonUp(const FGeometry& MyGeometry, const FPointerEvent& MouseEvent)
+{
+	if (TSharedPtr<FWidgetBlueprintEditor> BPEditor = OwningToolkit.Pin())
+	{
+		if (TSharedPtr<SDesignerView> PinnedOwningWidget = OwningWidget.Pin())
+		{
+			const FVector2D ClickPosition = PinnedOwningWidget->PanelCoordToGraphCoord(MyGeometry.AbsoluteToLocal(MouseEvent.GetScreenSpacePosition()));
+
+			// Fire drag drop event at the click location
+			TSharedPtr<FWidgetTemplateClass> WidgetTemplate = MakeShared<FWidgetTemplateClass>(WidgetClass);
+			TSharedRef<FWidgetTemplateDragDropOp> TemplateDragDrop = FWidgetTemplateDragDropOp::New(WidgetTemplate);
+			FDragDropEvent DragDropEvent = FDragDropEvent(MouseEvent, TemplateDragDrop);
+
+			PinnedOwningWidget->OnDrop(MyGeometry, DragDropEvent).IsEventHandled();
+			
+			// We only create one widget per click, activate default select tool after click.
+			GetToolManager()->SelectActiveToolType(EToolSide::Mouse, "DefaultCursor");
+			GetToolManager()->ActivateTool(EToolSide::Mouse);
+		}
+	}
+
+	// Always return not handled, we want to use the default mouse up FReply, which includes changing mouse capture.
+	return false;
+}
+
+#undef LOCTEXT_NAMESPACE
\ No newline at end of file
diff --git a/Engine/Plugins/Experimental/WidgetEditorToolPalette/Source/WidgetEditorToolPalette/Private/DefaultTools/RectangleSelectTool.cpp b/Engine/Plugins/Experimental/WidgetEditorToolPalette/Source/WidgetEditorToolPalette/Private/DefaultTools/RectangleSelectTool.cpp
new file mode 100644
--- /dev/null
+++ b/Engine/Plugins/Experimental/WidgetEditorToolPalette/Source/WidgetEditorToolPalette/Private/DefaultTools/RectangleSelectTool.cpp
@@ -0,0 +1,270 @@
+// Copyright Epic Games, Inc. All Rights Reserved.
+
+#include "DefaultTools/RectangleSelectTool.h"
+#include "InteractiveToolManager.h"
+#include "WidgetModeManager.h"
+#include "ToolBuilderUtil.h"
+#include "BaseBehaviors/ClickDragBehavior.h"
+#include "Tools/AssetEditorContextInterface.h"
+#include "ContextObjectStore.h"
+#include "MarqueeOperation.h"
+#include "Designer/SDesignerView.h"
+#include "WidgetBlueprintEditor.h"
+#include "WidgetBlueprint.h"
+#include "WidgetReference.h"
+#include "Blueprint/WidgetTree.h"
+
+// localization namespace
+#define LOCTEXT_NAMESPACE "URectangleSelectTool"
+
+/*
+ * ToolBuilder
+ */
+
+bool URectangleSelectToolBuilder::CanBuildTool(const FToolBuilderState & SceneState) const
+{
+	return true;
+}
+
+UInteractiveTool* URectangleSelectToolBuilder::BuildTool(const FToolBuilderState & SceneState) const
+{
+	URectangleSelectTool* NewTool = NewObject<URectangleSelectTool>(SceneState.ToolManager);
+
+	if (IAssetEditorContextInterface* AssetEditorContext = SceneState.ToolManager->GetContextObjectStore()->FindContext<IAssetEditorContextInterface>())
+	{
+		if (const IToolkitHost* ToolkitHost = AssetEditorContext->GetToolkitHost())
+		{
+			// @TODO: DarenC - try to avoid this cast
+			FWidgetModeManager& WidgetModeManager = static_cast<FWidgetModeManager&>(ToolkitHost->GetEditorModeManager());
+			NewTool->SetOwningWidget(WidgetModeManager.GetManagedWidget());
+			NewTool->SetOwningToolkit(WidgetModeManager.OwningToolkit.Pin());
+		}
+	}
+
+	return NewTool;
+}
+
+/*
+ * Tool
+ */
+
+URectangleSelectProperties::URectangleSelectProperties()
+	: UInteractiveToolPropertySet()
+	, bIsMarqueeActive(false)
+{
+}
+
+void URectangleSelectTool::SetOwningToolkit(TSharedPtr<IToolkit> InOwningToolkit)
+{
+	OwningToolkit = StaticCastSharedPtr<FWidgetBlueprintEditor>(InOwningToolkit);
+}
+
+void URectangleSelectTool::SetOwningWidget(TSharedPtr<SWidget> InOwningWidget)
+{
+	OwningWidget = StaticCastSharedPtr<SDesignerView>(InOwningWidget);
+}
+
+void URectangleSelectTool::Setup()
+{
+	UInteractiveTool::Setup();
+
+	// Create the property set and register it with the Tool
+	Properties = NewObject<URectangleSelectProperties>(this, "Rectangle Select");
+	AddToolPropertySource(Properties);
+}
+
+bool URectangleSelectTool::OnKeyDown(const FGeometry& MyGeometry, const FKeyEvent& InKeyEvent)
+{
+	// The existing marquee operation was canceled
+	if (InKeyEvent.GetKey() == EKeys::Escape && Properties->Marquee.IsValid())
+	{
+		Properties->Marquee = FMarqueeOperation();
+		Properties->bIsMarqueeActive = false;
+
+		return true;
+	}
+
+	return false;
+}
+
+bool URectangleSelectTool::OnMouseButtonDown(const FGeometry& MyGeometry, const FPointerEvent& MouseEvent)
+{
+	if (MouseEvent.GetEffectingButton() == EKeys::LeftMouseButton)
+	{
+		if (TSharedPtr<SDesignerView> PinnedOwningWidget = OwningWidget.Pin())
+		{
+			const FVector2D ClickPosition = PinnedOwningWidget->PanelCoordToGraphCoord(MyGeometry.AbsoluteToLocal(MouseEvent.GetScreenSpacePosition()));
+
+			Properties->Marquee = FMarqueeOperation();
+			Properties->Marquee.Start(ClickPosition, FMarqueeOperation::Type::Add);
+			Properties->bIsMarqueeActive = true;
+			PrevSelectionEndPoint = ClickPosition;
+
+			// Clear existing selection
+			if (!MouseEvent.GetModifierKeys().IsControlDown())
+			{
+				if (TSharedPtr<FWidgetBlueprintEditor> WidgetBlueprintEditor = OwningToolkit.Pin())
+				{
+					WidgetBlueprintEditor->SelectWidgets({}, false);
+				}
+			}
+
+			return true;
+		}
+	}
+
+	return false;
+}
+
+bool URectangleSelectTool::OnMouseButtonUp(const FGeometry& MyGeometry, const FPointerEvent& MouseEvent)
+{
+	if (MouseEvent.GetEffectingButton() == EKeys::LeftMouseButton)
+	{
+		if (Properties->bIsMarqueeActive)
+		{
+			if (TSharedPtr<SDesignerView> PinnedOwningWidget = OwningWidget.Pin())
+			{
+				const FVector2D Position = PinnedOwningWidget->PanelCoordToGraphCoord(MyGeometry.AbsoluteToLocal(MouseEvent.GetScreenSpacePosition()));
+
+				// Just test equality assuming whole pixel movements
+				if (Position != Properties->Marquee.Rect.EndPoint)
+				{
+					Properties->Marquee.Rect.UpdateEndPoint(Position);
+					SelectWidgetsAffectedByMarquee();
+				}
+			}
+
+			// The existing marquee operation has finished
+			Properties->Marquee = FMarqueeOperation();
+			Properties->bIsMarqueeActive = false;
+
+			return true;
+		}
+	}
+
+	return false;
+}
+
+bool URectangleSelectTool::OnMouseMove(const FGeometry& MyGeometry, const FPointerEvent& MouseEvent)
+{
+	if (Properties->bIsMarqueeActive)
+	{
+		if (TSharedPtr<SDesignerView> PinnedOwningWidget = OwningWidget.Pin())
+		{
+			const FVector2D Position = PinnedOwningWidget->PanelCoordToGraphCoord(MyGeometry.AbsoluteToLocal(MouseEvent.GetScreenSpacePosition()));
+
+			// Just test equality assuming whole pixel movements
+			if (Position != Properties->Marquee.Rect.EndPoint)
+			{
+				Properties->Marquee.Rect.UpdateEndPoint(Position);
+				SelectWidgetsAffectedByMarquee();
+			}
+		}
+
+		return true;
+	}
+
+	return false;
+}
+
+int32 URectangleSelectTool::OnPaint(const FPaintArgs& Args, const FGeometry& AllottedGeometry, const FSlateRect& MyCullingRect, FSlateWindowElementList& OutDrawElements, int32 LayerId, const FWidgetStyle& InWidgetStyle, bool bParentEnabled) const
+{
+	if (TSharedPtr<SDesignerView> PinnedOwningWidget = OwningWidget.Pin())
+	{
+		if (Properties->Marquee.IsValid())
+		{
+			FSlateDrawElement::MakeBox(
+				OutDrawElements,
+				LayerId,
+				AllottedGeometry.ToPaintGeometry(PinnedOwningWidget->GraphCoordToPanelCoord(Properties->Marquee.Rect.GetUpperLeft()), Properties->Marquee.Rect.GetSize() * PinnedOwningWidget->GetZoomAmount()),
+				FEditorStyle::GetBrush(TEXT("MarqueeSelection"))
+			);
+		}
+	}
+
+	return LayerId;
+}
+
+void URectangleSelectTool::SelectWidgetsAffectedByMarquee()
+{
+	TSharedPtr<SDesignerView> PinnedOwningWidget = OwningWidget.Pin();
+	TSharedPtr<FWidgetBlueprintEditor> WidgetBlueprintEditor = OwningToolkit.Pin();
+	if (PinnedOwningWidget && WidgetBlueprintEditor)
+	{
+		TArray<UWidget*> WidgetsInTree;
+		WidgetBlueprintEditor->GetPreview()->WidgetTree->GetAllWidgets(WidgetsInTree);
+
+		// We are only interested in the delta from the previous selection
+		// We can get two rectangles that any widgets in this selection delta must intersect 
+		// by comparing the previous selection's endpoint with our current endpoint, 
+		// and projecting those points onto the axis-aligned lines centered at the selection start
+
+		FVector2D StartPoint = Properties->Marquee.Rect.StartPoint;
+		FVector2D OldEndPoint = PrevSelectionEndPoint;
+		FVector2D NewEndPoint = Properties->Marquee.Rect.EndPoint;
+
+		float SelectionDeltaX = FMath::Square(NewEndPoint.X - StartPoint.X) - FMath::Square(OldEndPoint.X - StartPoint.X);
+		float SelectionDeltaY = FMath::Square(NewEndPoint.Y - StartPoint.Y) - FMath::Square(OldEndPoint.Y - StartPoint.Y);
+
+		// If our selection is shrinking, use previous coordinate pair value to ensure a collision when testing for deselection
+
+		FMarqueeRect RectDeltaX;
+		RectDeltaX.StartPoint = { OldEndPoint.X, StartPoint.Y };
+		RectDeltaX.EndPoint = { NewEndPoint.X, SelectionDeltaY > 0 ? NewEndPoint.Y : OldEndPoint.Y };
+
+		FMarqueeRect RectDeltaY;
+		RectDeltaY.StartPoint = { StartPoint.X, OldEndPoint.Y };
+		RectDeltaY.EndPoint = { SelectionDeltaX > 0 ? NewEndPoint.X : OldEndPoint.X, NewEndPoint.Y };
+
+		FMarqueeRect PrevSelectionRect;
+		PrevSelectionRect.StartPoint = StartPoint;
+		PrevSelectionRect.EndPoint = PrevSelectionEndPoint;
+
+		bool bToggledSelection = false;
+		TSet<FWidgetReference> ToggleSelectionSet;
+		for (UWidget* Widget : WidgetsInTree)
+		{
+			if (Widget && !Widget->IsRooted())
+			{
+				const FVector2D Position = PinnedOwningWidget->PanelCoordToGraphCoord(PinnedOwningWidget->GetTickSpaceGeometry().AbsoluteToLocal(Widget->GetTickSpaceGeometry().GetAbsolutePosition()));
+				const FVector2D Size = Widget->GetPaintSpaceGeometry().GetAbsoluteSize() * PinnedOwningWidget->GetPreviewDPIScale() / PinnedOwningWidget->GetZoomAmount();
+
+				if (Size.X > 0.f && Size.Y > 0.f)
+				{
+					const FSlateRect WidgetPaintArea(Position.X, Position.Y, Position.X + Size.X, Position.Y + Size.Y);
+
+					// Intersecting with delta rect, toggle selection appropriately if needed
+					auto ToggleSelection = [&](const FMarqueeRect& DeltaRect, float SelectionDelta)
+					{
+						const bool bCanSelectionChange = SelectionDelta && FSlateRect::DoRectanglesIntersect(DeltaRect.ToSlateRect(), WidgetPaintArea);
+						if (bCanSelectionChange)
+						{
+							if (SelectionDelta > 0 && FSlateRect::IsRectangleContained(Properties->Marquee.Rect.ToSlateRect(), WidgetPaintArea))
+							{
+								ToggleSelectionSet.Add(WidgetBlueprintEditor->GetReferenceFromPreview(Widget));
+								return true;
+							}
+							else if (SelectionDelta < 0 && FSlateRect::IsRectangleContained(PrevSelectionRect.ToSlateRect(), WidgetPaintArea))
+							{
+								ToggleSelectionSet.Add(WidgetBlueprintEditor->GetReferenceFromPreview(Widget));
+								return true;
+							}
+						}
+						return false;
+					};
+
+					bToggledSelection |= ToggleSelection(RectDeltaX, SelectionDeltaX) || ToggleSelection(RectDeltaY, SelectionDeltaY);
+				}
+			}
+		}
+
+		if (bToggledSelection)
+		{
+			WidgetBlueprintEditor->SelectWidgets(ToggleSelectionSet, true /*bool bAppendOrToggle - true implies toggle */);
+		}
+
+		PrevSelectionEndPoint = NewEndPoint;
+	}
+}
+
+#undef LOCTEXT_NAMESPACE
diff --git a/Engine/Plugins/Experimental/WidgetEditorToolPalette/Source/WidgetEditorToolPalette/Private/WidgetEditorToolPaletteCommands.cpp b/Engine/Plugins/Experimental/WidgetEditorToolPalette/Source/WidgetEditorToolPalette/Private/WidgetEditorToolPaletteCommands.cpp
new file mode 100644
--- /dev/null
+++ b/Engine/Plugins/Experimental/WidgetEditorToolPalette/Source/WidgetEditorToolPalette/Private/WidgetEditorToolPaletteCommands.cpp
@@ -0,0 +1,76 @@
+// Copyright Epic Games, Inc. All Rights Reserved.
+
+#include "WidgetEditorToolPaletteCommands.h"
+#include "WidgetEditorToolPaletteStyle.h"
+#include "EditorStyleSet.h"
+#include "Framework/Commands/UICommandInfo.h"
+#include "Settings/CreateWidgetToolSettings.h"
+#include "Styling/SlateIconFinder.h"
+
+#define LOCTEXT_NAMESPACE "WidgetEditorToolPaletteCommands"
+
+FWidgetEditorToolPaletteCommands::FWidgetEditorToolPaletteCommands()
+	: TCommands<FWidgetEditorToolPaletteCommands>(
+		TEXT("WidgetEditorToolPaletteCommands"), 
+		NSLOCTEXT("Contexts", "WidgetToolPalette", "Widget Tool Palette"), 
+		NAME_None, 
+		FWidgetEditorToolPaletteStyle::Get()->GetStyleSetName())
+{}
+
+void FWidgetEditorToolPaletteCommands::RegisterCommands()
+{
+	UI_COMMAND(ToggleWidgetEditorToolPalette, "Default Tools", "Standard UMG tools palette.", EUserInterfaceActionType::ToggleButton, FInputChord());
+
+	UI_COMMAND(DefaultSelectTool, "Select", "Select", EUserInterfaceActionType::ToggleButton, FInputChord(EKeys::A));
+	UI_COMMAND(BeginRectangleSelectTool, "Marquee", "Marquee", EUserInterfaceActionType::ToggleButton, FInputChord(EKeys::M));
+
+	// Create commands defined in settings
+	const UCreateWidgetToolSettings* Settings = GetDefault<UCreateWidgetToolSettings>();
+
+	for (const FCreateWidgetStackInfo& CreateWidgetStack : Settings->CreateWidgetStacks)
+	{
+		TSharedPtr<FUICommandInfo> StackCommand;
+
+		FUICommandInfo::MakeCommandInfo(
+			this->AsShared(),
+			StackCommand,
+			FName(CreateWidgetStack.DisplayName),
+			FText::FromString(CreateWidgetStack.DisplayName), // @TODO: how does this localize?
+			FText::FromString(CreateWidgetStack.DisplayName),
+			FSlateIcon(),
+			EUserInterfaceActionType::ToggleButton,
+			FInputChord());
+
+		CreateWidgetToolStacks.Emplace(CreateWidgetStack.DisplayName, StackCommand);
+
+		for (const FCreateWidgetToolInfo& CreateWidgetToolInfo : CreateWidgetStack.WidgetToolInfos)
+		{
+			check(CreateWidgetToolInfo.WidgetClass);
+
+			TSharedPtr<FUICommandInfo> ToolCommand;
+
+			FText DisplayName = CreateWidgetToolInfo.DisplayName.IsEmpty() 
+				? CreateWidgetToolInfo.WidgetClass->GetDisplayNameText()
+				: FText::FromString(CreateWidgetToolInfo.DisplayName);
+
+			FUICommandInfo::MakeCommandInfo(
+				this->AsShared(),
+				ToolCommand,
+				FName(DisplayName.ToString()),
+				DisplayName, 
+				DisplayName,
+				FSlateIconFinder::FindIconForClass(CreateWidgetToolInfo.WidgetClass),
+				EUserInterfaceActionType::ToggleButton,
+				CreateWidgetToolInfo.WidgetHotkey);
+
+			CreateWidgetTools.Emplace(DisplayName.ToString(), ToolCommand);
+		}
+	}
+}
+
+const FWidgetEditorToolPaletteCommands& FWidgetEditorToolPaletteCommands::Get()
+{
+	return TCommands<FWidgetEditorToolPaletteCommands>::Get();
+}
+
+#undef LOCTEXT_NAMESPACE
diff --git a/Engine/Plugins/Experimental/WidgetEditorToolPalette/Source/WidgetEditorToolPalette/Private/WidgetEditorToolPaletteMode.cpp b/Engine/Plugins/Experimental/WidgetEditorToolPalette/Source/WidgetEditorToolPalette/Private/WidgetEditorToolPaletteMode.cpp
new file mode 100644
--- /dev/null
+++ b/Engine/Plugins/Experimental/WidgetEditorToolPalette/Source/WidgetEditorToolPalette/Private/WidgetEditorToolPaletteMode.cpp
@@ -0,0 +1,88 @@
+// Copyright Epic Games, Inc. All Rights Reserved.
+
+#include "WidgetEditorToolPaletteMode.h"
+#include "WidgetEditorToolPaletteCommands.h"
+#include "WidgetEditorToolPaletteToolkit.h"
+#include "ContextObjectStore.h"
+#include "EdModeInteractiveToolsContext.h"
+#include "ToolContexts/ToolStackContext.h"
+#include "ToolTargetManager.h"
+#include "DefaultTools/RectangleSelectTool.h"
+#include "DefaultTools/CreateWidgetTool.h"
+#include "BaseTools/SingleClickTool.h"
+#include "Settings/CreateWidgetToolSettings.h"
+
+#define LOCTEXT_NAMESPACE "WidgetEditorToolPaletteMode"
+
+const FEditorModeID UWidgetEditorToolPaletteMode::Id("WidgetEditorToolPaletteMode");
+
+UWidgetEditorToolPaletteMode::UWidgetEditorToolPaletteMode()
+{
+	bool bVisibleInLevelEditor = false;
+	Info = FEditorModeInfo(
+		Id, 
+		LOCTEXT("WidgetEditorToolPaletteMode", "WidgetTools"), 
+		FSlateIcon("WidgetEditorToolPaletteStyle", "WidgetEditorToolPaletteCommands.DefaultSelectTool"), 
+		bVisibleInLevelEditor);
+}
+
+void UWidgetEditorToolPaletteMode::Enter()
+{
+	UEdMode::Enter();
+	
+	const FWidgetEditorToolPaletteCommands& ToolPaletteCommands = FWidgetEditorToolPaletteCommands::Get();
+
+	USingleClickToolBuilder* SingleClickToolBuilder = NewObject<USingleClickToolBuilder>();
+	RegisterTool(ToolPaletteCommands.DefaultSelectTool, TEXT("DefaultCursor"), SingleClickToolBuilder);
+
+	URectangleSelectToolBuilder* RectangleSelectToolBuilder = NewObject<URectangleSelectToolBuilder>();
+	RegisterTool(ToolPaletteCommands.BeginRectangleSelectTool, TEXT("BeginRectangleSelectTool"), RectangleSelectToolBuilder);
+
+	// Add a tool stack context for tool stack support, tools will automatically be unregistered via destructor when context store is freed
+	UEditorInteractiveToolsContext* EdToolsContext = GetInteractiveToolsContext(EToolsContextScope::EdMode);
+	if (EdToolsContext && EdToolsContext->ContextObjectStore && EdToolsContext->ToolManager)
+	{
+		UToolStackContext* ToolStackContext = NewObject<UToolStackContext>(EdToolsContext->ToolManager);
+		EdToolsContext->ContextObjectStore->AddContextObject(ToolStackContext);
+		ToolStackContext->EdMode = this;
+		ToolStackContext->ToolContextScope = EToolsContextScope::EdMode;
+
+		// Create tool stacks defined in settings
+		const UCreateWidgetToolSettings* Settings = GetDefault<UCreateWidgetToolSettings>();
+		for (const FCreateWidgetStackInfo& CreateWidgetStack : Settings->CreateWidgetStacks)
+		{
+			const TSharedPtr<FUICommandInfo>& StackCommand = ToolPaletteCommands.CreateWidgetToolStacks[CreateWidgetStack.DisplayName];
+			ToolStackContext->RegisterToolStack(StackCommand, CreateWidgetStack.DisplayName);
+
+			for (const FCreateWidgetToolInfo& CreateWidgetToolInfo : CreateWidgetStack.WidgetToolInfos)
+			{
+				check(CreateWidgetToolInfo.WidgetClass);
+
+				FText DisplayName = CreateWidgetToolInfo.DisplayName.IsEmpty()
+					? CreateWidgetToolInfo.WidgetClass->GetDisplayNameText()
+					: FText::FromString(CreateWidgetToolInfo.DisplayName);
+
+				const TSharedPtr<FUICommandInfo>& ToolCommand = ToolPaletteCommands.CreateWidgetTools[DisplayName.ToString()];
+
+				UCreateWidgetToolBuilder* CreateWidgetToolBuilder = NewObject<UCreateWidgetToolBuilder>(GetTransientPackage(), CreateWidgetToolInfo.CreateWidgetToolBuilder);
+				CreateWidgetToolBuilder->WidgetClass = CreateWidgetToolInfo.WidgetClass;
+
+				ToolStackContext->AddToolToStack(ToolCommand, CreateWidgetStack.DisplayName, DisplayName.ToString(), CreateWidgetToolBuilder);
+			}
+		}
+	}
+
+	GetToolManager()->SelectActiveToolType(EToolSide::Mouse, TEXT("DefaultCursor"));
+}
+
+bool UWidgetEditorToolPaletteMode::UsesToolkits() const
+{ 
+	return true; 
+}
+
+void UWidgetEditorToolPaletteMode::CreateToolkit()
+{
+	Toolkit = MakeShareable(new FWidgetEditorToolPaletteToolkit);
+}
+
+#undef LOCTEXT_NAMESPACE
diff --git a/Engine/Plugins/Experimental/WidgetEditorToolPalette/Source/WidgetEditorToolPalette/Private/WidgetEditorToolPaletteModule.cpp b/Engine/Plugins/Experimental/WidgetEditorToolPalette/Source/WidgetEditorToolPalette/Private/WidgetEditorToolPaletteModule.cpp
new file mode 100644
--- /dev/null
+++ b/Engine/Plugins/Experimental/WidgetEditorToolPalette/Source/WidgetEditorToolPalette/Private/WidgetEditorToolPaletteModule.cpp
@@ -0,0 +1,84 @@
+// Copyright Epic Games, Inc. All Rights Reserved.
+
+#include "WidgetEditorToolPaletteModule.h"
+#include "WidgetEditorToolPaletteCommands.h"
+#include "WidgetEditorToolPaletteMode.h"
+#include "WidgetEditorToolPaletteStyle.h"
+#include "ToolMenus.h"
+#include "UMGEditorModule.h"
+#include "WidgetBlueprintEditor.h"
+#include "EditorModeManager.h"
+
+#define LOCTEXT_NAMESPACE "FWidgetEditorToolPaletteModule"
+
+void FWidgetEditorToolPaletteModule::StartupModule()
+{
+	FCoreDelegates::OnPostEngineInit.AddRaw(this, &FWidgetEditorToolPaletteModule::OnPostEngineInit);
+
+	// Register command for toolbar action
+	auto ExtendUMGEditorToolbar = [this](const TSharedRef<FUICommandList> InCommandList, TSharedRef<FWidgetBlueprintEditor> Editor) -> TSharedRef<FExtender>
+	{
+		// Note: actual UI extension done via 'ToolPaletteCommands'
+		TSharedPtr<FExtender> ToolbarExtender = MakeShareable(new FExtender);
+		TWeakPtr<FWidgetBlueprintEditor> Ptr(Editor);
+
+		InCommandList->MapAction(FWidgetEditorToolPaletteCommands::Get().ToggleWidgetEditorToolPalette,
+			FExecuteAction::CreateRaw(this, &FWidgetEditorToolPaletteModule::OnToggleWidgetEditorToolPaletteMode, Ptr),
+			FCanExecuteAction(),
+			FIsActionChecked::CreateRaw(this, &FWidgetEditorToolPaletteModule::IsWidgetEditorToolPaletteModeActive, Ptr));
+		Editor->ToolPaletteCommands.Add(FWidgetEditorToolPaletteCommands::Get().ToggleWidgetEditorToolPalette);
+
+		return ToolbarExtender.ToSharedRef();
+	};
+
+	IUMGEditorModule& UMGEditorModule = FModuleManager::LoadModuleChecked<IUMGEditorModule>("UMGEditor");
+	UMGEditorModule.AddWidgetEditorToolbarExtender(IUMGEditorModule::FWidgetEditorToolbarExtender::CreateLambda(ExtendUMGEditorToolbar));
+}
+
+
+void FWidgetEditorToolPaletteModule::ShutdownModule()
+{
+	FCoreDelegates::OnPostEngineInit.RemoveAll(this);
+	UToolMenus::UnregisterOwner(this);
+	FWidgetEditorToolPaletteCommands::Unregister();
+	FWidgetEditorToolPaletteStyle::Shutdown();
+}
+
+bool FWidgetEditorToolPaletteModule::IsWidgetEditorToolPaletteModeActive(TWeakPtr<FWidgetBlueprintEditor> InEditor) const
+{
+	TSharedPtr<FWidgetBlueprintEditor> PinnedEditor = InEditor.Pin();
+	return PinnedEditor.IsValid() && PinnedEditor->GetEditorModeManager().IsModeActive(UWidgetEditorToolPaletteMode::Id);
+}
+
+
+void FWidgetEditorToolPaletteModule::OnToggleWidgetEditorToolPaletteMode(TWeakPtr<FWidgetBlueprintEditor> InEditor)
+{
+	TSharedPtr<FWidgetBlueprintEditor> PinnedEditor = InEditor.Pin();
+
+	if (PinnedEditor.IsValid())
+	{
+		if (!IsWidgetEditorToolPaletteModeActive(InEditor))
+		{
+			PinnedEditor->GetEditorModeManager().ActivateMode(UWidgetEditorToolPaletteMode::Id, true);
+		}
+		else
+		{
+			PinnedEditor->GetEditorModeManager().ActivateDefaultMode();
+		}
+	}
+}
+
+void FWidgetEditorToolPaletteModule::OnPostEngineInit()
+{
+	// Register slate style overrides
+	FWidgetEditorToolPaletteStyle::Initialize();
+
+	// This code will execute after your module is loaded into memory; the exact timing is specified in the .uplugin file per-module
+	FWidgetEditorToolPaletteCommands::Register();
+}
+
+
+#undef LOCTEXT_NAMESPACE
+	
+IMPLEMENT_MODULE(FWidgetEditorToolPaletteModule, WidgetEditorToolPalette)
+
diff --git a/Engine/Plugins/Experimental/WidgetEditorToolPalette/Source/WidgetEditorToolPalette/Private/WidgetEditorToolPaletteStyle.cpp b/Engine/Plugins/Experimental/WidgetEditorToolPalette/Source/WidgetEditorToolPalette/Private/WidgetEditorToolPaletteStyle.cpp
new file mode 100644
--- /dev/null
+++ b/Engine/Plugins/Experimental/WidgetEditorToolPalette/Source/WidgetEditorToolPalette/Private/WidgetEditorToolPaletteStyle.cpp
@@ -0,0 +1,102 @@
+// Copyright Epic Games, Inc. All Rights Reserved.
+
+#include "WidgetEditorToolPaletteStyle.h"
+#include "Styling/SlateStyleRegistry.h"
+#include "Styling/SlateTypes.h"
+#include "Styling/CoreStyle.h"
+#include "EditorStyleSet.h"
+#include "Interfaces/IPluginManager.h"
+#include "SlateOptMacros.h"
+#include "Styling/SlateStyleMacros.h"
+
+
+#define IMAGE_PLUGIN_BRUSH( RelativePath, ... ) FSlateImageBrush( FWidgetEditorToolPaletteStyle::InContent( RelativePath, ".png" ), __VA_ARGS__ )
+
+// This is to fix the issue that SlateStyleMacros like IMAGE_BRUSH look for RootToContentDir but StyleSet->RootToContentDir is how this style is set up
+#define RootToContentDir StyleSet->RootToContentDir
+
+FString FWidgetEditorToolPaletteStyle::InContent(const FString& RelativePath, const ANSICHAR* Extension)
+{
+	static FString ContentDir = IPluginManager::Get().FindPlugin(TEXT("WidgetEditorToolPalette"))->GetContentDir();
+	return (ContentDir / RelativePath) + Extension;
+}
+
+TSharedPtr< FSlateStyleSet > FWidgetEditorToolPaletteStyle::StyleSet = nullptr;
+TSharedPtr< class ISlateStyle > FWidgetEditorToolPaletteStyle::Get() { return StyleSet; }
+
+FName FWidgetEditorToolPaletteStyle::GetStyleSetName()
+{
+	static FName ModelingToolsStyleName(TEXT("WidgetEditorToolPaletteStyle"));
+	return ModelingToolsStyleName;
+}
+
+const FSlateBrush* FWidgetEditorToolPaletteStyle::GetBrush(FName PropertyName, const ANSICHAR* Specifier)
+{
+	return Get()->GetBrush(PropertyName, Specifier);
+}
+
+
+BEGIN_SLATE_FUNCTION_BUILD_OPTIMIZATION
+
+void FWidgetEditorToolPaletteStyle::Initialize()
+{
+	// Const icon sizes
+	const FVector2D Icon8x8(8.0f, 8.0f);
+	const FVector2D Icon16x16(16.0f, 16.0f);
+	const FVector2D Icon20x20(20.0f, 20.0f);
+	const FVector2D Icon28x28(28.0f, 28.0f);
+	const FVector2D Icon40x40(40.0f, 40.0f);
+	const FVector2D Icon120(120.0f, 120.0f);
+
+	// Only register once
+	if (StyleSet.IsValid())
+	{
+		return;
+	}
+
+	StyleSet = MakeShareable(new FSlateStyleSet(GetStyleSetName()));
+	StyleSet->SetContentRoot(FPaths::EngineContentDir() / TEXT("Editor/Slate"));
+	StyleSet->SetCoreContentRoot(FPaths::EngineContentDir() / TEXT("Slate"));
+
+	const FTextBlockStyle& NormalText = FEditorStyle::Get().GetWidgetStyle<FTextBlockStyle>("NormalText");
+
+	// Tool Manager icons
+	{
+		// if (FCoreStyle::IsStarshipStyle())
+		// {
+		// 	StyleSet->Set("LevelEditor.ModelingToolsMode", new IMAGE_BRUSH_SVG("Starship/geometry", FVector2D(20.0f, 20.0f)));
+		// }
+		// else
+		// {
+		// 	StyleSet->Set("LevelEditor.ModelingToolsMode", new IMAGE_PLUGIN_BRUSH("Icons/icon_WidgetEditorToolPalette", FVector2D(40.0f, 40.0f)));
+		// 	StyleSet->Set("LevelEditor.ModelingToolsMode.Small", new IMAGE_PLUGIN_BRUSH("Icons/icon_WidgetEditorToolPalette", FVector2D(20.0f, 20.0f)));
+		// }
+
+		StyleSet->Set("WidgetEditorToolPaletteCommands.ToggleWidgetEditorToolPalette", new IMAGE_BRUSH("Icons/GeneralTools/Paint_40x", Icon20x20));
+		StyleSet->Set("WidgetEditorToolPaletteCommands.ToggleWidgetEditorToolPalette.Small", new IMAGE_BRUSH("Icons/GeneralTools/Paint_40x", Icon20x20));
+
+		StyleSet->Set("WidgetEditorToolPaletteCommands.DefaultSelectTool", 		new IMAGE_BRUSH("Icons/GeneralTools/Select_40x", Icon20x20));
+		StyleSet->Set("WidgetEditorToolPaletteCommands.DefaultSelectTool.Small", new IMAGE_BRUSH("Icons/GeneralTools/Select_40x", Icon20x20));
+		StyleSet->Set("WidgetEditorToolPaletteCommands.BeginRectangleSelectTool", new IMAGE_BRUSH_SVG("Icons/GeneralTools/Marque", Icon20x20));
+		StyleSet->Set("WidgetEditorToolPaletteCommands.BeginRectangleSelectTool.Small", new IMAGE_BRUSH_SVG("Icons/GeneralTools/Marque", Icon20x20));
+	}
+
+	FSlateStyleRegistry::RegisterSlateStyle(*StyleSet.Get());
+};
+
+END_SLATE_FUNCTION_BUILD_OPTIMIZATION
+
+#undef IMAGE_PLUGIN_BRUSH
+#undef IMAGE_BRUSH
+#undef BOX_BRUSH
+#undef DEFAULT_FONT
+
+void FWidgetEditorToolPaletteStyle::Shutdown()
+{
+	if (StyleSet.IsValid())
+	{
+		FSlateStyleRegistry::UnRegisterSlateStyle(*StyleSet.Get());
+		ensure(StyleSet.IsUnique());
+		StyleSet.Reset();
+	}
+}
diff --git a/Engine/Plugins/Experimental/WidgetEditorToolPalette/Source/WidgetEditorToolPalette/Private/WidgetEditorToolPaletteToolkit.cpp b/Engine/Plugins/Experimental/WidgetEditorToolPalette/Source/WidgetEditorToolPalette/Private/WidgetEditorToolPaletteToolkit.cpp
new file mode 100644
--- /dev/null
+++ b/Engine/Plugins/Experimental/WidgetEditorToolPalette/Source/WidgetEditorToolPalette/Private/WidgetEditorToolPaletteToolkit.cpp
@@ -0,0 +1,123 @@
+// Copyright Epic Games, Inc. All Rights Reserved.
+
+#include "WidgetEditorToolPaletteToolkit.h"
+#include "WidgetEditorToolPaletteCommands.h"
+#include "Widgets/Images/SImage.h"
+#include "Widgets/Input/SButton.h"
+
+#include "Tools/UEdMode.h"
+#include "InteractiveTool.h"
+#include "EdModeInteractiveToolsContext.h"
+#include "InteractiveToolManager.h"
+#include "Widgets/Text/STextBlock.h"
+#include "EditorStyleSet.h"
+#include "IDetailsView.h"
+#include "Framework/MultiBox/MultiBoxBuilder.h"
+#include "Framework/MultiBox/SToolBarStackButtonBlock.h"
+#include "Settings/CreateWidgetToolSettings.h"
+
+#define LOCTEXT_NAMESPACE "WidgetEditorToolPalette"
+
+static const FName WidgetToolPaletteTabName(TEXT("WidgetToolPalette"));
+
+void FWidgetEditorToolPaletteToolkit::Init(
+	const TSharedPtr<IToolkitHost>& InToolkitHost,
+	TWeakObjectPtr<UEdMode> InOwningMode )
+{
+	FModeToolkit::Init(InToolkitHost, InOwningMode);
+
+	ActiveToolName = FText::GetEmpty();
+	ActiveToolMessage = FText::GetEmpty();
+	SetCurrentPalette(WidgetToolPaletteTabName);
+}
+
+FName FWidgetEditorToolPaletteToolkit::GetToolkitFName() const
+{
+	return FName("WidgetToolPaletteToolkit");
+}
+
+FText FWidgetEditorToolPaletteToolkit::GetBaseToolkitName() const
+{
+	return LOCTEXT("DisplayName", "WidgetToolPaletteToolkit");
+}
+
+void FWidgetEditorToolPaletteToolkit::GetToolPaletteNames(TArray<FName>& PaletteNames) const
+{
+	PaletteNames = { WidgetToolPaletteTabName };
+}
+
+FText FWidgetEditorToolPaletteToolkit::GetToolPaletteDisplayName(FName PaletteName) const
+{
+	return FText::FromName(PaletteName);
+}
+
+void FWidgetEditorToolPaletteToolkit::BuildToolPalette(FName PaletteName, class FToolBarBuilder& ToolbarBuilder)
+{
+	const FWidgetEditorToolPaletteCommands& ToolPaletteCommands = FWidgetEditorToolPaletteCommands::Get();
+
+	if (PaletteName == WidgetToolPaletteTabName)
+	{
+		ToolbarBuilder.AddToolBarButton(ToolPaletteCommands.DefaultSelectTool);
+		ToolbarBuilder.AddToolBarButton(ToolPaletteCommands.BeginRectangleSelectTool);
+		ToolbarBuilder.AddSeparator();
+
+		// Create tool stacks defined in settings
+		const UCreateWidgetToolSettings* Settings = GetDefault<UCreateWidgetToolSettings>();
+		for (const TPair<FString, TSharedPtr<FUICommandInfo>>& CreateWidgetStack : ToolPaletteCommands.CreateWidgetToolStacks)
+		{
+			ToolbarBuilder.AddToolbarStackButton(CreateWidgetStack.Value);
+		}
+	}
+}
+
+TSharedRef<SWidget> FWidgetEditorToolPaletteToolkit::CreatePaletteWidget(TSharedPtr<FUICommandList> InCommandList, FName InToolbarCustomizationName, FName InPaletteName)
+{
+	FVerticalToolBarBuilder ModeToolbarBuilder(InCommandList, FMultiBoxCustomization(InToolbarCustomizationName));
+	ModeToolbarBuilder.SetStyle(&FEditorStyle::Get(), "VerticalToolBar");
+
+	BuildToolPalette(InPaletteName, ModeToolbarBuilder);
+
+	return ModeToolbarBuilder.MakeWidget();
+}
+
+FText FWidgetEditorToolPaletteToolkit::GetActiveToolDisplayName() const
+{
+	return ActiveToolName;
+}
+
+FText FWidgetEditorToolPaletteToolkit::GetActiveToolMessage() const
+{
+	return ActiveToolMessage;
+}
+
+void FWidgetEditorToolPaletteToolkit::UpdateActiveToolProperties(UInteractiveTool* Tool)
+{
+	if (Tool)
+	{
+		DetailsView->SetObjects(Tool->GetToolProperties(true));
+	}
+}
+
+void FWidgetEditorToolPaletteToolkit::OnToolStarted(UInteractiveToolManager* Manager, UInteractiveTool* Tool)
+{
+	FModeToolkit::OnToolStarted(Manager, Tool);
+
+	Tool->OnPropertySetsModified.AddSP(this, &FWidgetEditorToolPaletteToolkit::UpdateActiveToolProperties, Tool);
+
+	ActiveToolName = Tool->GetToolInfo().ToolDisplayName;
+}
+
+
+void FWidgetEditorToolPaletteToolkit::OnToolEnded(UInteractiveToolManager* Manager, UInteractiveTool* Tool)
+{
+	FModeToolkit::OnToolEnded(Manager, Tool);
+
+	if (Tool)
+	{
+		Tool->OnPropertySetsModified.RemoveAll(this);
+	}
+
+	ActiveToolName = FText::GetEmpty();
+}
+
+#undef LOCTEXT_NAMESPACE
diff --git a/Engine/Plugins/Experimental/WidgetEditorToolPalette/Source/WidgetEditorToolPalette/Public/DefaultTools/CreateWidgetTool.h b/Engine/Plugins/Experimental/WidgetEditorToolPalette/Source/WidgetEditorToolPalette/Public/DefaultTools/CreateWidgetTool.h
new file mode 100644
--- /dev/null
+++ b/Engine/Plugins/Experimental/WidgetEditorToolPalette/Source/WidgetEditorToolPalette/Public/DefaultTools/CreateWidgetTool.h
@@ -0,0 +1,79 @@
+// Copyright Epic Games, Inc. All Rights Reserved.
+
+#pragma once
+
+#include "CoreMinimal.h"
+#include "BaseTools/SingleClickTool.h"
+#include "BaseBehaviors/Widgets/WidgetBaseBehavior.h"
+#include "InteractiveToolBuilder.h"
+#include "CreateWidgetTool.generated.h"
+
+/**
+ * Builder for UCreateWidgetTool
+ */
+UCLASS()
+class WIDGETEDITORTOOLPALETTE_API UCreateWidgetToolBuilder : public UInteractiveToolBuilder
+{
+	GENERATED_BODY()
+
+	UCreateWidgetToolBuilder();
+
+public:
+	virtual bool CanBuildTool(const FToolBuilderState& SceneState) const override { return true; }
+	virtual UInteractiveTool* BuildTool(const FToolBuilderState& SceneState) const override;
+
+public:
+	/** Widget Class to be created by this tool */
+	UClass* WidgetClass;
+};
+
+/**
+ * Settings UObject for UCreateWidgetTool. This UClass inherits from UInteractiveToolPropertySet,
+ * which provides an OnModified delegate that the Tool will listen to for changes in property values.
+ */
+UCLASS(Transient)
+class WIDGETEDITORTOOLPALETTE_API UCreateWidgetToolProperties : public UInteractiveToolPropertySet
+{
+	GENERATED_BODY()
+public:
+	UCreateWidgetToolProperties();
+};
+
+/**
+ * UCreateWidgetTool is a tool that allows for a specific widget to be created on click
+ * It allows for fast widget creation via hotkeys. The specific OnClick behavior can be customized
+ * by specifying a different create widget tool builder class in editor per project user settings.
+ */
+UCLASS()
+class WIDGETEDITORTOOLPALETTE_API UCreateWidgetTool : public USingleClickTool, public IWidgetBaseBehavior
+{
+	GENERATED_BODY()
+
+public:
+	UCreateWidgetTool();
+
+	virtual void SetOwningToolkit(TSharedPtr<class IToolkit> InOwningToolkit);
+	virtual void SetOwningWidget(TSharedPtr<class SWidget> InOwningWidget);
+
+	// ~Begin UInteractiveTool Interface
+	virtual void Setup() override;
+	// ~End UInteractiveTool Interface
+
+	// ~Begin IWidgetBaseBehavior Interface
+	virtual bool OnMouseButtonUp(const FGeometry& MyGeometry, const FPointerEvent& MouseEvent) override;
+	// ~End IWidgetBaseBehavior Interface
+
+public:
+	/** Widget Class to be created by this tool */
+	UClass* WidgetClass;
+
+protected:
+	UPROPERTY()
+	UCreateWidgetToolProperties* Properties;
+
+protected:
+
+	// @TODO: DarenC - Find a good parent class or interface for these (as base class)
+	TWeakPtr<class FWidgetBlueprintEditor> OwningToolkit;
+	TWeakPtr<class SDesignerView> OwningWidget;
+};
\ No newline at end of file
diff --git a/Engine/Plugins/Experimental/WidgetEditorToolPalette/Source/WidgetEditorToolPalette/Public/DefaultTools/RectangleSelectTool.h b/Engine/Plugins/Experimental/WidgetEditorToolPalette/Source/WidgetEditorToolPalette/Public/DefaultTools/RectangleSelectTool.h
new file mode 100644
--- /dev/null
+++ b/Engine/Plugins/Experimental/WidgetEditorToolPalette/Source/WidgetEditorToolPalette/Public/DefaultTools/RectangleSelectTool.h
@@ -0,0 +1,85 @@
+// Copyright Epic Games, Inc. All Rights Reserved.
+
+#pragma once
+
+#include "CoreMinimal.h"
+#include "BaseTools/ClickDragTool.h"
+#include "BaseBehaviors/Widgets/WidgetBaseBehavior.h"
+#include "InteractiveToolBuilder.h"
+#include "MarqueeOperation.h"
+
+#include "RectangleSelectTool.generated.h"
+
+/**
+ * Builder for URectangleSelectTool
+ */
+UCLASS()
+class WIDGETEDITORTOOLPALETTE_API URectangleSelectToolBuilder : public UInteractiveToolBuilder
+{
+	GENERATED_BODY()
+
+public:
+	virtual bool CanBuildTool(const FToolBuilderState& SceneState) const override;
+	virtual UInteractiveTool* BuildTool(const FToolBuilderState& SceneState) const override;
+};
+
+/**
+ * Property set for the URectangleSelectTool
+ * 
+ * Note: Since we don't expose this for editing via details (yet), could instead be member variables on tool.
+ */
+UCLASS(Transient)
+class WIDGETEDITORTOOLPALETTE_API URectangleSelectProperties : public UInteractiveToolPropertySet
+{
+	GENERATED_BODY()
+
+public:
+	URectangleSelectProperties();
+
+	/** A pending marquee operation if it's active */
+	bool bIsMarqueeActive;
+
+	/** A pending marquee operation if it's active */
+	FMarqueeOperation Marquee;
+};
+
+/**
+ * URectangleSelectTool is a simple marque widget select tool.
+ */
+UCLASS()
+class WIDGETEDITORTOOLPALETTE_API URectangleSelectTool : public UInteractiveTool, public IWidgetBaseBehavior
+{
+	GENERATED_BODY()
+
+public:
+	virtual void SetOwningToolkit(TSharedPtr<class IToolkit> InOwningToolkit);
+	virtual void SetOwningWidget(TSharedPtr<class SWidget> InOwningWidget);
+
+	// ~Begin UInteractiveTool Interface
+	virtual void Setup() override;
+	// ~End UInteractiveTool Interface
+
+	// ~Begin IWidgetBaseBehavior Interface
+	virtual bool OnKeyDown(const FGeometry& MyGeometry, const FKeyEvent& InKeyEvent) override;
+	virtual bool OnMouseButtonDown(const FGeometry& MyGeometry, const FPointerEvent& MouseEvent) override;
+	virtual bool OnMouseButtonUp(const FGeometry& MyGeometry, const FPointerEvent& MouseEvent) override;
+	virtual bool OnMouseMove(const FGeometry& MyGeometry, const FPointerEvent& MouseEvent) override;
+	virtual int32 OnPaint(const FPaintArgs& Args, const FGeometry& AllottedGeometry, const FSlateRect& MyCullingRect, FSlateWindowElementList& OutDrawElements, int32 LayerId, const FWidgetStyle& InWidgetStyle, bool bParentEnabled) const override;
+	// ~End IWidgetBaseBehavior Interface
+
+protected:
+	/** Properties of the tool are stored here */
+	UPROPERTY()
+	URectangleSelectProperties* Properties;
+
+	/** Cached last end point to optimize comparisons */
+	FVector2D PrevSelectionEndPoint;
+
+protected:
+
+	// @TODO: DarenC - Find a good parent class or interface for these (as base class)
+	TWeakPtr<class FWidgetBlueprintEditor> OwningToolkit;
+	TWeakPtr<class SDesignerView> OwningWidget;
+
+	void SelectWidgetsAffectedByMarquee();
+};
diff --git a/Engine/Plugins/Experimental/WidgetEditorToolPalette/Source/WidgetEditorToolPalette/Public/Settings/CreateWidgetToolSettings.h b/Engine/Plugins/Experimental/WidgetEditorToolPalette/Source/WidgetEditorToolPalette/Public/Settings/CreateWidgetToolSettings.h
new file mode 100644
--- /dev/null
+++ b/Engine/Plugins/Experimental/WidgetEditorToolPalette/Source/WidgetEditorToolPalette/Public/Settings/CreateWidgetToolSettings.h
@@ -0,0 +1,81 @@
+// Copyright Epic Games, Inc. All Rghts Reserved.
+
+#pragma once
+
+#include "CoreMinimal.h"
+#include "UObject/ObjectMacros.h"
+#include "UObject/Object.h"
+#include "Engine/DeveloperSettings.h"
+#include "Engine/EngineTypes.h"
+#include "Components/Widget.h"
+#include "DefaultTools/CreateWidgetTool.h"
+#include "CreateWidgetToolSettings.generated.h"
+
+
+/** Info used to populate a single create widget tool */
+USTRUCT()
+struct WIDGETEDITORTOOLPALETTE_API FCreateWidgetToolInfo
+{
+	GENERATED_BODY()
+
+	FCreateWidgetToolInfo()
+		: WidgetClass(nullptr)
+		, DisplayName("")
+		, WidgetHotkey(FInputChord())
+		, CreateWidgetToolBuilder(UCreateWidgetToolBuilder::StaticClass())
+	{
+	}
+
+	/** The widget to create when this tool is activated */
+	UPROPERTY(EditAnywhere, config, Category = CreateWidgetInfo)
+	TSubclassOf<UWidget> WidgetClass;
+
+	/** Display name for this tool, if empty will use the name of the widget */
+	UPROPERTY(EditAnywhere, config, Category = CreateWidgetInfo)
+	FString DisplayName;
+
+	/** The hotkey used to create this widget */
+	UPROPERTY(EditAnywhere, config, Category = CreateWidgetInfo)
+	FInputChord WidgetHotkey;
+
+	/** Builder that handles creation of tool for this widget */
+	UPROPERTY(EditAnywhere, config, AdvancedDisplay, Category = CreateWidgetInfo)
+	TSubclassOf<UCreateWidgetToolBuilder> CreateWidgetToolBuilder;
+};
+
+/** Info used to populate a create widget tool stack */
+USTRUCT()
+struct WIDGETEDITORTOOLPALETTE_API FCreateWidgetStackInfo
+{
+	GENERATED_BODY()
+
+	/** Display name for this tool stack */
+	UPROPERTY(EditAnywhere, config, Category = CreateWidgetInfo)
+	FString DisplayName;
+
+	/** Info for each widget that can be created by this stack */
+	UPROPERTY(EditAnywhere, config, Category = CreateWidgetInfo)
+	TArray<FCreateWidgetToolInfo> WidgetToolInfos;
+};
+
+/**
+ * Allows for create widget tool layout to be defined via settings.
+ */
+UCLASS(config=EditorPerProjectUserSettings, autoexpandcategories=(Keybinds))
+class WIDGETEDITORTOOLPALETTE_API UCreateWidgetToolSettings : public UDeveloperSettings
+{
+	GENERATED_BODY()
+
+public:
+
+	virtual FName GetContainerName() const { return FName("Editor"); }
+	virtual FName GetCategoryName() const { return FName("Plugins"); }
+	virtual FName GetSectionName() const { return FName("CreateWidgetHotkeys"); }
+
+
+public:
+
+	/** List of tool stacks to create */
+	UPROPERTY(EditAnywhere, config, Category=Startup)
+	TArray<FCreateWidgetStackInfo> CreateWidgetStacks;
+};
diff --git a/Engine/Plugins/Experimental/WidgetEditorToolPalette/Source/WidgetEditorToolPalette/Public/WidgetEditorToolPaletteCommands.h b/Engine/Plugins/Experimental/WidgetEditorToolPalette/Source/WidgetEditorToolPalette/Public/WidgetEditorToolPaletteCommands.h
new file mode 100644
--- /dev/null
+++ b/Engine/Plugins/Experimental/WidgetEditorToolPalette/Source/WidgetEditorToolPalette/Public/WidgetEditorToolPaletteCommands.h
@@ -0,0 +1,23 @@
+// Copyright Epic Games, Inc. All Rights Reserved.
+
+#pragma once
+
+#include "Framework/Commands/Commands.h"
+
+class WIDGETEDITORTOOLPALETTE_API FWidgetEditorToolPaletteCommands : public TCommands<FWidgetEditorToolPaletteCommands>
+{
+public:
+	FWidgetEditorToolPaletteCommands();
+	void RegisterCommands() override;
+	static const FWidgetEditorToolPaletteCommands& Get();
+
+	/** Command to toggle tool palette in widget editor */
+	TSharedPtr<FUICommandInfo> ToggleWidgetEditorToolPalette;
+
+	TSharedPtr<FUICommandInfo> DefaultSelectTool;
+	TSharedPtr<FUICommandInfo> BeginRectangleSelectTool;
+
+	TMap<FString, TSharedPtr<FUICommandInfo>> CreateWidgetToolStacks;
+	TMap<FString, TSharedPtr<FUICommandInfo>> CreateWidgetTools;
+};
+
diff --git a/Engine/Plugins/Experimental/WidgetEditorToolPalette/Source/WidgetEditorToolPalette/Public/WidgetEditorToolPaletteMode.h b/Engine/Plugins/Experimental/WidgetEditorToolPalette/Source/WidgetEditorToolPalette/Public/WidgetEditorToolPaletteMode.h
new file mode 100644
--- /dev/null
+++ b/Engine/Plugins/Experimental/WidgetEditorToolPalette/Source/WidgetEditorToolPalette/Public/WidgetEditorToolPaletteMode.h
@@ -0,0 +1,30 @@
+// Copyright Epic Games, Inc. All Rights Reserved.
+
+#pragma once
+
+#include "Tools/LegacyEdModeWidgetHelpers.h"
+#include "Tools/Modes.h"
+#include "WidgetEditorToolPaletteMode.generated.h"
+
+class SBorder;
+
+UCLASS(Transient)
+class UWidgetEditorToolPaletteMode : public UBaseLegacyWidgetEdMode
+{
+	GENERATED_BODY()
+
+public:
+
+	const static FEditorModeID Id;
+
+protected:
+
+	UWidgetEditorToolPaletteMode();
+
+	void Enter() override;
+
+	bool UsesToolkits() const override; 
+	void CreateToolkit() override;
+	
+
+};
diff --git a/Engine/Plugins/Experimental/WidgetEditorToolPalette/Source/WidgetEditorToolPalette/Public/WidgetEditorToolPaletteModule.h b/Engine/Plugins/Experimental/WidgetEditorToolPalette/Source/WidgetEditorToolPalette/Public/WidgetEditorToolPaletteModule.h
new file mode 100644
--- /dev/null
+++ b/Engine/Plugins/Experimental/WidgetEditorToolPalette/Source/WidgetEditorToolPalette/Public/WidgetEditorToolPaletteModule.h
@@ -0,0 +1,25 @@
+// Copyright Epic Games, Inc. All Rights Reserved.
+
+#pragma once
+
+#include "Modules/ModuleInterface.h"
+#include "Templates/SharedPointer.h"
+
+class FWidgetBlueprintEditor;
+class SDockTab;
+
+class FWidgetEditorToolPaletteModule : public IModuleInterface
+{
+public:
+
+	/** IModuleInterface implementation */
+	virtual void StartupModule() override;
+	virtual void ShutdownModule() override;
+
+	bool IsWidgetEditorToolPaletteModeActive(TWeakPtr<FWidgetBlueprintEditor> InEditor) const;
+	void OnToggleWidgetEditorToolPaletteMode(TWeakPtr<FWidgetBlueprintEditor> InEditor);
+
+	void OnPostEngineInit();
+
+};
+
diff --git a/Engine/Plugins/Experimental/WidgetEditorToolPalette/Source/WidgetEditorToolPalette/Public/WidgetEditorToolPaletteStyle.h b/Engine/Plugins/Experimental/WidgetEditorToolPalette/Source/WidgetEditorToolPalette/Public/WidgetEditorToolPaletteStyle.h
new file mode 100644
--- /dev/null
+++ b/Engine/Plugins/Experimental/WidgetEditorToolPalette/Source/WidgetEditorToolPalette/Public/WidgetEditorToolPaletteStyle.h
@@ -0,0 +1,28 @@
+// Copyright Epic Games, Inc. All Rights Reserved.
+
+#pragma once
+
+#include "CoreMinimal.h"
+#include "Styling/ISlateStyle.h"
+#include "Styling/SlateStyle.h"
+
+class WIDGETEDITORTOOLPALETTE_API FWidgetEditorToolPaletteStyle
+{
+public:
+	static void Initialize();
+
+	static void Shutdown();
+
+	static TSharedPtr< class ISlateStyle > Get();
+
+	static FName GetStyleSetName();
+
+	// use to access icons defined by the style set by name, eg GetBrush("BrushFalloffIcons.Smooth")
+	static const FSlateBrush* GetBrush(FName PropertyName, const ANSICHAR* Specifier = NULL);
+
+private:
+	static FString InContent(const FString& RelativePath, const ANSICHAR* Extension);
+
+private:
+	static TSharedPtr< class FSlateStyleSet > StyleSet;
+};
diff --git a/Engine/Plugins/Experimental/WidgetEditorToolPalette/Source/WidgetEditorToolPalette/Public/WidgetEditorToolPaletteToolkit.h b/Engine/Plugins/Experimental/WidgetEditorToolPalette/Source/WidgetEditorToolPalette/Public/WidgetEditorToolPaletteToolkit.h
new file mode 100644
--- /dev/null
+++ b/Engine/Plugins/Experimental/WidgetEditorToolPalette/Source/WidgetEditorToolPalette/Public/WidgetEditorToolPaletteToolkit.h
@@ -0,0 +1,44 @@
+// Copyright Epic Games, Inc. All Rights Reserved.
+
+#pragma once
+
+#include "Toolkits/BaseToolkit.h"
+
+class STextBlock;
+
+class FWidgetEditorToolPaletteToolkit : public FModeToolkit
+{
+public:
+
+	void Init( const TSharedPtr<IToolkitHost>& InToolkitHost, TWeakObjectPtr<UEdMode> InOwningMode ) override;
+
+	//~ Begin IToolkit interface
+	virtual FName GetToolkitFName() const override;
+	virtual FText GetBaseToolkitName() const override;
+	//~ End IToolkit interface
+
+	//~ Begin FModeToolkit interface
+	virtual void GetToolPaletteNames(TArray<FName>& InPaletteName) const override;
+	virtual FText GetToolPaletteDisplayName(FName PaletteName) const override;
+	virtual void BuildToolPalette(FName PaletteName, class FToolBarBuilder& ToolbarBuilder) override;
+	virtual TSharedRef<SWidget> CreatePaletteWidget(TSharedPtr<FUICommandList> InCommandList, FName InToolbarCustomizationName, FName InPaletteName);
+	virtual void OnToolStarted(UInteractiveToolManager* Manager, UInteractiveTool* Tool) override;
+	virtual void OnToolEnded(UInteractiveToolManager* Manager, UInteractiveTool* Tool) override;
+	//~ End FModeToolkit interface
+
+private:
+
+	void UpdateActiveToolProperties(UInteractiveTool* Tool);
+
+	//~ Begin FModeToolkit interface
+	virtual FText GetActiveToolDisplayName() const override;
+	virtual FText GetActiveToolMessage() const override;
+	//~ End FModeToolkit interface
+
+	FText ActiveToolName;
+	FText ActiveToolMessage;
+
+};
+
+
+
diff --git a/Engine/Plugins/Experimental/WidgetEditorToolPalette/Source/WidgetEditorToolPalette/WidgetEditorToolPalette.Build.cs b/Engine/Plugins/Experimental/WidgetEditorToolPalette/Source/WidgetEditorToolPalette/WidgetEditorToolPalette.Build.cs
new file mode 100644
--- /dev/null
+++ b/Engine/Plugins/Experimental/WidgetEditorToolPalette/Source/WidgetEditorToolPalette/WidgetEditorToolPalette.Build.cs
@@ -0,0 +1,67 @@
+// Copyright Epic Games, Inc. All Rights Reserved.
+
+using UnrealBuildTool;
+
+public class WidgetEditorToolPalette : ModuleRules
+{
+	public WidgetEditorToolPalette(ReadOnlyTargetRules Target) : base(Target)
+	{
+		PCHUsage = ModuleRules.PCHUsageMode.UseExplicitOrSharedPCHs;
+		
+		PublicIncludePaths.AddRange(
+			new string[] {
+				// ... add public include paths required here ...
+			}
+			);
+				
+		
+		PrivateIncludePaths.AddRange(
+			new string[] {
+				// ... add other private include paths required here ...
+			}
+			);
+			
+		
+		PublicDependencyModuleNames.AddRange(
+			new string[]
+			{			
+				// ... add other public dependencies that you statically link with here ...
+			}
+			);
+			
+		
+		PrivateDependencyModuleNames.AddRange(
+			new string[]
+			{
+				"Core",
+				"InputCore",
+				"CoreUObject",
+				"DeveloperSettings",
+				"Engine",
+				"EditorInteractiveToolsFramework",
+				"EditorStyle",
+				"EditorFramework",
+				"GraphEditor",
+				"ContentBrowser",
+				"LevelEditor",
+				"Projects",
+				"InteractiveToolsFramework",
+				"Slate",
+				"SlateCore",
+				"UMG",
+				"UMGEditor",
+				"ToolMenus",
+				"UnrealEd",
+			}
+			);
+		
+		
+		DynamicallyLoadedModuleNames.AddRange(
+			new string[]
+			{
+				// ... add any modules that your module loads dynamically here ...
+			}
+			);
+
+	}
+}
diff --git a/Engine/Plugins/Experimental/WidgetEditorToolPalette/WidgetEditorToolPalette.uplugin b/Engine/Plugins/Experimental/WidgetEditorToolPalette/WidgetEditorToolPalette.uplugin
new file mode 100644
--- /dev/null
+++ b/Engine/Plugins/Experimental/WidgetEditorToolPalette/WidgetEditorToolPalette.uplugin
@@ -0,0 +1,25 @@
+{
+	"FileVersion": 3,
+	"Version": 1,
+	"VersionName": "0.1",
+	"FriendlyName": "Tool Palette for the Widget Editor",
+	"Description": "A set of tools to enhance UMG creation UX",
+	"Category": "Editor",
+	"CreatedBy": "Epic Games, Inc.",
+	"CreatedByURL": "http://epicgames.com",
+	"DocsURL": "",
+	"MarketplaceURL": "",
+	"SupportURL": "",
+	"CanContainContent": true,
+	"IsBetaVersion": true,
+	"Installed": false,
+	"Hidden": false,
+	"Modules": [
+		{
+			"Name": "WidgetEditorToolPalette",
+			"Type": "Editor",
+			"LoadingPhase": "Default",
+			"WhitelistPlatforms": [ "Win64" ]
+		}
+	]
+}
diff --git a/Engine/Source/Editor/EditorFramework/Public/Tools/AssetEditorContextInterface.h b/Engine/Source/Editor/EditorFramework/Public/Tools/AssetEditorContextInterface.h
--- a/Engine/Source/Editor/EditorFramework/Public/Tools/AssetEditorContextInterface.h
+++ b/Engine/Source/Editor/EditorFramework/Public/Tools/AssetEditorContextInterface.h
@@ -43,4 +43,14 @@ public:
 	 * Gets the world on which the asset editor performs actions.
 	 */
 	virtual UWorld* GetEditingWorld() const = 0;
+
+	/**
+	 * Get the toolkit host associated with our asset editor.
+	 */
+	virtual const IToolkitHost* GetToolkitHost() const = 0;
+
+	/**
+	 * Get the toolkit host associated with our asset editor.
+	 */
+	virtual IToolkitHost* GetMutableToolkitHost() = 0;
 };
diff --git a/Engine/Source/Editor/EditorStyle/Private/StarshipStyle.cpp b/Engine/Source/Editor/EditorStyle/Private/StarshipStyle.cpp
--- a/Engine/Source/Editor/EditorStyle/Private/StarshipStyle.cpp
+++ b/Engine/Source/Editor/EditorStyle/Private/StarshipStyle.cpp
@@ -1737,6 +1737,56 @@ void FStarshipEditorStyle::FStyle::SetupGeneralStyles()
 		Set("EditorModesPanel.ToolDescriptionFont", DEFAULT_FONT("Italic", 10));
 
 	}
+	// Vertical ToolPalette 
+	{
+		FToolBarStyle VerticalToolBarStyle = FStarshipCoreStyle::GetCoreStyle().GetWidgetStyle<FToolBarStyle>("SlimToolBar");
+
+		FTextBlockStyle VerticalToolBarLabelStyle = FTextBlockStyle(GetParentStyle()->GetWidgetStyle<FTextBlockStyle>("SmallText"));
+		VerticalToolBarLabelStyle.SetOverflowPolicy(ETextOverflowPolicy::Ellipsis);
+
+		VerticalToolBarStyle.SetLabelStyle(VerticalToolBarLabelStyle);
+		
+		VerticalToolBarStyle.SetBackground(FSlateColorBrush(FStyleColors::Recessed));
+
+		VerticalToolBarStyle.SetLabelPadding(FMargin(0.0f, 1.0f, 0.0f, 0.0f));
+
+		VerticalToolBarStyle.SetButtonPadding(       FMargin(0.0f, 0.0f));
+		VerticalToolBarStyle.SetCheckBoxPadding(     FMargin(0.0f, 0.0f));
+		VerticalToolBarStyle.SetComboButtonPadding(  FMargin(0.0f, 0.0f));
+		VerticalToolBarStyle.SetIndentedBlockPadding(FMargin(0.0f, 0.0f));
+		VerticalToolBarStyle.SetBlockPadding(        FMargin(0.0f, 0.0f));
+		VerticalToolBarStyle.SetBackgroundPadding(   FMargin(4.0f, 2.0f));
+		VerticalToolBarStyle.ToggleButton.SetPadding(FMargin(0.0f, 6.0f));
+		VerticalToolBarStyle.ButtonStyle.SetNormalPadding(FMargin(2.0f, 6.0f));
+		VerticalToolBarStyle.ButtonStyle.SetPressedPadding(FMargin(2.0f, 6.0f));
+
+		Set( "VerticalToolBar.Tab",  FCheckBoxStyle()
+			.SetCheckBoxType(            ESlateCheckBoxType::ToggleButton)
+
+			.SetCheckedImage(            FSlateRoundedBoxBrush(FStyleColors::Input, 2.0f))
+			.SetCheckedHoveredImage(     FSlateRoundedBoxBrush(FStyleColors::Input, 2.0f))
+			.SetCheckedPressedImage(     FSlateRoundedBoxBrush(FStyleColors::Input, 2.0f))
+
+			.SetUncheckedImage(          FSlateRoundedBoxBrush(FStyleColors::Secondary, 2.0f))
+			.SetUncheckedHoveredImage(   FSlateRoundedBoxBrush(FStyleColors::Hover, 2.0f))
+			.SetUncheckedPressedImage(   FSlateRoundedBoxBrush(FStyleColors::Secondary, 2.0f))
+
+			.SetForegroundColor(         FStyleColors::Foreground)
+			.SetHoveredForegroundColor(  FStyleColors::ForegroundHover)
+			.SetPressedForegroundColor(  FStyleColors::ForegroundHover)
+			.SetCheckedForegroundColor(  FStyleColors::Primary)
+			.SetCheckedHoveredForegroundColor(FStyleColors::PrimaryHover)
+			.SetPadding(FMargin(2.f, 6.f))
+		);
+
+		Set("VerticalToolBar.MaxUniformToolbarSize", 48.f);
+		Set("VerticalToolBar.MinUniformToolbarSize", 48.f);
+
+		Set("VerticalToolBar.ExpandableAreaHeader", new FSlateRoundedBoxBrush(FStyleColors::Dropdown, FVector4(4.0, 4.0, 0.0, 0.0)));
+		Set("VerticalToolBar.ExpandableAreaBody", new FSlateRoundedBoxBrush(FStyleColors::Recessed, FVector4(0.0, 0.0, 4.0, 4.0)));
+
+		Set("VerticalToolBar", VerticalToolBarStyle);
+	}
 	
 	// Ctrl+Tab menu
 	{
diff --git a/Engine/Source/Editor/Experimental/EditorInteractiveToolsFramework/Private/EdModeInteractiveToolsContext.cpp b/Engine/Source/Editor/Experimental/EditorInteractiveToolsFramework/Private/EdModeInteractiveToolsContext.cpp
--- a/Engine/Source/Editor/Experimental/EditorInteractiveToolsFramework/Private/EdModeInteractiveToolsContext.cpp
+++ b/Engine/Source/Editor/Experimental/EditorInteractiveToolsFramework/Private/EdModeInteractiveToolsContext.cpp
@@ -394,7 +394,7 @@ void UEditorInteractiveToolsContext::Tick(FEditorViewportClient* ViewportClient,
 {
 	// invalidate this viewport if it's timestamp is not current
 	const int32* FoundTimestamp = InvalidationMap.Find(ViewportClient);
-	if (FoundTimestamp == nullptr)
+	if (FoundTimestamp == nullptr && ViewportClient)
 	{
 		ViewportClient->Invalidate(false, false);
 		InvalidationMap.Add(ViewportClient, InvalidationTimestamp);
@@ -429,7 +429,10 @@ void UEditorInteractiveToolsContext::Tick(FEditorViewportClient* ViewportClient,
 	// Cache current camera state from this Viewport in the ContextQueries, which we will use for things like snapping/etc that
 	// is computed by the Tool and Gizmo Tick()s
 	// (This is not necessarily correct for Hover, because we might be Hovering over a different Viewport than the Active one...)
-	((FEdModeToolsContextQueriesImpl*)this->QueriesAPI)->CacheCurrentViewState(ViewportClient);
+	if (ViewportClient)
+	{
+		((FEdModeToolsContextQueriesImpl*)this->QueriesAPI)->CacheCurrentViewState(ViewportClient);
+	}
 
 	// tick our stuff
 	ToolManager->Tick(DeltaTime);
@@ -474,6 +477,11 @@ public:
 
 	void CacheCurrentViewState(FViewport* Viewport, FEditorViewportClient* ViewportClient)
 	{
+		if (!ViewportClient)
+		{
+			return;
+		}
+
 		FViewportCameraTransform ViewTransform = ViewportClient->GetViewTransform();
 		ViewCameraState.bIsOrthographic = ViewportClient->IsOrtho();
 		ViewCameraState.Position = ViewTransform.GetLocation();
@@ -631,6 +639,11 @@ bool UEditorInteractiveToolsContext::ProcessEditDelete()
 
 FRay UEditorInteractiveToolsContext::GetRayFromMousePos(FEditorViewportClient* ViewportClient, FViewport* Viewport, int MouseX, int MouseY)
 {
+	if (!ViewportClient)
+	{
+		return FRay();
+	}
+
 	FSceneViewFamilyContext ViewFamily(FSceneViewFamily::ConstructionValues(
 		ViewportClient->Viewport,
 		ViewportClient->GetScene(),
@@ -871,16 +884,19 @@ bool UModeManagerInteractiveToolsContext::InputKey(FEditorViewportClient* Viewpo
 				// to wherever they get handled.
 				// Someday these kinds of prioritizations will be handled by having camera manipulation be
 				// in a common input router so that behavior priorities can determine the ordering.
-				if (ViewportClient->IsAltPressed() && InputRouter->HasActiveMouseCapture() == false)
+				if (ViewportClient && ViewportClient->IsAltPressed() && InputRouter->HasActiveMouseCapture() == false)
 				{
 					return false;
 				}
 
 				FInputDeviceState InputState = CurrentMouseState;
 				InputState.InputDevice = EInputDevices::Mouse;
+
+				FModifierKeysState ModifierKeys = FSlateApplication::Get().GetModifierKeys();
 				InputState.SetModifierKeyStates(
-					ViewportClient->IsShiftPressed(), ViewportClient->IsAltPressed(),
-					ViewportClient->IsCtrlPressed(), ViewportClient->IsCmdPressed());
+					ModifierKeys.IsShiftDown(), ModifierKeys.IsAltDown(),
+					ModifierKeys.IsControlDown(), ModifierKeys.IsCommandDown());
+
 				if (bIsLeftMouse)
 				{
 					InputState.Mouse.Left.SetStates(
@@ -917,9 +933,11 @@ bool UModeManagerInteractiveToolsContext::InputKey(FEditorViewportClient* Viewpo
 
 				FInputDeviceState InputState = CurrentMouseState;
 				InputState.InputDevice = EInputDevices::Mouse;
+
+				FModifierKeysState ModifierKeys = FSlateApplication::Get().GetModifierKeys();
 				InputState.SetModifierKeyStates(
-					ViewportClient->IsShiftPressed(), ViewportClient->IsAltPressed(),
-					ViewportClient->IsCtrlPressed(), ViewportClient->IsCmdPressed());
+					ModifierKeys.IsShiftDown(), ModifierKeys.IsAltDown(),
+					ModifierKeys.IsControlDown(), ModifierKeys.IsCommandDown());
 
 				InputState.Mouse.WheelDelta = (Event != IE_Pressed) ? 0
 					: (Key == EKeys::MouseScrollUp) ? 1 
@@ -944,9 +962,12 @@ bool UModeManagerInteractiveToolsContext::InputKey(FEditorViewportClient* Viewpo
 		{
 			FInputDeviceState InputState;
 			InputState.InputDevice = EInputDevices::Keyboard;
+
+			FModifierKeysState ModifierKeys = FSlateApplication::Get().GetModifierKeys();
 			InputState.SetModifierKeyStates(
-				ViewportClient->IsShiftPressed(), ViewportClient->IsAltPressed(),
-				ViewportClient->IsCtrlPressed(), ViewportClient->IsCmdPressed());
+				ModifierKeys.IsShiftDown(), ModifierKeys.IsAltDown(),
+				ModifierKeys.IsControlDown(), ModifierKeys.IsCommandDown());
+
 			InputState.Keyboard.ActiveKey.Button = Key;
 			bool bPressed = (Event == IE_Pressed);
 			InputState.Keyboard.ActiveKey.SetStates(bPressed, bPressed, !bPressed);
@@ -981,9 +1002,10 @@ bool UModeManagerInteractiveToolsContext::MouseMove(FEditorViewportClient* Viewp
 	FInputDeviceState InputState = CurrentMouseState;
 	InputState.InputDevice = EInputDevices::Mouse;
 
+	FModifierKeysState ModifierKeys = FSlateApplication::Get().GetModifierKeys();
 	InputState.SetModifierKeyStates(
-		ViewportClient->IsShiftPressed(), ViewportClient->IsAltPressed(),
-		ViewportClient->IsCtrlPressed(), ViewportClient->IsCmdPressed());
+		ModifierKeys.IsShiftDown(), ModifierKeys.IsAltDown(),
+		ModifierKeys.IsControlDown(), ModifierKeys.IsCommandDown());
 
 	if (InputRouter->HasActiveMouseCapture())
 	{
@@ -1033,9 +1055,12 @@ bool UModeManagerInteractiveToolsContext::CapturedMouseMove(FEditorViewportClien
 
 		FInputDeviceState InputState = CurrentMouseState;
 		InputState.InputDevice = EInputDevices::Mouse;
+		
+		FModifierKeysState ModifierKeys = FSlateApplication::Get().GetModifierKeys();
 		InputState.SetModifierKeyStates(
-			InViewportClient->IsShiftPressed(), InViewportClient->IsAltPressed(),
-			InViewportClient->IsCtrlPressed(), InViewportClient->IsCmdPressed());
+			ModifierKeys.IsShiftDown(), ModifierKeys.IsAltDown(),
+			ModifierKeys.IsControlDown(), ModifierKeys.IsCommandDown());
+
 		InputState.Mouse.Delta2D = CurrentMouseState.Mouse.Position2D - OldPosition;
 		InputRouter->PostInputEvent(InputState);
 		return true;
@@ -1051,7 +1076,10 @@ bool UModeManagerInteractiveToolsContext::EndTracking(FEditorViewportClient* InV
 		// If the input router captured the mouse input, we need to invalidate the viewport client here, since the mouse delta tracker's end tracking will not be called.
 		constexpr bool bForceChildViewportRedraw = true;
 		constexpr bool bInvalidateHitProxies = true;
-		InViewportClient->Invalidate(bForceChildViewportRedraw, bInvalidateHitProxies);
+		if (InViewportClient)
+		{
+			InViewportClient->Invalidate(bForceChildViewportRedraw, bInvalidateHitProxies);
+		}
 		bIsTrackingMouse = false;
 		return true;
 	}
diff --git a/Engine/Source/Editor/Experimental/EditorInteractiveToolsFramework/Private/ToolContexts/WidgetToolsContext.cpp b/Engine/Source/Editor/Experimental/EditorInteractiveToolsFramework/Private/ToolContexts/WidgetToolsContext.cpp
new file mode 100644
--- /dev/null
+++ b/Engine/Source/Editor/Experimental/EditorInteractiveToolsFramework/Private/ToolContexts/WidgetToolsContext.cpp
@@ -0,0 +1,183 @@
+// Copyright Epic Games, Inc. All Rights Reserved.
+
+
+#include "ToolContexts/WidgetToolsContext.h"
+#include "BaseBehaviors/Widgets/WidgetBaseBehavior.h"
+
+#include "InteractiveToolManager.h"
+
+
+
+bool UWidgetToolsContext::OnKeyChar(const FGeometry& MyGeometry, const FCharacterEvent& InCharacterEvent)
+{
+	for (TObjectPtr<UEdModeInteractiveToolsContext> EdModeContext : EdModeToolsContexts)
+	{
+		if (EdModeContext && EdModeContext->ToolManager)
+		{
+			if (IWidgetBaseBehavior* WidgetTool = Cast<IWidgetBaseBehavior>(EdModeContext->ToolManager->GetActiveTool(EToolSide::Mouse)))
+			{
+				return WidgetTool->OnKeyChar(MyGeometry, InCharacterEvent);
+			}
+		}
+	}
+
+	return false;
+}
+
+bool UWidgetToolsContext::OnKeyDown(const FGeometry& MyGeometry, const FKeyEvent& InKeyEvent)
+{
+	for (TObjectPtr<UEdModeInteractiveToolsContext> EdModeContext : EdModeToolsContexts)
+	{
+		if (EdModeContext && EdModeContext->ToolManager)
+		{
+			if (IWidgetBaseBehavior* WidgetTool = Cast<IWidgetBaseBehavior>(EdModeContext->ToolManager->GetActiveTool(EToolSide::Mouse)))
+			{
+				return WidgetTool->OnKeyDown(MyGeometry, InKeyEvent);
+			}
+		}
+	}
+
+	return false;
+}
+
+bool UWidgetToolsContext::OnKeyUp(const FGeometry& MyGeometry, const FKeyEvent& InKeyEvent)
+{
+	for (TObjectPtr<UEdModeInteractiveToolsContext> EdModeContext : EdModeToolsContexts)
+	{
+		if (EdModeContext && EdModeContext->ToolManager)
+		{
+			if (IWidgetBaseBehavior* WidgetTool = Cast<IWidgetBaseBehavior>(EdModeContext->ToolManager->GetActiveTool(EToolSide::Mouse)))
+			{
+				return WidgetTool->OnKeyUp(MyGeometry, InKeyEvent);
+			}
+		}
+	}
+
+	return false;
+}
+
+bool UWidgetToolsContext::OnMouseButtonDown(const FGeometry& MyGeometry, const FPointerEvent& MouseEvent)
+{
+	for (TObjectPtr<UEdModeInteractiveToolsContext> EdModeContext : EdModeToolsContexts)
+	{
+		if (EdModeContext && EdModeContext->ToolManager)
+		{
+			if (IWidgetBaseBehavior* WidgetTool = Cast<IWidgetBaseBehavior>(EdModeContext->ToolManager->GetActiveTool(EToolSide::Mouse)))
+			{
+				return WidgetTool->OnMouseButtonDown(MyGeometry, MouseEvent);
+			}
+		}
+	}
+
+	return false;
+}
+
+bool UWidgetToolsContext::OnMouseButtonUp(const FGeometry& MyGeometry, const FPointerEvent& MouseEvent)
+{
+	for (TObjectPtr<UEdModeInteractiveToolsContext> EdModeContext : EdModeToolsContexts)
+	{
+		if (EdModeContext && EdModeContext->ToolManager)
+		{
+			if (IWidgetBaseBehavior* WidgetTool = Cast<IWidgetBaseBehavior>(EdModeContext->ToolManager->GetActiveTool(EToolSide::Mouse)))
+			{
+				return WidgetTool->OnMouseButtonUp(MyGeometry, MouseEvent);
+			}
+		}
+	}
+
+	return false;
+}
+
+bool UWidgetToolsContext::OnMouseMove(const FGeometry& MyGeometry, const FPointerEvent& MouseEvent)
+{
+	for (TObjectPtr<UEdModeInteractiveToolsContext> EdModeContext : EdModeToolsContexts)
+	{
+		if (EdModeContext && EdModeContext->ToolManager)
+		{
+			if (IWidgetBaseBehavior* WidgetTool = Cast<IWidgetBaseBehavior>(EdModeContext->ToolManager->GetActiveTool(EToolSide::Mouse)))
+			{
+				return WidgetTool->OnMouseMove(MyGeometry, MouseEvent);
+			}
+		}
+	}
+
+	return false;
+}
+
+bool UWidgetToolsContext::OnMouseWheel(const FGeometry& MyGeometry, const FPointerEvent& MouseEvent)
+{
+	for (TObjectPtr<UEdModeInteractiveToolsContext> EdModeContext : EdModeToolsContexts)
+	{
+		if (EdModeContext && EdModeContext->ToolManager)
+		{
+			if (IWidgetBaseBehavior* WidgetTool = Cast<IWidgetBaseBehavior>(EdModeContext->ToolManager->GetActiveTool(EToolSide::Mouse)))
+			{
+				return WidgetTool->OnMouseWheel(MyGeometry, MouseEvent);
+			}
+		}
+	}
+
+	return false;
+}
+
+bool UWidgetToolsContext::OnMouseButtonDoubleClick(const FGeometry& InMyGeometry, const FPointerEvent& InMouseEvent)
+{
+	for (TObjectPtr<UEdModeInteractiveToolsContext> EdModeContext : EdModeToolsContexts)
+	{
+		if (EdModeContext && EdModeContext->ToolManager)
+		{
+			if (IWidgetBaseBehavior* WidgetTool = Cast<IWidgetBaseBehavior>(EdModeContext->ToolManager->GetActiveTool(EToolSide::Mouse)))
+			{
+				return WidgetTool->OnMouseButtonDoubleClick(InMyGeometry, InMouseEvent);
+			}
+		}
+	}
+
+	return false;
+}
+
+bool UWidgetToolsContext::OnDragDetected(const FGeometry& MyGeometry, const FPointerEvent& MouseEvent)
+{
+	for (TObjectPtr<UEdModeInteractiveToolsContext> EdModeContext : EdModeToolsContexts)
+	{
+		if (EdModeContext && EdModeContext->ToolManager)
+		{
+			if (IWidgetBaseBehavior* WidgetTool = Cast<IWidgetBaseBehavior>(EdModeContext->ToolManager->GetActiveTool(EToolSide::Mouse)))
+			{
+				return WidgetTool->OnDragDetected(MyGeometry, MouseEvent);
+			}
+		}
+	}
+
+	return false;
+}
+
+void UWidgetToolsContext::OnMouseCaptureLost(const FCaptureLostEvent& CaptureLostEvent)
+{
+	for (TObjectPtr<UEdModeInteractiveToolsContext> EdModeContext : EdModeToolsContexts)
+	{
+		if (EdModeContext && EdModeContext->ToolManager)
+		{
+			if (IWidgetBaseBehavior* WidgetTool = Cast<IWidgetBaseBehavior>(EdModeContext->ToolManager->GetActiveTool(EToolSide::Mouse)))
+			{
+				WidgetTool->OnMouseCaptureLost(CaptureLostEvent);
+			}
+		}
+	}
+}
+
+int32 UWidgetToolsContext::OnPaint(const FPaintArgs& Args, const FGeometry& AllottedGeometry, const FSlateRect& MyCullingRect, FSlateWindowElementList& OutDrawElements, int32 LayerId, const FWidgetStyle& InWidgetStyle, bool bParentEnabled) const
+{
+	for (TObjectPtr<UEdModeInteractiveToolsContext> EdModeContext : EdModeToolsContexts)
+	{
+		if (EdModeContext && EdModeContext->ToolManager)
+		{
+			if (IWidgetBaseBehavior* WidgetTool = Cast<IWidgetBaseBehavior>(EdModeContext->ToolManager->GetActiveTool(EToolSide::Mouse)))
+			{
+				return WidgetTool->OnPaint(Args, AllottedGeometry, MyCullingRect, OutDrawElements, LayerId, InWidgetStyle, bParentEnabled);
+			}
+		}
+	}
+
+	return LayerId;
+}
\ No newline at end of file
diff --git a/Engine/Source/Editor/Experimental/EditorInteractiveToolsFramework/Public/ToolContexts/WidgetToolsContext.h b/Engine/Source/Editor/Experimental/EditorInteractiveToolsFramework/Public/ToolContexts/WidgetToolsContext.h
new file mode 100644
--- /dev/null
+++ b/Engine/Source/Editor/Experimental/EditorInteractiveToolsFramework/Public/ToolContexts/WidgetToolsContext.h
@@ -0,0 +1,36 @@
+// Copyright Epic Games, Inc. All Rights Reserved.
+
+#pragma once
+
+#include "EdModeInteractiveToolsContext.h"
+#include "Layout/Geometry.h"
+#include "Input/Events.h"
+#include "Rendering/DrawElements.h"
+
+#include "WidgetToolsContext.generated.h"
+
+/**
+ * UWidgetToolsContext extends UModeManagerInteractiveToolsContext with methods needed for
+ * tools operating on general widgets that do not have a viewport.
+ */
+UCLASS(Transient)
+class EDITORINTERACTIVETOOLSFRAMEWORK_API UWidgetToolsContext : public UModeManagerInteractiveToolsContext
+{
+	GENERATED_BODY()
+
+public:
+
+	// ~Begin SWidget "Proxy" Interface
+	bool OnKeyChar(const FGeometry& MyGeometry, const FCharacterEvent& InCharacterEvent);
+	bool OnKeyDown(const FGeometry& MyGeometry, const FKeyEvent& InKeyEvent);
+	bool OnKeyUp(const FGeometry& MyGeometry, const FKeyEvent& InKeyEvent);
+	bool OnMouseButtonDown(const FGeometry& MyGeometry, const FPointerEvent& MouseEvent);
+	bool OnMouseButtonUp(const FGeometry& MyGeometry, const FPointerEvent& MouseEvent);
+	bool OnMouseMove(const FGeometry& MyGeometry, const FPointerEvent& MouseEvent);
+	bool OnMouseWheel(const FGeometry& MyGeometry, const FPointerEvent& MouseEvent);
+	bool OnMouseButtonDoubleClick(const FGeometry& InMyGeometry, const FPointerEvent& InMouseEvent);
+	bool OnDragDetected(const FGeometry& MyGeometry, const FPointerEvent& MouseEvent);
+	void OnMouseCaptureLost(const FCaptureLostEvent& CaptureLostEvent);
+	int32 OnPaint(const FPaintArgs& Args, const FGeometry& AllottedGeometry, const FSlateRect& MyCullingRect, FSlateWindowElementList& OutDrawElements, int32 LayerId, const FWidgetStyle& InWidgetStyle, bool bParentEnabled) const;
+	// ~End SWidget "Proxy" Interface
+};
\ No newline at end of file
diff --git a/Engine/Source/Editor/UMGEditor/Private/BlueprintModes/WidgetDesignerApplicationMode.cpp b/Engine/Source/Editor/UMGEditor/Private/BlueprintModes/WidgetDesignerApplicationMode.cpp
--- a/Engine/Source/Editor/UMGEditor/Private/BlueprintModes/WidgetDesignerApplicationMode.cpp
+++ b/Engine/Source/Editor/UMGEditor/Private/BlueprintModes/WidgetDesignerApplicationMode.cpp
@@ -12,6 +12,7 @@
 #include "StatusBarSubsystem.h"
 #include "ToolMenus.h"
 #include "WidgetDrawerConfig.h"
+#include "Widgets/Docking/SDockTab.h"
 
 #include "TabFactory/PaletteTabSummoner.h"
 #include "TabFactory/LibraryTabSummoner.h"
@@ -45,7 +46,7 @@ FWidgetDesignerApplicationMode::FWidgetDesignerApplicationMode(TSharedPtr<FWidge
 		DisplayMetrics.PrimaryDisplayWorkAreaRect.Bottom - DisplayMetrics.PrimaryDisplayWorkAreaRect.Top);
 	const FVector2D WindowSize = (CenterScale * DisplaySize) / DPIScale;
 
-	TabLayout = FTabManager::NewLayout("WidgetBlueprintEditor_Designer_Layout_v4_7")
+	TabLayout = FTabManager::NewLayout("WidgetBlueprintEditor_Designer_Layout_v4_8")
 	->AddArea
 	(
 		FTabManager::NewPrimaryArea()
@@ -86,7 +87,14 @@ FWidgetDesignerApplicationMode::FWidgetDesignerApplicationMode(TSharedPtr<FWidge
 				(
 					FTabManager::NewStack()
 					->SetHideTabWell(true)
-					->SetSizeCoefficient(0.85f)
+					->SetSizeCoefficient(0.05f)
+					->AddTab(FDesignerTabSummoner::ToolPaletteTabID, ETabState::ClosedTab)
+				)
+				->Split
+				(
+					FTabManager::NewStack()
+					->SetHideTabWell(true)
+					->SetSizeCoefficient(0.80f)
 					->AddTab( FDesignerTabSummoner::TabID, ETabState::OpenedTab )
 				)
 				->Split
@@ -131,11 +139,18 @@ FWidgetDesignerApplicationMode::FWidgetDesignerApplicationMode(TSharedPtr<FWidge
 	TabFactories.RegisterFactory(MakeShareable(new FCompilerResultsSummoner(InWidgetEditor)));
 	TabFactories.RegisterFactory(MakeShareable(new FNavigationTabSummoner(InWidgetEditor)));
 
-	IUMGEditorModule& EditorModule = FModuleManager::GetModuleChecked<IUMGEditorModule>("UMGEditor");
-	EditorModule.OnRegisterTabsForEditor().Broadcast(*this, TabFactories);
+	IUMGEditorModule& UMGEditorModule = FModuleManager::LoadModuleChecked<IUMGEditorModule>("UMGEditor");
+	UMGEditorModule.OnRegisterTabsForEditor().Broadcast(*this, TabFactories);
+
+	// Add any extenders specified by the UMG Editor Module
+	// Note: Used by WidgetEditorModeUILayer to register the toolbox tab
+	if (LayoutExtender)
+	{
+		UMGEditorModule.OnRegisterLayoutExtensions().Broadcast(*LayoutExtender);
+		TabLayout->ProcessExtensions(*LayoutExtender);
+	}
 
 	//Make sure we start with our existing list of extenders instead of creating a new one
-	IUMGEditorModule& UMGEditorModule = FModuleManager::LoadModuleChecked<IUMGEditorModule>("UMGEditor");
 	ToolbarExtender = UMGEditorModule.GetToolBarExtensibilityManager()->GetAllExtenders();
 	
 	InWidgetEditor->GetWidgetToolbarBuilder()->AddWidgetBlueprintEditorModesToolbar(ToolbarExtender);
@@ -143,6 +158,7 @@ FWidgetDesignerApplicationMode::FWidgetDesignerApplicationMode(TSharedPtr<FWidge
 	if (UToolMenu* Toolbar = InWidgetEditor->RegisterModeToolbarIfUnregistered(GetModeName()))
 	{
 		InWidgetEditor->GetWidgetToolbarBuilder()->AddWidgetReflector(Toolbar);
+		InWidgetEditor->GetWidgetToolbarBuilder()->AddToolPalettes(Toolbar);
 		InWidgetEditor->GetToolbarBuilder()->AddCompileToolbar(Toolbar);
 		InWidgetEditor->GetToolbarBuilder()->AddDebuggingToolbar(Toolbar);
 	}
@@ -159,6 +175,15 @@ void FWidgetDesignerApplicationMode::RegisterTabFactories(TSharedPtr<FTabManager
 void FWidgetDesignerApplicationMode::PreDeactivateMode()
 {
 	//FWidgetBlueprintApplicationMode::PreDeactivateMode();
+
+	TSharedPtr<FWidgetBlueprintEditor> BP = GetBlueprintEditor();
+	if (TSharedPtr<FTabManager> TabManager = BP->GetTabManager())
+	{
+		if (TSharedPtr<SDockTab> PaletteTab = TabManager->FindExistingLiveTab(FDesignerTabSummoner::ToolPaletteTabID))
+		{
+			PaletteTab->RequestCloseTab();
+		}
+	}
 }
 
 void FWidgetDesignerApplicationMode::PostActivateMode()
@@ -176,6 +201,19 @@ void FWidgetDesignerApplicationMode::PostActivateMode()
 	WidgetAnimSequencerDrawer.Icon = FAppStyle::Get().GetBrush("UMGEditor.AnimTabIcon");
 	BP->RegisterDrawer(MoveTemp(WidgetAnimSequencerDrawer), 1);
 
+	// Toggle any active tool palette modes
+	for (TSharedPtr<FUICommandInfo>& Command : BP->ToolPaletteCommands)
+	{
+		if (Command && BP->GetToolkitCommands()->GetCheckState(Command.ToSharedRef()) == ECheckBoxState::Checked)
+		{
+			BP->GetToolkitCommands()->TryExecuteAction(Command.ToSharedRef());
+
+			// @FIXME: DarenC - Executing twice since a single action is a toggle, 
+			// need toggle twice if in incorrect state. This is not relevant when we move to drop downs.
+			BP->GetToolkitCommands()->TryExecuteAction(Command.ToSharedRef());
+		}
+	}
+
 	BP->OnEnteringDesigner();
 }
 
diff --git a/Engine/Source/Editor/UMGEditor/Private/BlueprintModes/WidgetGraphApplicationMode.cpp b/Engine/Source/Editor/UMGEditor/Private/BlueprintModes/WidgetGraphApplicationMode.cpp
--- a/Engine/Source/Editor/UMGEditor/Private/BlueprintModes/WidgetGraphApplicationMode.cpp
+++ b/Engine/Source/Editor/UMGEditor/Private/BlueprintModes/WidgetGraphApplicationMode.cpp
@@ -16,7 +16,7 @@
 FWidgetGraphApplicationMode::FWidgetGraphApplicationMode(TSharedPtr<FWidgetBlueprintEditor> InWidgetEditor)
 	: FWidgetBlueprintApplicationMode(InWidgetEditor, FWidgetBlueprintApplicationModes::GraphMode)
 {
-	TabLayout = FTabManager::NewLayout( "WidgetBlueprintEditor_Graph_Layout_v2" )
+	TabLayout = FTabManager::NewLayout( "WidgetBlueprintEditor_Graph_Layout_v2x1" )
 		->AddArea
 		(
 			FTabManager::NewPrimaryArea() ->SetOrientation(Orient_Vertical)
@@ -69,8 +69,16 @@ FWidgetGraphApplicationMode::FWidgetGraphApplicationMode(TSharedPtr<FWidgetBluep
 			)
 		);
 	
-	IUMGEditorModule& EditorModule = FModuleManager::GetModuleChecked<IUMGEditorModule>("UMGEditor");
-	EditorModule.OnRegisterTabsForEditor().Broadcast(*this, TabFactories);
+	IUMGEditorModule& UMGEditorModule = FModuleManager::GetModuleChecked<IUMGEditorModule>("UMGEditor");
+	UMGEditorModule.OnRegisterTabsForEditor().Broadcast(*this, TabFactories);
+
+	// Add any extenders specified by the UMG Editor Module
+	// Note: Used by WidgetEditorModeUILayer to register the toolbox tab
+	if (LayoutExtender)
+	{
+		UMGEditorModule.OnRegisterLayoutExtensions().Broadcast(*LayoutExtender);
+		TabLayout->ProcessExtensions(*LayoutExtender);
+	}
 
 	ToolbarExtender = MakeShareable(new FExtender);
 	InWidgetEditor->GetWidgetToolbarBuilder()->AddWidgetBlueprintEditorModesToolbar(ToolbarExtender);
diff --git a/Engine/Source/Editor/UMGEditor/Private/Designer/SDesignSurface.h b/Engine/Source/Editor/UMGEditor/Private/Designer/SDesignSurface.h
deleted file mode 100644
diff --git a/Engine/Source/Editor/UMGEditor/Private/Designer/SDesignerView.h b/Engine/Source/Editor/UMGEditor/Private/Designer/SDesignerView.h
deleted file mode 100644
diff --git a/Engine/Source/Editor/UMGEditor/Private/DragDrop/WidgetTemplateDragDropOp.h b/Engine/Source/Editor/UMGEditor/Private/DragDrop/WidgetTemplateDragDropOp.h
deleted file mode 100644
diff --git a/Engine/Source/Editor/UMGEditor/Private/Hierarchy/SHierarchyViewItem.cpp b/Engine/Source/Editor/UMGEditor/Private/Hierarchy/SHierarchyViewItem.cpp
--- a/Engine/Source/Editor/UMGEditor/Private/Hierarchy/SHierarchyViewItem.cpp
+++ b/Engine/Source/Editor/UMGEditor/Private/Hierarchy/SHierarchyViewItem.cpp
@@ -34,8 +34,8 @@
 #include "WidgetBlueprintEditorUtils.h"
 #include "ScopedTransaction.h"
 #include "Styling/SlateIconFinder.h"
-#include "WidgetTemplateBlueprintClass.h"
-#include "WidgetTemplateImageClass.h"
+#include "Templates/WidgetTemplateBlueprintClass.h"
+#include "Templates/WidgetTemplateImageClass.h"
 
 #define LOCTEXT_NAMESPACE "UMG"
 
diff --git a/Engine/Source/Editor/UMGEditor/Private/TabFactory/DesignerTabSummoner.cpp b/Engine/Source/Editor/UMGEditor/Private/TabFactory/DesignerTabSummoner.cpp
--- a/Engine/Source/Editor/UMGEditor/Private/TabFactory/DesignerTabSummoner.cpp
+++ b/Engine/Source/Editor/UMGEditor/Private/TabFactory/DesignerTabSummoner.cpp
@@ -1,13 +1,23 @@
 // Copyright Epic Games, Inc. All Rights Reserved.
 
 #include "TabFactory/DesignerTabSummoner.h"
+#include "Tools/ToolCompatible.h"
 #include "Designer/SDesignerView.h"
-#include "UMGStyle.h"
+#include "Toolkits/AssetEditorModeUILayer.h"
+#include "UMGStyle.h" 
 
 #define LOCTEXT_NAMESPACE "UMG"
 
+const FName FDesignerTabSummoner::ToolPaletteTabID = UAssetEditorUISubsystem::TopLeftTabID;
 const FName FDesignerTabSummoner::TabID(TEXT("SlatePreview"));
 
+static int32 bEnableWidgetDesignerTools = 1;
+static FAutoConsoleVariableRef CVarEnableWidgetDesignerTools(
+	TEXT("UMG.Editor.EnableWidgetDesignerTools"),
+	bEnableWidgetDesignerTools,
+	TEXT("Toggles processing of widget designer tools. Must be set before opening widget designer."));
+
+
 FDesignerTabSummoner::FDesignerTabSummoner(TSharedPtr<class FWidgetBlueprintEditor> InBlueprintEditor)
 		: FWorkflowTabFactory(TabID, InBlueprintEditor)
 		, BlueprintEditor(InBlueprintEditor)
@@ -23,8 +33,19 @@ FDesignerTabSummoner::FDesignerTabSummoner(TSharedPtr<class FWidgetBlueprintEdit
 
 TSharedRef<SWidget> FDesignerTabSummoner::CreateTabBody(const FWorkflowTabSpawnInfo& Info) const
 {
-	return SNew(SDesignerView, BlueprintEditor.Pin())
+	if (!bEnableWidgetDesignerTools)
+	{
+		return SNew(SDesignerView, BlueprintEditor.Pin())
+			.AddMetaData<FTagMetaData>(FTagMetaData(TEXT("Designer")));
+	}
+
+	TSharedRef<TToolCompatibleMixin<SDesignerView>> DesignerView = SNew(TToolCompatibleMixin<SDesignerView>, BlueprintEditor.Pin())
 		.AddMetaData<FTagMetaData>(FTagMetaData(TEXT("Designer")));
+	DesignerView->SetParentToolkit(BlueprintEditor.Pin());
+
+	// @TODO: DarenC - Ideally this is done by SetParentToolkit, but can't seem to do that due to template issues with SharedThis.
+	DesignerView->GetWidgetModeManger()->SetManagedWidget(DesignerView);
+	return DesignerView;
 }
 
 #undef LOCTEXT_NAMESPACE 
diff --git a/Engine/Source/Editor/UMGEditor/Private/TabFactory/DesignerTabSummoner.h b/Engine/Source/Editor/UMGEditor/Private/TabFactory/DesignerTabSummoner.h
--- a/Engine/Source/Editor/UMGEditor/Private/TabFactory/DesignerTabSummoner.h
+++ b/Engine/Source/Editor/UMGEditor/Private/TabFactory/DesignerTabSummoner.h
@@ -10,6 +10,7 @@
 struct FDesignerTabSummoner : public FWorkflowTabFactory
 {
 public:
+	static const FName ToolPaletteTabID;
 	static const FName TabID;
 	
 public:
diff --git a/Engine/Source/Editor/UMGEditor/Private/Templates/WidgetTemplateBlueprintClass.h b/Engine/Source/Editor/UMGEditor/Private/Templates/WidgetTemplateBlueprintClass.h
deleted file mode 100644
diff --git a/Engine/Source/Editor/UMGEditor/Private/Templates/WidgetTemplateClass.h b/Engine/Source/Editor/UMGEditor/Private/Templates/WidgetTemplateClass.h
deleted file mode 100644
diff --git a/Engine/Source/Editor/UMGEditor/Private/Templates/WidgetTemplateImageClass.h b/Engine/Source/Editor/UMGEditor/Private/Templates/WidgetTemplateImageClass.h
deleted file mode 100644
diff --git a/Engine/Source/Editor/UMGEditor/Private/ToolPalette/WidgetEditorModeUILayer.cpp b/Engine/Source/Editor/UMGEditor/Private/ToolPalette/WidgetEditorModeUILayer.cpp
new file mode 100644
--- /dev/null
+++ b/Engine/Source/Editor/UMGEditor/Private/ToolPalette/WidgetEditorModeUILayer.cpp
@@ -0,0 +1,58 @@
+// Copyright Epic Games, Inc. All Rights Reserved.
+
+#include "WidgetEditorModeUILayer.h"
+#include "WorkspaceMenuStructure.h"
+#include "WorkspaceMenuStructureModule.h"
+#include "UMGEditorModule.h"
+#include "Toolkits/IToolkit.h"
+#include "DesignerTabSummoner.h"
+
+void UWidgetEditorModeUISubsystem::Initialize(FSubsystemCollectionBase& Collection)
+{
+	IUMGEditorModule& UMGEditorModule = FModuleManager::GetModuleChecked<IUMGEditorModule>("UMGEditor");
+	UMGEditorModule.OnRegisterLayoutExtensions().AddUObject(this, &UWidgetEditorModeUISubsystem::RegisterLayoutExtensions);
+}
+
+void UWidgetEditorModeUISubsystem::Deinitialize()
+{
+	IUMGEditorModule& UMGEditorModule = FModuleManager::GetModuleChecked<IUMGEditorModule>("UMGEditor");
+	UMGEditorModule.OnRegisterLayoutExtensions().RemoveAll(this);
+}
+
+void UWidgetEditorModeUISubsystem::RegisterLayoutExtensions(FLayoutExtender& Extender)
+{
+	FTabManager::FTab NewTab(FTabId(UAssetEditorUISubsystem::TopLeftTabID, ETabIdFlags::SaveLayout), ETabState::ClosedTab);
+	Extender.ExtendLayout(FDesignerTabSummoner::ToolPaletteTabID, ELayoutExtensionPosition::After, NewTab);
+}
+
+FWidgetEditorModeUILayer::FWidgetEditorModeUILayer(const IToolkitHost* InToolkitHost) :
+	FAssetEditorModeUILayer(InToolkitHost)
+{}
+
+void FWidgetEditorModeUILayer::OnToolkitHostingStarted(const TSharedRef<IToolkit>& Toolkit)
+{
+	if (!Toolkit->IsAssetEditor())
+	{
+		FAssetEditorModeUILayer::OnToolkitHostingStarted(Toolkit);
+		HostedToolkit = Toolkit;
+		Toolkit->SetModeUILayer(SharedThis(this));
+		Toolkit->RegisterTabSpawners(ToolkitHost->GetTabManager().ToSharedRef());
+		RegisterModeTabSpawners();
+		OnToolkitHostReadyForUI.ExecuteIfBound();
+	}
+}
+
+void FWidgetEditorModeUILayer::OnToolkitHostingFinished(const TSharedRef<IToolkit>& Toolkit)
+{
+	if (HostedToolkit.IsValid() && HostedToolkit.Pin() == Toolkit)
+	{
+		FAssetEditorModeUILayer::OnToolkitHostingFinished(Toolkit);
+	}
+}
+
+TSharedPtr<FWorkspaceItem> FWidgetEditorModeUILayer::GetModeMenuCategory() const
+{
+	const IWorkspaceMenuStructure& MenuStructure = WorkspaceMenu::GetMenuStructure();
+	return MenuStructure.GetLevelEditorModesCategory();
+}
+
diff --git a/Engine/Source/Editor/UMGEditor/Private/ToolPalette/WidgetEditorModeUILayer.h b/Engine/Source/Editor/UMGEditor/Private/ToolPalette/WidgetEditorModeUILayer.h
new file mode 100644
--- /dev/null
+++ b/Engine/Source/Editor/UMGEditor/Private/ToolPalette/WidgetEditorModeUILayer.h
@@ -0,0 +1,32 @@
+// Copyright Epic Games, Inc. All Rights Reserved.
+
+#pragma once
+
+#include "Toolkits/AssetEditorModeUILayer.h"
+#include "WidgetEditorModeUILayer.generated.h"
+
+UCLASS()
+class UWidgetEditorModeUISubsystem : public UAssetEditorUISubsystem
+{
+	GENERATED_BODY()
+
+public:
+	virtual void Initialize(FSubsystemCollectionBase& Collection) override;
+	virtual void Deinitialize() override;
+	virtual void RegisterLayoutExtensions(FLayoutExtender& Extender) override;
+};
+
+
+class FWidgetEditorModeUILayer : public FAssetEditorModeUILayer
+{
+public:
+
+	FWidgetEditorModeUILayer(const IToolkitHost* InToolkitHost);
+
+	void OnToolkitHostingStarted(const TSharedRef<IToolkit>& Toolkit) override;
+	void OnToolkitHostingFinished(const TSharedRef<IToolkit>& Toolkit) override;
+	
+	TSharedPtr<FWorkspaceItem> GetModeMenuCategory() const override;
+};
+
+
diff --git a/Engine/Source/Editor/UMGEditor/Private/UMGEditorModule.cpp b/Engine/Source/Editor/UMGEditor/Private/UMGEditorModule.cpp
--- a/Engine/Source/Editor/UMGEditor/Private/UMGEditorModule.cpp
+++ b/Engine/Source/Editor/UMGEditor/Private/UMGEditorModule.cpp
@@ -127,7 +127,7 @@ public:
 		ToolBarExtensibilityManager.Reset();
 
 		IKismetCompilerInterface& KismetCompilerModule = FModuleManager::LoadModuleChecked<IKismetCompilerInterface>("KismetCompiler");
-		KismetCompilerModule.GetCompilers().Remove(&WidgetBlueprintCompiler);
+		KismetCompilerModule.GetCompilers().Remove(&WidgetBlueprintCompiler); 
 
 		// Unregister all the asset types that we registered
 		if ( FModuleManager::Get().IsModuleLoaded("AssetTools") )
@@ -227,6 +227,21 @@ public:
 		return RegisterTabsForEditor;
 	}
 
+	virtual void AddWidgetEditorToolbarExtender(FWidgetEditorToolbarExtender&& InToolbarExtender) override
+	{
+		WidgetEditorToolbarExtenders.Add(MoveTemp(InToolbarExtender));
+	}
+
+	virtual TArrayView<FWidgetEditorToolbarExtender> GetAllWidgetEditorToolbarExtenders() override
+	{
+		return WidgetEditorToolbarExtenders;
+	}
+
+	virtual FOnRegisterLayoutExtensions& OnRegisterLayoutExtensions() override 
+	{
+		return RegisterLayoutExtensions; 
+	}
+
 private:
 	void RegisterAssetTypeAction(IAssetTools& AssetTools, TSharedRef<IAssetTypeActions> Action)
 	{
@@ -277,6 +292,9 @@ private:
 	/** All created asset type actions.  Cached here so that we can unregister it during shutdown. */
 	TArray< TSharedPtr<IAssetTypeActions> > CreatedAssetTypeActions;
 
+	/** All toolbar extenders. Utilized by tool palette */
+	TArray<FWidgetEditorToolbarExtender> WidgetEditorToolbarExtenders;
+
 	USequencerSettings* Settings;
 
 	/** Compiler customization for Widgets */
@@ -285,6 +303,9 @@ private:
 	/** */
 	FOnRegisterTabs RegisterTabsForEditor;
 
+	/** Support layout extensions */
+	FOnRegisterLayoutExtensions	RegisterLayoutExtensions;
+
 	bool bThumbnailRenderersRegistered;
 	bool bOnPostEngineInitHandled;
 };
diff --git a/Engine/Source/Editor/UMGEditor/Private/WidgetBlueprintEditor.cpp b/Engine/Source/Editor/UMGEditor/Private/WidgetBlueprintEditor.cpp
--- a/Engine/Source/Editor/UMGEditor/Private/WidgetBlueprintEditor.cpp
+++ b/Engine/Source/Editor/UMGEditor/Private/WidgetBlueprintEditor.cpp
@@ -40,6 +40,7 @@
 #include "WorkflowOrientedApp/ApplicationMode.h"
 #include "BlueprintModes/WidgetDesignerApplicationMode.h"
 #include "BlueprintModes/WidgetGraphApplicationMode.h"
+#include "WidgetModeManager.h"
 
 #include "WidgetBlueprintEditorToolbar.h"
 #include "Components/CanvasPanel.h"
@@ -54,11 +55,13 @@
 
 #include "ScopedTransaction.h"
 
+#include "Designer/SDesignerView.h"
 #include "Framework/Notifications/NotificationManager.h"
 #include "Widgets/Notifications/SNotificationList.h"
 #include "UMGEditorActions.h"
 #include "UMGEditorModule.h"
 #include "GameProjectGenerationModule.h"
+#include "Tools/ToolCompatible.h"
 
 #include "SPaletteViewModel.h"
 #include "SLibraryViewModel.h"
@@ -72,6 +75,8 @@
 #include "Serialization/BufferArchive.h"
 #include "Widgets/SVirtualWindow.h"
 #include "TabFactory/AnimationTabSummoner.h"
+#include "TabFactory/DesignerTabSummoner.h"
+#include "ToolPalette/WidgetEditorModeUILayer.h"
 #include "Kismet/Public/BlueprintEditorTabs.h"
 
 #include "Editor/UnrealEdEngine.h"
@@ -190,6 +195,10 @@ void FWidgetBlueprintEditor::InitWidgetBlueprintEditor(const EToolkitMode::Type
 		FExecuteAction::CreateSP(this, &FWidgetBlueprintEditor::DuplicateSelectedWidgets),
 		FCanExecuteAction::CreateSP(this, &FWidgetBlueprintEditor::CanDuplicateSelectedWidgets)
 		);
+
+	TSharedPtr<class IToolkitHost> PinnedToolkitHost = ToolkitHost.Pin();
+	check(PinnedToolkitHost.IsValid());
+	ModeUILayer = MakeShared<FWidgetEditorModeUILayer>(PinnedToolkitHost.Get());
 }
 
 void FWidgetBlueprintEditor::InitalizeExtenders()
@@ -200,6 +209,15 @@ void FWidgetBlueprintEditor::InitalizeExtenders()
 	AddMenuExtender(UMGEditorModule.GetMenuExtensibilityManager()->GetAllExtenders(GetToolkitCommands(), GetEditingObjects()));
 
 	AddMenuExtender(CreateMenuExtender());
+
+	TArrayView<IUMGEditorModule::FWidgetEditorToolbarExtender> ToolbarExtenderDelegates = UMGEditorModule.GetAllWidgetEditorToolbarExtenders();
+	for (auto& ToolbarExtenderDelegate : ToolbarExtenderDelegates)
+	{
+		if (ToolbarExtenderDelegate.IsBound())
+		{
+			AddToolbarExtender(ToolbarExtenderDelegate.Execute(GetToolkitCommands(), SharedThis(this)));
+		}
+	}
 }
 
 TSharedPtr<FExtender> FWidgetBlueprintEditor::CreateMenuExtender()
@@ -1455,6 +1473,31 @@ void FWidgetBlueprintEditor::Compile()
 	}
 }
 
+bool FWidgetBlueprintEditor::OnRequestClose()
+{
+	bool bAllowClose = Super::OnRequestClose();
+
+	// Give any active modes a chance to shutdown while the toolkit host is still alive
+	// Note: This along side with the default tool palette extension tab being closed 
+	// is what prevents an unrecognized tab from spawning on layout restore
+	if (bAllowClose)
+	{
+		GetEditorModeManager().ActivateDefaultMode();
+	}
+
+	return bAllowClose;
+}
+
+void FWidgetBlueprintEditor::OnToolkitHostingStarted(const TSharedRef<IToolkit>& Toolkit)
+{
+	ModeUILayer->OnToolkitHostingStarted(Toolkit);
+}
+
+void FWidgetBlueprintEditor::OnToolkitHostingFinished(const TSharedRef<IToolkit>& Toolkit)
+{
+	ModeUILayer->OnToolkitHostingFinished(Toolkit);
+}
+
 void FWidgetBlueprintEditor::DestroyPreview()
 {
 	UUserWidget* PreviewUserWidget = GetPreview();
@@ -1653,6 +1696,14 @@ void FWidgetBlueprintEditor::SetIsRespectingLocks(bool Value)
 	bRespectLocks = Value;
 }
 
+void FWidgetBlueprintEditor::CreateEditorModeManager()
+{
+	TSharedPtr<FWidgetModeManager> WidgetModeManager = MakeShared<FWidgetModeManager>();
+	WidgetModeManager->OwningToolkit = SharedThis(this);
+	EditorModeManager = WidgetModeManager;
+
+}
+
 class FObjectAndDisplayName
 {
 public:
diff --git a/Engine/Source/Editor/UMGEditor/Private/WidgetBlueprintEditorToolbar.cpp b/Engine/Source/Editor/UMGEditor/Private/WidgetBlueprintEditorToolbar.cpp
--- a/Engine/Source/Editor/UMGEditor/Private/WidgetBlueprintEditorToolbar.cpp
+++ b/Engine/Source/Editor/UMGEditor/Private/WidgetBlueprintEditorToolbar.cpp
@@ -141,4 +141,17 @@ void FWidgetBlueprintEditorToolbar::AddWidgetReflector(UToolMenu* InMenu)
 	));
 }
 
+void FWidgetBlueprintEditorToolbar::AddToolPalettes(UToolMenu* InMenu)
+{
+	// @TODO: DarenC - For now we only support one tool palette, switch this to a dropdown when we support multiple tool palettes.
+	for (TSharedPtr<FUICommandInfo>& Command : WidgetEditor.Pin()->ToolPaletteCommands)
+	{
+		FToolMenuSection& Section = InMenu->FindOrAddSection("UMGToolPalette");
+		Section.AddDynamicEntry(Command->GetCommandName(), FNewToolMenuSectionDelegate::CreateLambda([this, Command](FToolMenuSection& InSection)
+			{
+				InSection.AddEntry(FToolMenuEntry::InitToolBarButton(Command));
+			}));
+	}
+}
+
 #undef LOCTEXT_NAMESPACE
diff --git a/Engine/Source/Editor/UMGEditor/Private/WidgetBlueprintEditorToolbar.h b/Engine/Source/Editor/UMGEditor/Private/WidgetBlueprintEditorToolbar.h
--- a/Engine/Source/Editor/UMGEditor/Private/WidgetBlueprintEditorToolbar.h
+++ b/Engine/Source/Editor/UMGEditor/Private/WidgetBlueprintEditorToolbar.h
@@ -26,6 +26,8 @@ public:
 
 	void AddWidgetReflector(UToolMenu* InMenu);
 
+	void AddToolPalettes(UToolMenu* InMenu);
+
 public:
 	/**  */
 	void FillWidgetBlueprintEditorModesToolbar(FToolBarBuilder& ToolbarBuilder);
diff --git a/Engine/Source/Editor/UMGEditor/Private/WidgetBlueprintEditorUtils.h b/Engine/Source/Editor/UMGEditor/Private/WidgetBlueprintEditorUtils.h
deleted file mode 100644
diff --git a/Engine/Source/Editor/UMGEditor/Public/Designer/SDesignSurface.h b/Engine/Source/Editor/UMGEditor/Public/Designer/SDesignSurface.h
new file mode 100644
--- /dev/null
+++ b/Engine/Source/Editor/UMGEditor/Public/Designer/SDesignSurface.h
@@ -0,0 +1,153 @@
+// Copyright Epic Games, Inc. All Rights Reserved.
+
+#pragma once
+
+#include "CoreMinimal.h"
+#include "Misc/Attribute.h"
+#include "Styling/SlateColor.h"
+#include "Layout/SlateRect.h"
+#include "Layout/Geometry.h"
+#include "Input/CursorReply.h"
+#include "Input/Reply.h"
+#include "Animation/CurveSequence.h"
+#include "Widgets/DeclarativeSyntaxSupport.h"
+#include "Widgets/SCompoundWidget.h"
+#include "SNodePanel.h"
+#include "Templates/UniquePtr.h"
+
+class FActiveTimerHandle;
+class FSlateWindowElementList;
+
+class UMGEDITOR_API SDesignSurface : public SCompoundWidget
+{
+public:
+
+	SLATE_BEGIN_ARGS( SDesignSurface )
+		: _AllowContinousZoomInterpolation(false)
+	{ }
+
+		/** Slot for this designers content (optional) */
+		SLATE_DEFAULT_SLOT(FArguments, Content)
+
+		SLATE_ATTRIBUTE(bool, AllowContinousZoomInterpolation)
+	SLATE_END_ARGS()
+
+	void Construct(const FArguments& InArgs);
+
+	// SWidget interface
+	virtual void Tick(const FGeometry& AllottedGeometry, const double InCurrentTime, const float InDeltaTime) override;
+	virtual FCursorReply OnCursorQuery(const FGeometry& MyGeometry, const FPointerEvent& CursorEvent) const override;
+	virtual FReply OnMouseButtonDown(const FGeometry& MyGeometry, const FPointerEvent& MouseEvent) override;
+	virtual FReply OnMouseButtonUp(const FGeometry& MyGeometry, const FPointerEvent& MouseEvent) override;
+	virtual FReply OnMouseMove(const FGeometry& MyGeometry, const FPointerEvent& MouseEvent) override;
+	virtual FReply OnMouseWheel(const FGeometry& MyGeometry, const FPointerEvent& MouseEvent) override;
+	virtual FReply OnTouchGesture(const FGeometry& MyGeometry, const FPointerEvent& GestureEvent) override;
+	virtual FReply OnTouchEnded(const FGeometry& MyGeometry, const FPointerEvent& InTouchEvent) override;
+	virtual int32 OnPaint(const FPaintArgs& Args, const FGeometry& AllottedGeometry, const FSlateRect& MyCullingRect, FSlateWindowElementList& OutDrawElements, int32 LayerId, const FWidgetStyle& InWidgetStyle, bool bParentEnabled) const override;
+	virtual bool SupportsKeyboardFocus() const override { return true; }
+	// End of Swidget interface
+
+	/** Gets the current zoom factor. */
+	float GetZoomAmount() const;
+
+	FVector2D GraphCoordToPanelCoord(const FVector2D& GraphSpaceCoordinate) const;
+	FVector2D PanelCoordToGraphCoord(const FVector2D& PanelSpaceCoordinate) const;
+
+protected:
+	virtual void OnPaintBackground(const FGeometry& AllottedGeometry, const FSlateRect& MyCullingRect, FSlateWindowElementList& OutDrawElements, int32 LayerId) const;
+
+	void PaintBackgroundAsLines(const FSlateBrush* BackgroundImage, const FGeometry& AllottedGeometry, const FSlateRect& MyCullingRect, FSlateWindowElementList& OutDrawElements, int32& DrawLayerId) const;
+
+	void ChangeZoomLevel(int32 ZoomLevelDelta, const FVector2D& WidgetSpaceZoomOrigin, bool bOverrideZoomLimiting);
+	
+	void PostChangedZoom();
+
+	bool ScrollToLocation(const FGeometry& MyGeometry, FVector2D DesiredCenterPosition, const float InDeltaTime);
+
+	bool ZoomToLocation(const FVector2D& CurrentSizeWithoutZoom, const FVector2D& DesiredSize, bool bDoneScrolling);
+
+	void ZoomToFit(bool bInstantZoom);
+
+	FText GetZoomText() const;
+	FSlateColor GetZoomTextColorAndOpacity() const;
+
+	FVector2D GetViewOffset() const;
+
+	FSlateRect ComputeSensibleBounds() const;
+
+protected:
+	virtual FSlateRect ComputeAreaBounds() const;
+	virtual float GetGridScaleAmount() const;
+	virtual int32 GetGraphRulePeriod() const;
+	virtual int32 GetSnapGridSize() const = 0;
+
+protected:
+	/** The position within the graph at which the user is looking */
+	FVector2D ViewOffset;
+
+	/** The position in the grid to begin drawing at. */
+	FVector2D GridOrigin;
+
+	/** Should we render the grid lines? */
+	bool bDrawGridLines;
+
+	/** Previous Zoom Level */
+	int32 PreviousZoomLevel;
+
+	/** How zoomed in/out we are. e.g. 0.25f results in quarter-sized nodes. */
+	int32 ZoomLevel;
+
+	/** Are we panning the view at the moment? */
+	bool bIsPanning;
+
+	/** Are we zooming the view with trackpad at the moment? */
+	bool bIsZooming;
+
+	/** Allow continuous zoom interpolation? */
+	TAttribute<bool> AllowContinousZoomInterpolation;
+
+	/** Fade on zoom for graph */
+	FCurveSequence ZoomLevelGraphFade;
+
+	/** Curve that handles fading the 'Zoom +X' text */
+	FCurveSequence ZoomLevelFade;
+
+	// The interface for mapping ZoomLevel values to actual node scaling values
+	TUniquePtr<FZoomLevelsContainer> ZoomLevels;
+
+	bool bAllowContinousZoomInterpolation;
+
+	bool bTeleportInsteadOfScrollingWhenZoomingToFit;
+
+	FVector2D ZoomTargetTopLeft;
+	FVector2D ZoomTargetBottomRight;
+	FVector2D ZoomToFitPadding;
+
+	/** The Y component of mouse drag (used when zooming) */
+	float TotalMouseDelta;
+
+	/** Offset in the panel the user started the LMB+RMB zoom from */
+	FVector2D ZoomStartOffset;
+
+	/**  */
+	FVector2D ViewOffsetStart;
+
+	/**  */
+	FVector2D MouseDownPositionAbsolute;
+
+	/** Cumulative magnify delta from trackpad gesture */
+	float TotalGestureMagnify;
+
+	/** Does the user need to press Control in order to over-zoom. */
+	bool bRequireControlToOverZoom;
+
+private:
+	/** Active timer that handles deferred zooming until the target zoom is reached */
+	EActiveTimerReturnType HandleZoomToFit(double InCurrentTime, float InDeltaTime);
+
+	/** The handle to the active timer */
+	TWeakPtr<FActiveTimerHandle> ActiveTimerHandle;
+
+	// A flag noting if we have a pending zoom to extents operation to perform next tick.
+	bool bDeferredZoomToExtents;
+};
diff --git a/Engine/Source/Editor/UMGEditor/Public/Designer/SDesignerView.h b/Engine/Source/Editor/UMGEditor/Public/Designer/SDesignerView.h
new file mode 100644
--- /dev/null
+++ b/Engine/Source/Editor/UMGEditor/Public/Designer/SDesignerView.h
@@ -0,0 +1,422 @@
+// Copyright Epic Games, Inc. All Rights Reserved.
+
+#pragma once
+
+#include "CoreMinimal.h"
+#include "Templates/SubclassOf.h"
+#include "Styling/SlateColor.h"
+#include "Layout/SlateRect.h"
+#include "Layout/Visibility.h"
+#include "Layout/Geometry.h"
+#include "Input/Reply.h"
+#include "Layout/ArrangedWidget.h"
+#include "Widgets/SWidget.h"
+#include "UObject/GCObject.h"
+#include "Types/SlateStructs.h"
+#include "Animation/CurveSequence.h"
+#include "Widgets/DeclarativeSyntaxSupport.h"
+#include "Styling/SlateBrush.h"
+#include "Components/Widget.h"
+#include "WidgetReference.h"
+#include "WidgetBlueprintEditor.h"
+#include "Settings/LevelEditorPlaySettings.h"
+#include "Layout/WidgetPath.h"
+#include "IUMGDesigner.h"
+#include "DesignerExtension.h"
+#include "Designer/SDesignSurface.h"
+#include "UMGEditorProjectSettings.h"
+
+class FMenuBuilder;
+class FScopedTransaction;
+class SBox;
+class SCanvas;
+class SPaintSurface;
+class SRuler;
+class SZoomPan;
+class UPanelWidget;
+class UWidgetBlueprint;
+class FHittestGrid;
+struct FOnPaintHandlerParams;
+struct FWidgetHitResult;
+
+/**
+ * The designer for widgets.  Allows for laying out widgets in a drag and drop environment.
+ */
+class UMGEDITOR_API SDesignerView : public SDesignSurface, public FGCObject, public IUMGDesigner
+{
+public:
+
+	SLATE_BEGIN_ARGS( SDesignerView ) {}
+	SLATE_END_ARGS()
+
+	void Construct(const FArguments& InArgs, TSharedPtr<class FWidgetBlueprintEditor> InBlueprintEditor);
+	virtual ~SDesignerView();
+
+	EActiveTimerReturnType EnsureTick(double InCurrentTime, float InDeltaTime);
+
+	TSharedRef<SWidget> CreateOverlayUI();
+
+	// SWidget interface
+	virtual FReply OnMouseButtonDown(const FGeometry& MyGeometry, const FPointerEvent& MouseEvent) override;
+	virtual FReply OnMouseButtonUp(const FGeometry& MyGeometry, const FPointerEvent& MouseEvent) override;
+	virtual FReply OnMouseMove(const FGeometry& MyGeometry, const FPointerEvent& MouseEvent) override;
+	virtual void OnMouseEnter(const FGeometry& MyGeometry, const FPointerEvent& MouseEvent) override;
+	virtual void OnMouseLeave(const FPointerEvent& MouseEvent) override;
+
+	virtual FReply OnKeyDown(const FGeometry& MyGeometry, const FKeyEvent& InKeyEvent) override;
+	virtual FReply OnKeyUp(const FGeometry& MyGeometry, const FKeyEvent& InKeyEvent) override;
+
+	virtual void Tick( const FGeometry& AllottedGeometry, const double InCurrentTime, const float InDeltaTime ) override;
+
+	virtual FReply OnDragDetected(const FGeometry& MyGeometry, const FPointerEvent& MouseEvent) override;
+	virtual void OnDragEnter(const FGeometry& MyGeometry, const FDragDropEvent& DragDropEvent) override;
+	virtual void OnDragLeave(const FDragDropEvent& DragDropEvent) override;
+	virtual FReply OnDragOver(const FGeometry& MyGeometry, const FDragDropEvent& DragDropEvent) override;
+	virtual FReply OnDrop(const FGeometry& MyGeometry, const FDragDropEvent& DragDropEvent) override;
+	// End of SWidget interface
+
+	void Register(TSharedRef<FDesignerExtension> Extension);
+	void Unregister(TSharedRef<FDesignerExtension> Extension);
+
+	// IUMGDesigner interface
+	virtual float GetPreviewScale() const override;
+	virtual const TSet<FWidgetReference>& GetSelectedWidgets() const override;
+	virtual FWidgetReference GetSelectedWidget() const override;
+	virtual ETransformMode::Type GetTransformMode() const override;
+	virtual FGeometry GetDesignerGeometry() const override;
+	virtual FVector2D GetWidgetOriginAbsolute() const;
+	virtual bool GetWidgetGeometry(const FWidgetReference& Widget, FGeometry& Geometry) const override;
+	virtual bool GetWidgetGeometry(const UWidget* PreviewWidget, FGeometry& Geometry) const override;
+	virtual bool GetWidgetParentGeometry(const FWidgetReference& Widget, FGeometry& Geometry) const override;
+	virtual FGeometry MakeGeometryWindowLocal(const FGeometry& WidgetGeometry) const override;
+	virtual void MarkDesignModifed(bool bRequiresRecompile) override;
+	virtual void PushDesignerMessage(const FText& Message) override;
+	virtual void PopDesignerMessage() override;
+	// End of IUMGDesigner interface
+
+	// FGCObject interface
+	virtual void AddReferencedObjects(FReferenceCollector& Collector) override;
+	virtual FString GetReferencerName() const override
+	{
+		return TEXT("SDesignerView");
+	}
+	// End of FGCObject interface
+
+public:
+
+	/** The width of the preview screen for the UI */
+	FOptionalSize GetPreviewAreaWidth() const;
+
+	/** The height of the preview screen for the UI */
+	FOptionalSize GetPreviewAreaHeight() const;
+
+	/** The width of the preview widget for the UI */
+	FOptionalSize GetPreviewSizeWidth() const;
+
+	/** The height of the preview widget for the UI */
+	FOptionalSize GetPreviewSizeHeight() const;
+
+	/** Gets the DPI scale that would be applied given the current preview width and height */
+	float GetPreviewDPIScale() const;
+
+	/** Set the size of the preview screen for the UI */
+	void SetPreviewAreaSize(int32 Width, int32 Height);
+
+	void BeginResizingArea();
+	void EndResizingArea();
+
+protected:
+	virtual void OnPaintBackground(const FGeometry& AllottedGeometry, const FSlateRect& MyCullingRect, FSlateWindowElementList& OutDrawElements, int32 LayerId) const override;
+	void DrawResolution(const FDebugResolution& Resolution, const FGeometry& AllottedGeometry, const FSlateRect& MyCullingRect, FSlateWindowElementList& OutDrawElements, int32 LayerId) const;
+
+	virtual FSlateRect ComputeAreaBounds() const override;
+	virtual int32 GetGraphRulePeriod() const override;
+	virtual float GetGridScaleAmount() const override;
+	virtual int32 GetSnapGridSize() const override;
+
+private:
+	void RegisterExtensions();
+
+	/** Establishes the resolution and aspect ratio to use on construction from config settings */
+	void SetStartupResolution();
+
+	FVector2D GetAreaResizeHandlePosition() const;
+	EVisibility GetAreaResizeHandleVisibility() const;
+
+	const FSlateBrush* GetPreviewBackground() const;
+
+	/** Adds any pending selected widgets to the selection set */
+	void ResolvePendingSelectedWidgets();
+
+	/** Updates the designer to display the latest preview widget */
+	void UpdatePreviewWidget(bool bForceUpdate);
+
+	void BroadcastDesignerChanged();
+
+	void ClearExtensionWidgets();
+	void CreateExtensionWidgetsForSelection();
+
+	EVisibility GetInfoBarVisibility() const;
+	FText GetInfoBarText() const;
+
+	/** Displays the context menu when you right click */
+	void ShowContextMenu(const FGeometry& MyGeometry, const FPointerEvent& MouseEvent);
+
+	void OnEditorSelectionChanged();
+
+	/** Called when a new widget is being hovered */
+	void OnHoveredWidgetSet(const FWidgetReference& InHoveredWidget);
+
+	/** Called when a widget is no longer being hovered */
+	void OnHoveredWidgetCleared();
+
+	/** Gets the blueprint being edited by the designer */
+	UWidgetBlueprint* GetBlueprint() const;
+
+	/** Called whenever the blueprint is recompiled or the preview is updated */
+	void OnPreviewNeedsRecreation();
+
+	void PopulateWidgetGeometryCache(FArrangedWidget& Root);
+	void PopulateWidgetGeometryCache_Loop(FArrangedWidget& Parent);
+
+	/** @return Formatted text for the given resolution params */
+	FText GetResolutionText(int32 Width, int32 Height, const FString& AspectRatio) const;
+
+	/** Move the selected widget by the nudge amount. */
+	FReply NudgeSelectedWidget(FVector2D Nudge);
+
+	FText GetCurrentResolutionText() const;
+	FText GetCurrentDPIScaleText() const;
+	FSlateColor GetCurrentDPIScaleColor() const;
+	FText GetCurrentScaleFactorText() const;
+	FText GetCurrentSafeZoneText() const;
+	FSlateColor GetResolutionTextColorAndOpacity() const;
+	EVisibility GetResolutionTextVisibility() const;
+
+	TOptional<int32> GetCustomResolutionWidth() const;
+	TOptional<int32> GetCustomResolutionHeight() const;
+	void OnCustomResolutionWidthChanged(int32 InValue);
+	void OnCustomResolutionHeightChanged(int32 InValue);
+	EVisibility GetCustomResolutionEntryVisibility() const;
+	
+	// Handles selecting a common screen resolution.
+	void HandleOnCommonResolutionSelected(const FPlayScreenResolution InResolution);
+	bool HandleIsCommonResolutionSelected(const FPlayScreenResolution InResolution) const;
+	FUIAction GetResolutionMenuAction( const FPlayScreenResolution& ScreenResolution );
+	TSharedRef<SWidget> GetResolutionsMenu();
+
+	TSharedRef<SWidget> GetScreenSizingFillMenu();
+	void CreateScreenFillEntry(FMenuBuilder& MenuBuilder, EDesignPreviewSizeMode SizeMode);
+	FText GetScreenSizingFillText() const;
+	bool GetIsScreenFillRuleSelected(EDesignPreviewSizeMode SizeMode) const;
+	void OnScreenFillRuleSelected(EDesignPreviewSizeMode SizeMode);
+
+	const FSlateBrush* GetAspectRatioSwitchImage() const;
+	bool GetAspectRatioSwitchEnabled() const;
+	bool GetFlipDeviceEnabled() const;
+	EVisibility GetDesignerOutlineVisibility() const;
+	FSlateColor GetDesignerOutlineColor() const;
+	FText GetDesignerOutlineText() const;
+
+	FText GetCursorPositionText() const;
+	EVisibility GetCursorPositionTextVisibility() const;
+
+	// Handles drawing selection and other effects a SPaintSurface widget injected into the hierarchy.
+	int32 HandleEffectsPainting(const FOnPaintHandlerParams& PaintArgs);
+	void DrawSelectionAndHoverOutline(const FOnPaintHandlerParams& PaintArgs);
+	void DrawSafeZone(const FOnPaintHandlerParams& PaintArgs);
+
+	FReply HandleDPISettingsClicked();
+
+	UUserWidget* GetDefaultWidget() const;
+
+	void BeginTransaction(const FText& SessionName);
+	bool InTransaction() const;
+	void EndTransaction(bool bCancel);
+
+	UWidget* GetWidgetInDesignScopeFromSlateWidget(TSharedRef<SWidget>& InWidget);
+
+	EVisibility GetExtensionCanvasVisibility() const;
+
+private:
+	struct FWidgetHitResult
+	{
+	public:
+		FWidgetReference Widget;
+		FArrangedWidget WidgetArranged;
+
+		FName NamedSlot;
+
+	public:
+		FWidgetHitResult();
+	};
+
+	/** @returns Gets the widget under the cursor based on a mouse pointer event. */
+	bool FindWidgetUnderCursor(const FGeometry& MyGeometry, const FPointerEvent& MouseEvent, TSubclassOf<UWidget> FindType, FWidgetHitResult& HitResult);
+
+private:
+	FReply HandleZoomToFitClicked();
+	FReply HandleSwapAspectRatioClicked();
+	FReply HandleFlipSafeZonesClicked();
+	EVisibility GetRulerVisibility() const;
+
+private:
+	static const FString ConfigSectionName;
+	static const uint32 DefaultResolutionWidth;
+	static const uint32 DefaultResolutionHeight;
+	static const FString DefaultAspectRatio;
+	static const FString DefaultPreviewOverrideName;
+
+	/** Extensions for the designer to allow for custom widgets to be inserted onto the design surface as selection changes. */
+	TArray< TSharedRef<FDesignerExtension> > DesignerExtensions;
+
+private:
+	void BindCommands();
+
+	void SetTransformMode(ETransformMode::Type InTransformMode);
+	bool CanSetTransformMode(ETransformMode::Type InTransformMode) const;
+	bool IsTransformModeActive(ETransformMode::Type InTransformMode) const;
+
+	void ToggleShowingOutlines();
+	bool IsShowingOutlines() const;
+
+	void ToggleRespectingLocks();
+	bool IsRespectingLocks() const;
+
+	void ProcessDropAndAddWidget(const FGeometry& MyGeometry, const FDragDropEvent& DragDropEvent, const bool bIsPreview);
+
+	FVector2D GetExtensionPosition(TSharedRef<FDesignerSurfaceElement> ExtensionElement) const;
+
+	FVector2D GetExtensionSize(TSharedRef<FDesignerSurfaceElement> ExtensionElement) const;
+	
+	void ClearDropPreviews();
+
+	void DetermineDragDropPreviewWidgets(TArray<UWidget*>& OutWidgets, const FDragDropEvent& DragDropEvent, UWidgetTree* RootWidgetTree);
+
+	void SwapSafeZoneTypes();
+
+private:
+	/** A reference to the BP Editor that owns this designer */
+	TWeakPtr<FWidgetBlueprintEditor> BlueprintEditor;
+
+	/** The designer command list */
+	TSharedPtr<FUICommandList> CommandList;
+
+	/** The transaction used to commit undoable actions from resize, move...etc */
+	FScopedTransaction* ScopedTransaction;
+
+	/** The current preview widget */
+	UUserWidget* PreviewWidget;
+
+	/** The current preview widget's slate widget */
+	TWeakPtr<SWidget> PreviewSlateWidget;
+
+	/** Cache last mouse position to be used as a paste drop location */
+	FVector2D CachedMousePosition;
+	
+	struct FDropPreview
+	{
+		UWidget* Widget;
+		UPanelWidget* Parent;
+		TWeakPtr<FDragDropOperation> DragOperation;
+	};
+
+	TArray<FDropPreview> DropPreviews;
+
+	TSharedPtr<SWidget> PreviewHitTestRoot;
+	TSharedPtr<SBox> PreviewAreaConstraint;
+	TSharedPtr<SDPIScaler> PreviewSurface;
+	TSharedPtr<SBox> PreviewSizeConstraint;
+
+	TSharedPtr<SCanvas> DesignerControls;
+	TSharedPtr<SCanvas> DesignerWidgetCanvas;
+	TSharedPtr<SCanvas> ExtensionWidgetCanvas;
+	TSharedPtr<SPaintSurface> EffectsLayer;
+
+	/** The currently selected preview widgets in the preview GUI, just a cache used to determine changes between selection changes. */
+	TSet< FWidgetReference > SelectedWidgetsCache;
+
+	/** The location in selected widget local space where the context menu was summoned. */
+	FVector2D SelectedWidgetContextMenuLocation;
+
+	/**
+	 * Holds onto a temporary widget that the user may be getting ready to select, or may just 
+	 * be the widget that got hit on the initial mouse down before moving the parent.
+	 */
+	FWidgetReference PendingSelectedWidget;
+
+	/** The position in screen space where the user began dragging a widget */
+	FVector2D DraggingStartPositionScreenSpace;
+
+	/** An existing widget is being moved in its current container, or in to a new container. */
+	bool bMovingExistingWidget;
+
+	/** The configured Width of the preview area, simulates screen size. */
+	int32 PreviewWidth;
+
+	/** The configured Height of the preview area, simulates screen size. */
+	int32 PreviewHeight;
+
+	/** The original Width of the preview area, read from the settings file. */
+	int32 WidthReadFromSettings;
+
+	/** The original Height of the preview area, read from the settings file. */
+	int32 HeightReadFromSettings;
+
+	/***/
+	bool bShowResolutionOutlines;
+
+	/** The slate brush we use to hold the background image shown in the designer. */
+	mutable FSlateBrush BackgroundImage;
+
+	/** We cache the desired preview desired size to maintain the same size between compiles when it lags a frame behind and no widget is available. */
+	FVector2D CachedPreviewDesiredSize;
+
+	// Resolution Info
+	FString PreviewAspectRatio;
+
+	// Whether or not the resolution can flip between portrait and landscape
+	bool bCanPreviewSwapAspectRatio;
+
+	FString PreviewOverrideName;
+
+	float ScaleFactor;
+
+	bool bSafeZoneFlipped;
+
+	// Whether the preview is currently in portrait mode
+	bool bPreviewIsPortrait;
+
+	/** Curve to handle fading of the resolution */
+	FCurveSequence ResolutionTextFade;
+
+	/** Curve to handle the fade-in of the border around the hovered widget */
+	FCurveSequence HoveredWidgetOutlineFade;
+
+	/**  */
+	FWeakWidgetPath SelectedWidgetPath;
+
+	/** The ruler bar at the top of the designer. */
+	TSharedPtr<SRuler> TopRuler;
+
+	/** The ruler bar on the left side of the designer. */
+	TSharedPtr<SRuler> SideRuler;
+
+	/** */
+	ETransformMode::Type TransformMode;
+
+	/**  */
+	TMap<TSharedRef<SWidget>, FArrangedWidget> CachedWidgetGeometry;
+
+	/** */
+	TSharedPtr<FHittestGrid> DesignerHittestGrid;
+
+	/** The message stack to display the last item to the user in a non-modal fashion. */
+	TArray<FText> DesignerMessageStack;
+
+	TArray<FVector2D> CustomSafeZoneStarts;
+
+	TArray<FVector2D> CustomSafeZoneDimensions;
+
+	FMargin DesignerSafeZoneOverride;
+};
diff --git a/Engine/Source/Editor/UMGEditor/Public/DragDrop/WidgetTemplateDragDropOp.h b/Engine/Source/Editor/UMGEditor/Public/DragDrop/WidgetTemplateDragDropOp.h
new file mode 100644
--- /dev/null
+++ b/Engine/Source/Editor/UMGEditor/Public/DragDrop/WidgetTemplateDragDropOp.h
@@ -0,0 +1,25 @@
+// Copyright Epic Games, Inc. All Rights Reserved.
+
+#pragma once
+
+#include "CoreMinimal.h"
+#include "Input/DragAndDrop.h"
+#include "DragAndDrop/DecoratedDragDropOp.h"
+
+class FWidgetTemplate;
+
+/**
+ * This drag drop operation allows widget templates from the palate to be dragged and dropped into the designer
+ * or the widget hierarchy in order to spawn new widgets.
+ */
+class UMGEDITOR_API FWidgetTemplateDragDropOp : public FDecoratedDragDropOp
+{
+public:
+	DRAG_DROP_OPERATOR_TYPE(FWidgetTemplateDragDropOp, FDecoratedDragDropOp)
+
+	/** The template to create an instance */
+	TSharedPtr<FWidgetTemplate> Template;
+
+	/** Constructs the drag drop operation */
+	static TSharedRef<FWidgetTemplateDragDropOp> New(const TSharedPtr<FWidgetTemplate>& InTemplate);
+};
diff --git a/Engine/Source/Editor/UMGEditor/Public/Templates/WidgetTemplateBlueprintClass.h b/Engine/Source/Editor/UMGEditor/Public/Templates/WidgetTemplateBlueprintClass.h
new file mode 100644
--- /dev/null
+++ b/Engine/Source/Editor/UMGEditor/Public/Templates/WidgetTemplateBlueprintClass.h
@@ -0,0 +1,50 @@
+// Copyright Epic Games, Inc. All Rights Reserved.
+
+#pragma once
+
+#include "CoreMinimal.h"
+#include "Templates/SubclassOf.h"
+#include "Input/Reply.h"
+#include "AssetData.h"
+#include "Blueprint/UserWidget.h"
+#include "Widgets/IToolTip.h"
+#include "Templates/WidgetTemplateClass.h"
+
+class UWidgetTree;
+
+/**
+* A template for classes generated by UWidgetBlueprint assets
+* Accounts for the possibility that a the blueprint may not be loaded when the palette is created
+*/
+class UMGEDITOR_API FWidgetTemplateBlueprintClass : public FWidgetTemplateClass
+{
+public:
+	/** 
+	* Constructor
+	* @param InWidgetAssetData The FAssetData for the widget blueprint asset
+	* @param InUserWidgetClass The user-created widget class that will be created by this template
+	*/
+	FWidgetTemplateBlueprintClass(const FAssetData& InWidgetAssetData, TSubclassOf<UUserWidget> InUserWidgetClass = nullptr);
+
+	/** Destructor */
+	virtual ~FWidgetTemplateBlueprintClass();
+
+	/** Gets the category for the widget */
+	virtual FText GetCategory() const override;
+
+	/** Creates an instance of the widget for the tree. */
+	virtual UWidget* Create(UWidgetTree* Tree) override;
+
+	/** The icon coming from the default object of the class */
+	virtual const FSlateBrush* GetIcon() const override;
+
+	/** Gets the tooltip widget for this palette item. */
+	virtual TSharedRef<IToolTip> GetToolTip() const override;
+
+	/** Opens the widget blueprint for edit */
+	virtual FReply OnDoubleClicked() override;
+
+	/** Returns true if the supplied class is supported by this template */
+	static bool Supports(UClass* InClass);
+
+};
diff --git a/Engine/Source/Editor/UMGEditor/Public/Templates/WidgetTemplateClass.h b/Engine/Source/Editor/UMGEditor/Public/Templates/WidgetTemplateClass.h
new file mode 100644
--- /dev/null
+++ b/Engine/Source/Editor/UMGEditor/Public/Templates/WidgetTemplateClass.h
@@ -0,0 +1,65 @@
+// Copyright Epic Games, Inc. All Rights Reserved.
+
+#pragma once
+
+#include "CoreMinimal.h"
+#include "Templates/SubclassOf.h"
+#include "Components/Widget.h"
+#include "Widgets/IToolTip.h"
+#include "WidgetTemplate.h"
+#include "AssetData.h"
+
+class UWidgetTree;
+
+/**
+ * A template that can spawn any widget derived from the UWidget class.
+ */
+class UMGEDITOR_API FWidgetTemplateClass : public FWidgetTemplate
+{
+public:
+	/** Constructor */
+	explicit FWidgetTemplateClass(TSubclassOf<UWidget> InWidgetClass);
+
+	explicit FWidgetTemplateClass(const FAssetData& InWidgetAssetData, TSubclassOf<UWidget> InWidgetClass);
+
+	/** Destructor */
+	virtual ~FWidgetTemplateClass();
+
+	/** Gets the category for the widget */
+	virtual FText GetCategory() const override;
+
+	/** Creates an instance of the widget for the tree. */
+	virtual UWidget* Create(UWidgetTree* Tree) override;
+
+	/** The icon coming from the default object of the class */
+	virtual const FSlateBrush* GetIcon() const override;
+
+	/** Gets the tooltip widget for this palette item. */
+	virtual TSharedRef<IToolTip> GetToolTip() const override;
+
+	/** @param OutStrings - Returns an array of strings used for filtering/searching this widget template. */
+	virtual void GetFilterStrings(TArray<FString>& OutStrings) const override;
+
+	/** Gets the WidgetClass which might be null. */
+	TWeakObjectPtr<UClass> GetWidgetClass() const { return WidgetClass; }
+
+	/** Returns the asset data for this widget which might be invalid. */
+	const FAssetData& GetWidgetAssetData() { return WidgetAssetData; }
+
+protected:
+	/** Creates a widget template class without any class reference */
+	FWidgetTemplateClass();
+
+	/** Called when objects need to be swapped out for new versions, like after a blueprint recompile. */
+	void OnObjectsReplaced(const TMap<UObject*, UObject*>& ReplacementMap);
+
+	/** Constructs the widget template with an overridden object name. */
+	UWidget* CreateNamed(class UWidgetTree* Tree, FName NameOverride);
+
+protected:
+	/** The widget class that will be created by this template */
+	TWeakObjectPtr<UClass> WidgetClass;
+
+	/** The asset data for the widget blueprint */
+	FAssetData WidgetAssetData;
+};
diff --git a/Engine/Source/Editor/UMGEditor/Public/Templates/WidgetTemplateImageClass.h b/Engine/Source/Editor/UMGEditor/Public/Templates/WidgetTemplateImageClass.h
new file mode 100644
--- /dev/null
+++ b/Engine/Source/Editor/UMGEditor/Public/Templates/WidgetTemplateImageClass.h
@@ -0,0 +1,32 @@
+// Copyright Epic Games, Inc. All Rights Reserved.
+
+#pragma once
+
+#include "CoreMinimal.h"
+#include "AssetData.h"
+#include "Templates/WidgetTemplateClass.h"
+
+class UWidgetTree;
+
+/**
+ * A template for classes generated by UTexture or UMaterial classes, or implements a USlateTextureAtlasInterface
+ */
+
+class UMGEDITOR_API FWidgetTemplateImageClass : public FWidgetTemplateClass
+{
+
+public:
+	/**
+	 * Constructor.
+	 * @param	InAssetData		The asset data used to create the widget
+	 */
+	FWidgetTemplateImageClass(const FAssetData& InAssetData);
+
+	virtual ~FWidgetTemplateImageClass();
+
+	/** Creates an instance of the widget for the tree */
+	virtual UWidget* Create(UWidgetTree* WidgetTree) override;
+
+	/** Returns true if the supplied class is supported by this template */
+	static bool Supports(UClass* InClass);
+};
diff --git a/Engine/Source/Editor/UMGEditor/Public/UMGEditorModule.h b/Engine/Source/Editor/UMGEditor/Public/UMGEditorModule.h
--- a/Engine/Source/Editor/UMGEditor/Public/UMGEditorModule.h
+++ b/Engine/Source/Editor/UMGEditor/Public/UMGEditorModule.h
@@ -26,4 +26,16 @@ public:
 
 	DECLARE_EVENT_TwoParams(IUMGEditorModule, FOnRegisterTabs, const FWidgetBlueprintApplicationMode&, FWorkflowAllowedTabSet&);
 	virtual FOnRegisterTabs& OnRegisterTabsForEditor() = 0;
+
+	DECLARE_DELEGATE_RetVal_TwoParams(TSharedRef<FExtender>, FWidgetEditorToolbarExtender, const TSharedRef<FUICommandList>, TSharedRef<FWidgetBlueprintEditor>);
+
+	/** Add Toolbar extender for widget editor, used by widget tool palette. */
+	virtual void AddWidgetEditorToolbarExtender(FWidgetEditorToolbarExtender&& InToolbarExtender) = 0;
+
+	/** Get current toolbar extenders for widget editor */
+	virtual TArrayView<FWidgetEditorToolbarExtender> GetAllWidgetEditorToolbarExtenders() = 0;
+
+	/** Support for general layout extenders */
+	DECLARE_EVENT_OneParam(IStaticMeshEditor, FOnRegisterLayoutExtensions, FLayoutExtender&);
+	virtual FOnRegisterLayoutExtensions& OnRegisterLayoutExtensions() = 0;
 };
diff --git a/Engine/Source/Editor/UMGEditor/Public/WidgetBlueprintEditor.h b/Engine/Source/Editor/UMGEditor/Public/WidgetBlueprintEditor.h
--- a/Engine/Source/Editor/UMGEditor/Public/WidgetBlueprintEditor.h
+++ b/Engine/Source/Editor/UMGEditor/Public/WidgetBlueprintEditor.h
@@ -67,11 +67,15 @@ public:
 	virtual void PostUndo(bool bSuccessful) override;
 	virtual void PostRedo(bool bSuccessful) override;
 	virtual void Compile() override;
-	//~ End Begin FBlueprintEditor interface
+	//~ End FBlueprintEditor interface
+	
+	//~ Begin FAssetEditorToolkit Interface
+	virtual bool OnRequestClose() override;
+	// End of FAssetEditorToolkit 
 
 	//~ Begin FGCObjectInterface interface
 	virtual void AddReferencedObjects( FReferenceCollector& Collector ) override;
-	//~ end FGCObjectInterface interface
+	//~ End FGCObjectInterface interface
 
 	//~ Begin IToolkit interface
 	virtual FName GetToolkitFName() const override;
@@ -79,6 +83,8 @@ public:
 	virtual FString GetWorldCentricTabPrefix() const override;
 	virtual FLinearColor GetWorldCentricTabColorScale() const override;
 	virtual void InitToolMenuContext(FToolMenuContext& MenuContext) override;
+	void OnToolkitHostingStarted(const TSharedRef<IToolkit>& Toolkit) override;
+	void OnToolkitHostingFinished(const TSharedRef<IToolkit>& Toolkit) override;
 	//~ End IToolkit interface
 
 	/** @return The widget blueprint currently being edited in this editor */
@@ -224,6 +230,8 @@ public:
 	TSharedPtr<FPaletteViewModel> GetPaletteViewModel() { return PaletteViewModel; };
 	TSharedPtr<FLibraryViewModel> GetLibraryViewModel() { return LibraryViewModel; };
 
+	void CreateEditorModeManager() override;
+
 public:
 	/** Fires whenever a new widget is being hovered over */
 	FOnHoveredWidgetSet OnHoveredWidgetSet;
@@ -249,6 +257,9 @@ public:
 	/** Command list for handling widget actions in the WidgetBlueprintEditor */
 	TSharedPtr< FUICommandList > DesignerCommandList;
 
+	/** Commands for switching between tool palettes */
+	TArray< TSharedPtr<FUICommandInfo>> ToolPaletteCommands;
+
 	/** Paste Metadata */
 	FVector2D PasteDropLocation;
 
@@ -441,6 +452,9 @@ private:
 	/** The toolbar builder associated with this editor */
 	TSharedPtr<class FWidgetBlueprintEditorToolbar> WidgetToolbar;
 
+	/** Used to spawn sidebar tool palette */
+	TSharedPtr<class FWidgetEditorModeUILayer> ModeUILayer;
+
 	/** The widget references out in the ether that may need to be updated after being issued. */
 	TArray< TWeakPtr<FWidgetHandle> > WidgetHandlePool;
 
diff --git a/Engine/Source/Editor/UMGEditor/Public/WidgetBlueprintEditorUtils.h b/Engine/Source/Editor/UMGEditor/Public/WidgetBlueprintEditorUtils.h
new file mode 100644
--- /dev/null
+++ b/Engine/Source/Editor/UMGEditor/Public/WidgetBlueprintEditorUtils.h
@@ -0,0 +1,134 @@
+// Copyright Epic Games, Inc. All Rights Reserved.
+
+#pragma once
+
+#include "CoreMinimal.h"
+#include "WidgetBlueprintEditor.h"
+#include "WidgetBlueprint.h"
+#include "WidgetReference.h"
+
+class FDragDropOperation;
+class FHittestGrid;
+class FMenuBuilder;
+class UWidgetBlueprint;
+class UWidgetSlotPair;
+class UWidgetTree;
+class SWindow;
+
+//////////////////////////////////////////////////////////////////////////
+// FWidgetBlueprintEditorUtils
+
+class FWidgetBlueprintEditorUtils
+{
+
+public:
+	struct FWidgetThumbnailProperties
+	{
+		FVector2D ScaledSize;
+		FVector2D Offset;
+	};
+
+	static bool VerifyWidgetRename(TSharedRef<class FWidgetBlueprintEditor> BlueprintEditor, FWidgetReference Widget, const FText& NewName, FText& OutErrorMessage);
+
+	static bool RenameWidget(TSharedRef<class FWidgetBlueprintEditor> BlueprintEditor, const FName& OldObjectName, const FString& NewDisplayName);
+
+	static void CreateWidgetContextMenu(FMenuBuilder& MenuBuilder, TSharedRef<FWidgetBlueprintEditor> BlueprintEditor, FVector2D TargetLocation);
+
+	static void CopyWidgets(UWidgetBlueprint* BP, TSet<FWidgetReference> Widgets);
+
+	static TArray<UWidget*> PasteWidgets(TSharedRef<FWidgetBlueprintEditor> BlueprintEditor, UWidgetBlueprint* BP, FWidgetReference ParentWidget, FName SlotName, FVector2D PasteLocation);
+
+	static void DeleteWidgets(UWidgetBlueprint* BP, TSet<FWidgetReference> Widgets, bool bSilentDelete = false);
+
+	static void CutWidgets(UWidgetBlueprint* BP, TSet<FWidgetReference> Widgets);
+
+	static TArray<UWidget*> DuplicateWidgets(TSharedRef<FWidgetBlueprintEditor> BlueprintEditor, UWidgetBlueprint* BP, TSet<FWidgetReference> Widgets);
+
+	static UWidget* GetWidgetTemplateFromDragDrop(UWidgetBlueprint* Blueprint, UWidgetTree* RootWidgetTree, TSharedPtr<FDragDropOperation>& DragDropOp);
+
+	static bool IsBindWidgetProperty(FProperty* InProperty);
+	static bool IsBindWidgetProperty(FProperty* InProperty, bool& bIsOptional);
+
+	static bool IsBindWidgetAnimProperty(FProperty* InProperty);
+	static bool IsBindWidgetAnimProperty(FProperty* InProperty, bool& bIsOptional);
+
+	static bool IsUsableWidgetClass(UClass* WidgetClass);
+
+	static void ExportWidgetsToText(TArray<UWidget*> WidgetsToExport, /*out*/ FString& ExportedText);
+
+	static void ImportWidgetsFromText(UWidgetBlueprint* BP, const FString& TextToImport, /*out*/ TSet<UWidget*>& ImportedWidgetSet, /*out*/ TMap<FName, UWidgetSlotPair*>& PastedExtraSlotData);
+
+	/** Exports the individual properties of an object to text and stores them in a map. */
+	static void ExportPropertiesToText(UObject* Object, TMap<FName, FString>& ExportedProperties);
+
+	/** Attempts to import any property in the map and apply it to a property with the same name on the object. */
+	static void ImportPropertiesFromText(UObject* Object, const TMap<FName, FString>& ExportedProperties);
+
+	static TScriptInterface<INamedSlotInterface> FindNamedSlotHostForContent(UWidget* WidgetTemplate, UWidgetTree* WidgetTree);
+
+	static UWidget* FindNamedSlotHostWidgetForContent(UWidget* WidgetTemplate, UWidgetTree* WidgetTree);
+
+	static void FindAllAncestorNamedSlotHostWidgetsForContent(TArray<FWidgetReference>& OutSlotHostWidgets, UWidget* WidgetTemplate, TSharedRef<FWidgetBlueprintEditor> BlueprintEditor);
+
+	static bool RemoveNamedSlotHostContent(UWidget* WidgetTemplate, TScriptInterface<INamedSlotInterface> NamedSlotHost);
+
+	static bool ReplaceNamedSlotHostContent(UWidget* WidgetTemplate, TScriptInterface<INamedSlotInterface> NamedSlotHost, UWidget* NewContentWidget);
+
+	static int32 UpdateHittestGrid(FHittestGrid& HitTestGrid, TSharedRef<SWindow> Window, float Scale, FVector2D DrawSize, float DeltaTime);
+
+	static TTuple<FVector2D, FVector2D> GetWidgetPreviewAreaAndSize(UUserWidget* UserWidget, FVector2D DesiredSize, FVector2D PreviewSize, EDesignPreviewSizeMode ThumbnailSizeMode, TOptional<FVector2D> ThumbnailCustomSize);
+
+	static float GetWidgetPreviewDPIScale(UUserWidget* UserWidget, FVector2D PreviewSize);
+
+	static EDesignPreviewSizeMode ConvertThumbnailSizeModeToDesignerSizeMode(EThumbnailPreviewSizeMode ThumbnailSizeMode, UUserWidget* WidgetInstance);
+
+	static FVector2D GetWidgetPreviewUnScaledCustomSize(FVector2D DesiredSize, UUserWidget* UserWidget, TOptional<FVector2D> ThumbnailCustomSize, EThumbnailPreviewSizeMode ThumbnailSizeMode = EThumbnailPreviewSizeMode::MatchDesignerMode);
+
+	static TTuple<float, FVector2D> GetThumbnailImageScaleAndOffset(FVector2D WidgetSize, FVector2D ThumbnailSize);
+
+	static void SetTextureAsAssetThumbnail(UWidgetBlueprint* WidgetBlueprint, UTexture2D* ThumbnailTexture);
+
+	static TOptional<FWidgetThumbnailProperties> DrawSWidgetInRenderTargetForThumbnail(UUserWidget* WidgetInstance, FRenderTarget* RenderTarget2D, FVector2D ThumbnailSize, TOptional<FVector2D> ThumbnailCustomSize, EThumbnailPreviewSizeMode ThumbnailSizeMode = EThumbnailPreviewSizeMode::MatchDesignerMode);
+
+	static TOptional<FWidgetThumbnailProperties> DrawSWidgetInRenderTargetForThumbnail(UUserWidget* WidgetInstance, UTextureRenderTarget2D* RenderTarget2D, FVector2D ThumbnailSize, TOptional<FVector2D> ThumbnailCustomSize, EThumbnailPreviewSizeMode ThumbnailSizeMode);
+
+	static TOptional<FWidgetThumbnailProperties> DrawSWidgetInRenderTarget(UUserWidget* WidgetInstance, UTextureRenderTarget2D* RenderTarget2D);
+
+private:
+
+	static FString CopyWidgetsInternal(UWidgetBlueprint* BP, TSet<FWidgetReference> Widgets);
+
+	static TArray<UWidget*> PasteWidgetsInternal(TSharedRef<FWidgetBlueprintEditor> BlueprintEditor, UWidgetBlueprint* BP, const FString& TextToImport, FWidgetReference ParentWidget, FName SlotName, FVector2D PasteLocation, bool bForceSibling, bool& TransactionSuccesful);
+
+	static void ExecuteOpenSelectedWidgetsForEdit( TSet<FWidgetReference> SelectedWidgets );
+
+	static bool FindAndRemoveNamedSlotContent(UWidget* WidgetTemplate, UWidgetTree* WidgetTree);
+
+	static bool CanOpenSelectedWidgetsForEdit( TSet<FWidgetReference> SelectedWidgets );
+
+	static void BuildWrapWithMenu(FMenuBuilder& Menu, TSharedRef<FWidgetBlueprintEditor> BlueprintEditor, UWidgetBlueprint* BP, TSet<FWidgetReference> Widgets);
+
+	static void WrapWidgets(TSharedRef<FWidgetBlueprintEditor> BlueprintEditor, UWidgetBlueprint* BP, TSet<FWidgetReference> Widgets, UClass* WidgetClass);
+
+	static void BuildReplaceWithMenu(FMenuBuilder& Menu, TSharedRef<FWidgetBlueprintEditor> BlueprintEditor, UWidgetBlueprint* BP, TSet<FWidgetReference> Widgets);
+
+	static void ReplaceWidgetWithSelectedTemplate(TSharedRef<FWidgetBlueprintEditor> BlueprintEditor, UWidgetBlueprint* BP, FWidgetReference Widget);
+
+	static bool CanBeReplacedWithTemplate(TSharedRef<FWidgetBlueprintEditor> BlueprintEditor, UWidgetBlueprint* BP, FWidgetReference Widget);
+
+	static void ReplaceWidgetWithChildren(TSharedRef<FWidgetBlueprintEditor> BlueprintEditor, UWidgetBlueprint* BP, FWidgetReference Widget);
+
+	static void ReplaceWidgetWithNamedSlot(TSharedRef<FWidgetBlueprintEditor> BlueprintEditor, UWidgetBlueprint* BP, FWidgetReference Widget, FName NamedSlot);
+
+	static void ReplaceWidgets(TSharedRef<FWidgetBlueprintEditor> BlueprintEditor, UWidgetBlueprint* BP, TSet<FWidgetReference> Widgets, UClass* WidgetClass);
+
+	static FString FindNextValidName(UWidgetTree* WidgetTree, const FString& Name);
+
+	static void FindUsedVariablesForWidgets(const TSet<FWidgetReference>& Widgets, const UWidgetBlueprint* BP, TArray<UWidget*>& UsedVariables, TArray<FText>& WidgetNames, bool bIncludeVariablesOnChildren);
+
+	static bool ShouldContinueDeleteOperation(UWidgetBlueprint* BP, const TArray<FText>& WidgetNames);
+
+	static bool ShouldContinueReplaceOperation(UWidgetBlueprint* BP, const TArray<FText>& WidgetNames);	
+
+	static TOptional<FWidgetThumbnailProperties> DrawSWidgetInRenderTargetInternal(UUserWidget* WidgetInstance, FRenderTarget* RenderTarget2D, UTextureRenderTarget2D* TextureRenderTarget,FVector2D ThumbnailSize, bool bIsForThumbnail, TOptional<FVector2D> ThumbnailCustomSize, EThumbnailPreviewSizeMode ThumbnailSizeMode);
+};
diff --git a/Engine/Source/Editor/UMGEditor/UMGEditor.Build.cs b/Engine/Source/Editor/UMGEditor/UMGEditor.Build.cs
--- a/Engine/Source/Editor/UMGEditor/UMGEditor.Build.cs
+++ b/Engine/Source/Editor/UMGEditor/UMGEditor.Build.cs
@@ -48,6 +48,7 @@ public class UMGEditor : ModuleRules
 				"Engine",
 				"AssetTools",
 				"EditorFramework",
+				"InteractiveToolsFramework",
 				"UnrealEd", // for Asset Editor Subsystem
 				"KismetWidgets",
 				"EditorWidgets",
@@ -79,6 +80,7 @@ public class UMGEditor : ModuleRules
 				"DeveloperSettings",
 				"ImageWrapper",
 				"ToolWidgets",
+				"WorkspaceMenuStructure"
 			}
 			);
 	}
diff --git a/Engine/Source/Editor/UnrealEd/Private/Toolkits/BaseToolkit.cpp b/Engine/Source/Editor/UnrealEd/Private/Toolkits/BaseToolkit.cpp
--- a/Engine/Source/Editor/UnrealEd/Private/Toolkits/BaseToolkit.cpp
+++ b/Engine/Source/Editor/UnrealEd/Private/Toolkits/BaseToolkit.cpp
@@ -396,10 +396,8 @@ void FModeToolkit::InvokeUI()
 				GetToolPaletteNames(PaletteNames);
 				for (auto Palette : PaletteNames)
 				{
-					FUniformToolBarBuilder ModeToolbarBuilder(CommandList, FMultiBoxCustomization(ScriptableMode->GetModeInfo().ToolbarCustomizationName), TSharedPtr<FExtender>(), false);
-					ModeToolbarBuilder.SetStyle(&FEditorStyle::Get(), "PaletteToolBar");
-					BuildToolPalette(Palette, ModeToolbarBuilder);
-					ActiveToolBarRows.Emplace(ScriptableMode->GetID(), Palette, GetToolPaletteDisplayName(Palette), ModeToolbarBuilder.MakeWidget());
+					TSharedRef<SWidget> PaletteWidget = CreatePaletteWidget(CommandList, ScriptableMode->GetModeInfo().ToolbarCustomizationName, Palette);
+					ActiveToolBarRows.Emplace(ScriptableMode->GetID(), Palette, GetToolPaletteDisplayName(Palette), PaletteWidget);
 				}
 			}
 			TSharedPtr<SDockTab> CreatedToolbarTab = ModeUILayerPtr->GetTabManager()->TryInvokeTab(UAssetEditorUISubsystem::VerticalToolbarID);
@@ -510,12 +508,7 @@ void FModeToolkit::UpdatePrimaryModePanel()
 			for (auto Palette : PaletteNames)
 			{
 				FName ToolbarCustomizationName = GetEditorMode() ?  GetEditorMode()->GetModeInfo().ToolbarCustomizationName : GetScriptableEditorMode()->GetModeInfo().ToolbarCustomizationName;
-				FUniformToolBarBuilder ModeToolbarBuilder(CommandList, FMultiBoxCustomization(ToolbarCustomizationName));
-				ModeToolbarBuilder.SetStyle(&FEditorStyle::Get(), "PaletteToolBar");
-
-				BuildToolPalette(Palette, ModeToolbarBuilder);
-
-				TSharedRef<SWidget> PaletteWidget = ModeToolbarBuilder.MakeWidget();
+				TSharedRef<SWidget> PaletteWidget = CreatePaletteWidget(CommandList, ToolbarCustomizationName, Palette);
 
 				const bool bRebuildChildren = false;
 				PaletteTabBox->AddSlot(Palette, false)
@@ -662,6 +655,16 @@ bool FModeToolkit::ShouldShowModeToolbar() const
 	return ActiveToolBarRows.Num() > 0;
 }
 
+TSharedRef<SWidget> FModeToolkit::CreatePaletteWidget(TSharedPtr<FUICommandList> InCommandList, FName InToolbarCustomizationName, FName InPaletteName)
+{
+	FUniformToolBarBuilder ModeToolbarBuilder(InCommandList, FMultiBoxCustomization(InToolbarCustomizationName));
+	ModeToolbarBuilder.SetStyle(&FEditorStyle::Get(), "PaletteToolBar");
+
+	BuildToolPalette(InPaletteName, ModeToolbarBuilder);
+
+	return ModeToolbarBuilder.MakeWidget();
+}
+
 void FModeToolkit::SpawnOrUpdateModeToolbar()
 {
 	if (ShouldShowModeToolbar())
diff --git a/Engine/Source/Editor/UnrealEd/Private/Tools/AssetEditorContextObject.cpp b/Engine/Source/Editor/UnrealEd/Private/Tools/AssetEditorContextObject.cpp
--- a/Engine/Source/Editor/UnrealEd/Private/Tools/AssetEditorContextObject.cpp
+++ b/Engine/Source/Editor/UnrealEd/Private/Tools/AssetEditorContextObject.cpp
@@ -47,7 +47,24 @@ UWorld* UAssetEditorContextObject::GetEditingWorld() const
 	return nullptr;
 }
 
+const IToolkitHost* UAssetEditorContextObject::GetToolkitHost() const
+{
+	return ToolkitHost;
+}
+
+IToolkitHost* UAssetEditorContextObject::GetMutableToolkitHost()
+{
+	return ToolkitHost;
+}
+
+void UAssetEditorContextObject::SetToolkitHost(IToolkitHost* InToolkitHost)
+{
+	ToolkitHost = InToolkitHost;
+}
+
 UWorld* UAssetEditorContextObject::GetWorld() const
 {
 	return GetEditingWorld();
 }
+
+
diff --git a/Engine/Source/Editor/UnrealEd/Private/Tools/AssetEditorContextObject.h b/Engine/Source/Editor/UnrealEd/Private/Tools/AssetEditorContextObject.h
--- a/Engine/Source/Editor/UnrealEd/Private/Tools/AssetEditorContextObject.h
+++ b/Engine/Source/Editor/UnrealEd/Private/Tools/AssetEditorContextObject.h
@@ -2,14 +2,13 @@
 
 #pragma once
 
-#include "Elements/Framework/TypedElementAssetEditorToolkitHostMixin.h"
 #include "Tools/AssetEditorContextInterface.h"
 #include "UObject/Object.h"
 
 #include "AssetEditorContextObject.generated.h"
 
 UCLASS(Transient)
-class UAssetEditorContextObject : public UObject, public IAssetEditorContextInterface, public FTypedElementAssetEditorToolkitHostMixin
+class UAssetEditorContextObject : public UObject, public IAssetEditorContextInterface
 {
 	GENERATED_BODY()
 
@@ -19,9 +18,20 @@ public:
 	UTypedElementSelectionSet* GetMutableSelectionSet() override;
 	UTypedElementCommonActions* GetCommonActions() override;
 	UWorld* GetEditingWorld() const override;
+	const IToolkitHost* GetToolkitHost() const override;
+	IToolkitHost* GetMutableToolkitHost() override;
 	// End UAssetEditorContextObject interface
 
+	/**
+	 * Set the toolkit host associated with our asset editor.
+	 */
+	void SetToolkitHost(IToolkitHost* InToolkitHost);
+
 	// Begin UObject interface
 	UWorld* GetWorld() const override;
 	// End UObject interface
+
+private:
+	/** The toolkit host associated with our asset editor. */
+	IToolkitHost* ToolkitHost = nullptr;
 };
\ No newline at end of file
diff --git a/Engine/Source/Editor/UnrealEd/Private/WidgetModeManager.cpp b/Engine/Source/Editor/UnrealEd/Private/WidgetModeManager.cpp
new file mode 100644
--- /dev/null
+++ b/Engine/Source/Editor/UnrealEd/Private/WidgetModeManager.cpp
@@ -0,0 +1,141 @@
+// Copyright Epic Games, Inc. All Rights Reserved.
+
+#include "WidgetModeManager.h"
+
+#include "ToolContexts/WidgetToolsContext.h"
+
+//////////////////////////////////////////////////////////////////////////
+// FWidgetModeManager
+
+FWidgetModeManager::FWidgetModeManager()
+	: FEditorModeTools()
+{
+	// Since we can't call virtuals in constructor, need free and re-create InteractiveToolsContext resource
+	if (UObjectInitialized())
+	{
+		InteractiveToolsContext->ShutdownContext();
+		InteractiveToolsContext = nullptr;
+	}
+
+	CachedWidgetToolContext = NewObject<UWidgetToolsContext>(GetTransientPackage(), UWidgetToolsContext::StaticClass(), NAME_None, RF_Transient);
+	InteractiveToolsContext = CachedWidgetToolContext;
+	InteractiveToolsContext->InitializeContextWithEditorModeManager(this);
+}
+
+bool FWidgetModeManager::OnKeyChar(const FGeometry& MyGeometry, const FCharacterEvent& InCharacterEvent)
+{
+	if (CachedWidgetToolContext)
+	{
+		return CachedWidgetToolContext->OnKeyChar(MyGeometry, InCharacterEvent);
+	}
+
+	return false;
+}
+
+bool FWidgetModeManager::OnKeyDown(const FGeometry& MyGeometry, const FKeyEvent& InKeyEvent)
+{
+	if (CachedWidgetToolContext)
+	{
+		return CachedWidgetToolContext->OnKeyDown(MyGeometry, InKeyEvent);
+	}
+
+	return false;
+}
+
+bool FWidgetModeManager::OnKeyUp(const FGeometry& MyGeometry, const FKeyEvent& InKeyEvent)
+{
+	if (CachedWidgetToolContext)
+	{
+		return CachedWidgetToolContext->OnKeyUp(MyGeometry, InKeyEvent);
+	}
+
+	return false;
+}
+
+bool FWidgetModeManager::OnMouseButtonDown(const FGeometry& MyGeometry, const FPointerEvent& MouseEvent)
+{
+	if (CachedWidgetToolContext)
+	{
+		return CachedWidgetToolContext->OnMouseButtonDown(MyGeometry, MouseEvent);
+	}
+
+	return false;
+}
+
+bool FWidgetModeManager::OnMouseButtonUp(const FGeometry& MyGeometry, const FPointerEvent& MouseEvent)
+{
+	if (CachedWidgetToolContext)
+	{
+		return CachedWidgetToolContext->OnMouseButtonUp(MyGeometry, MouseEvent);
+	}
+
+	return false;
+}
+
+bool FWidgetModeManager::OnMouseMove(const FGeometry& MyGeometry, const FPointerEvent& MouseEvent)
+{
+	if (CachedWidgetToolContext)
+	{
+		return CachedWidgetToolContext->OnMouseMove(MyGeometry, MouseEvent);
+	}
+
+	return false;
+}
+
+bool FWidgetModeManager::OnMouseWheel(const FGeometry& MyGeometry, const FPointerEvent& MouseEvent)
+{
+	if (CachedWidgetToolContext)
+	{
+		return CachedWidgetToolContext->OnMouseWheel(MyGeometry, MouseEvent);
+	}
+
+	return false;
+}
+
+bool FWidgetModeManager::OnMouseButtonDoubleClick(const FGeometry& InMyGeometry, const FPointerEvent& InMouseEvent)
+{
+	if (CachedWidgetToolContext)
+	{
+		return CachedWidgetToolContext->OnMouseButtonDoubleClick(InMyGeometry, InMouseEvent);
+	}
+
+	return false;
+}
+
+bool FWidgetModeManager::OnDragDetected(const FGeometry& MyGeometry, const FPointerEvent& MouseEvent)
+{
+	if (CachedWidgetToolContext)
+	{
+		return CachedWidgetToolContext->OnDragDetected(MyGeometry, MouseEvent);
+	}
+
+	return false;
+}
+
+void FWidgetModeManager::OnMouseCaptureLost(const FCaptureLostEvent& CaptureLostEvent)
+{
+	if (CachedWidgetToolContext)
+	{
+		CachedWidgetToolContext->OnMouseCaptureLost(CaptureLostEvent);
+	}
+}
+
+int32 FWidgetModeManager::OnPaint(const FPaintArgs& Args, const FGeometry& AllottedGeometry, const FSlateRect& MyCullingRect, FSlateWindowElementList& OutDrawElements, int32 LayerId, const FWidgetStyle& InWidgetStyle, bool bParentEnabled)
+{
+	if (CachedWidgetToolContext)
+	{
+		return CachedWidgetToolContext->OnPaint(Args, AllottedGeometry, MyCullingRect, OutDrawElements, LayerId, InWidgetStyle, bParentEnabled);
+	}
+
+	return LayerId;
+}
+
+void FWidgetModeManager::SetManagedWidget(TSharedPtr<SWidget> InManagedWidget)
+{
+	ManagedWidget = InManagedWidget;
+}
+
+TSharedPtr<SWidget> FWidgetModeManager::GetManagedWidget() const
+{
+	return ManagedWidget.Pin();
+}
diff --git a/Engine/Source/Editor/UnrealEd/Public/Toolkits/BaseToolkit.h b/Engine/Source/Editor/UnrealEd/Public/Toolkits/BaseToolkit.h
--- a/Engine/Source/Editor/UnrealEd/Public/Toolkits/BaseToolkit.h
+++ b/Engine/Source/Editor/UnrealEd/Public/Toolkits/BaseToolkit.h
@@ -176,8 +176,8 @@ protected:
 	const FEditorModeInfo* GetEditorModeInfo() const;
 
 	/**
- * Whether or not the mode toolbar should be shown.  If any active modes generated a toolbar this method will return true
- */
+	 * Whether or not the mode toolbar should be shown.  If any active modes generated a toolbar this method will return true
+	 */
 	bool ShouldShowModeToolbar() const;
 	TSharedRef<SDockTab> CreatePrimaryModePanel(const FSpawnTabArgs& Args);
 	void UpdatePrimaryModePanel();
@@ -189,7 +189,10 @@ protected:
 	 */
 	TSharedRef<SDockTab> MakeModeToolbarTab(const FSpawnTabArgs& Args);
 
-
+	/**
+	 * Creates the entire tool palette widget, override to specify toolbar style
+	 */
+	virtual TSharedRef<SWidget> CreatePaletteWidget(TSharedPtr<FUICommandList> InCommandList, FName InToolbarCustomizationName, FName InPaletteName);
 
 	void SpawnOrUpdateModeToolbar();
 	void RebuildModeToolBar();
diff --git a/Engine/Source/Editor/UnrealEd/Public/Tools/ToolCompatible.h b/Engine/Source/Editor/UnrealEd/Public/Tools/ToolCompatible.h
new file mode 100644
--- /dev/null
+++ b/Engine/Source/Editor/UnrealEd/Public/Tools/ToolCompatible.h
@@ -0,0 +1,68 @@
+// Copyright Epic Games, Inc. All Rights Reserved.
+
+#pragma once
+
+#include "CoreMinimal.h"
+
+#include "WidgetModeManager.h"
+#include "Input/Events.h"
+#include "Templates/SharedPointer.h"
+
+/**
+ * Mixin that handles complexity of tool support for you. Calling base widget methods only when a tool does not handle them if appropriate.
+ * 
+ * Note: Widget must have overriden SWidget::OnPaint to be protected at a minimum to use. Also currently does not support touch or gamepad navigation. 
+ *  
+ * @TODO: Does not currently allow for ArrangeChildren support this might be useful for rendering widgets on top of the tool supporting widget?
+ * A general overlay that accepts addtional widgets to render on top of the tool supporting widget might be interesting. For instance, an editable
+ * preview graph / editor of the selected widget might work this way (Such as for graph nodes). 
+ * 
+ */
+template <typename ToolCompatibleWidget>
+class TToolCompatibleMixin : public ToolCompatibleWidget
+{
+	static_assert(std::is_base_of<SWidget, ToolCompatibleWidget>::value, "TToolCompatibleMixin only works with SWidgets.");
+public:
+
+	/** Forwarding constructor to ensure mixin members always initialized (Currently none) */
+	template<typename ... Args>
+	TToolCompatibleMixin(Args&&... args)
+		: ToolCompatibleWidget(Forward<Args>(args)...)
+	{ }
+
+	// ~Begin SWidget Interface
+	virtual void Tick(const FGeometry& AllottedGeometry, const double InCurrentTime, const float InDeltaTime) override;
+	virtual FReply OnFocusReceived(const FGeometry& MyGeometry, const FFocusEvent& InFocusEvent) override;
+	virtual void OnFocusLost(const FFocusEvent& InFocusEvent) override;
+	virtual FReply OnKeyChar(const FGeometry& MyGeometry, const FCharacterEvent& InCharacterEvent) override;
+	virtual FReply OnKeyDown(const FGeometry& MyGeometry, const FKeyEvent& InKeyEvent) override;
+	virtual FReply OnKeyUp(const FGeometry& MyGeometry, const FKeyEvent& InKeyEvent) override;
+	virtual FReply OnMouseButtonDown(const FGeometry& MyGeometry, const FPointerEvent& MouseEvent) override;
+	virtual FReply OnMouseButtonUp(const FGeometry& MyGeometry, const FPointerEvent& MouseEvent) override;
+	virtual FReply OnMouseMove(const FGeometry& MyGeometry, const FPointerEvent& MouseEvent) override;
+	virtual void OnMouseEnter(const FGeometry& MyGeometry, const FPointerEvent& MouseEvent) override;
+	virtual void OnMouseLeave(const FPointerEvent& MouseEvent) override;
+	virtual FReply OnMouseWheel(const FGeometry& MyGeometry, const FPointerEvent& MouseEvent) override;
+	virtual FReply OnMouseButtonDoubleClick(const FGeometry& InMyGeometry, const FPointerEvent& InMouseEvent) override;
+	virtual FReply OnDragDetected(const FGeometry& MyGeometry, const FPointerEvent& MouseEvent) override;
+	virtual void OnMouseCaptureLost(const FCaptureLostEvent& CaptureLostEvent) override;
+	virtual TOptional<EMouseCursor::Type> GetCursor() const override;
+	virtual int32 OnPaint(const FPaintArgs& Args, const FGeometry& AllottedGeometry, const FSlateRect& MyCullingRect, FSlateWindowElementList& OutDrawElements, int32 LayerId, const FWidgetStyle& InWidgetStyle, bool bParentEnabled) const override;
+	// ~End SWidget Interface
+
+	/** Use this if your widget is created after your editor's mode manager */
+	void SetParentToolkit(TSharedPtr<IToolkit> InParentToolkit);
+
+	/** Use this if your widget is created after your editor's mode manager */
+	TSharedPtr<IToolkit> GetParentToolkit();
+
+	/** Helper method for accessing our toolkit's widget mode tool manager */
+	FWidgetModeManager* GetWidgetModeManger() const;
+
+protected:
+	/** Toolkit containing this tool compatible widget. */
+	TWeakPtr<IToolkit> ParentToolkit;
+
+};
+
+#include "Tools/ToolCompatible.inl"
diff --git a/Engine/Source/Editor/UnrealEd/Public/Tools/ToolCompatible.inl b/Engine/Source/Editor/UnrealEd/Public/Tools/ToolCompatible.inl
new file mode 100644
--- /dev/null
+++ b/Engine/Source/Editor/UnrealEd/Public/Tools/ToolCompatible.inl
@@ -0,0 +1,231 @@
+// Copyright Epic Games, Inc. All Rights Reserved.
+
+#pragma once
+
+template <typename ToolCompatibleWidget> class TToolCompatibleMixin;
+
+template <typename ToolCompatibleWidget>
+void TToolCompatibleMixin<ToolCompatibleWidget>::Tick(const FGeometry& AllottedGeometry, const double InCurrentTime, const float InDeltaTime)
+{
+	ToolCompatibleWidget::Tick(AllottedGeometry, InCurrentTime, InDeltaTime);
+	if (FWidgetModeManager* WidgetModeManager = GetWidgetModeManger())
+	{
+		WidgetModeManager->Tick(nullptr, InDeltaTime);
+	}
+}
+
+template <typename ToolCompatibleWidget>
+FReply TToolCompatibleMixin<ToolCompatibleWidget>::OnFocusReceived(const FGeometry& MyGeometry, const FFocusEvent& InFocusEvent)
+{
+	FReply Reply = ToolCompatibleWidget::OnFocusReceived(MyGeometry, InFocusEvent);
+	if (FWidgetModeManager* WidgetModeManager = GetWidgetModeManger())
+	{
+		WidgetModeManager->ReceivedFocus(nullptr, nullptr);
+	}
+	return Reply;
+}
+
+template <typename ToolCompatibleWidget>
+void TToolCompatibleMixin<ToolCompatibleWidget>::OnFocusLost(const FFocusEvent& InFocusEvent)
+{
+	ToolCompatibleWidget::OnFocusLost(InFocusEvent);
+	if (FWidgetModeManager* WidgetModeManager = GetWidgetModeManger())
+	{
+		WidgetModeManager->LostFocus(nullptr, nullptr);
+	}
+}
+
+template <typename ToolCompatibleWidget>
+FReply TToolCompatibleMixin<ToolCompatibleWidget>::OnKeyChar(const FGeometry& MyGeometry, const FCharacterEvent& InCharacterEvent)
+{
+	if (FWidgetModeManager* WidgetModeManager = GetWidgetModeManger())
+	{
+		if (WidgetModeManager->OnKeyChar(MyGeometry, InCharacterEvent))
+		{
+			return FReply::Handled();
+		}
+	}
+	return ToolCompatibleWidget::OnKeyChar(MyGeometry, InCharacterEvent);
+}
+
+template <typename ToolCompatibleWidget>
+FReply TToolCompatibleMixin<ToolCompatibleWidget>::OnKeyDown(const FGeometry& MyGeometry, const FKeyEvent& InKeyEvent)
+{
+	if (FWidgetModeManager* WidgetModeManager = GetWidgetModeManger())
+	{
+		if (WidgetModeManager->OnKeyDown(MyGeometry, InKeyEvent))
+		{
+			return FReply::Handled();
+		}
+	}
+	return ToolCompatibleWidget::OnKeyDown(MyGeometry, InKeyEvent);
+}
+
+template <typename ToolCompatibleWidget>
+FReply TToolCompatibleMixin<ToolCompatibleWidget>::OnKeyUp(const FGeometry& MyGeometry, const FKeyEvent& InKeyEvent)
+{
+	if (FWidgetModeManager* WidgetModeManager = GetWidgetModeManger())
+	{
+		if (WidgetModeManager->OnKeyUp(MyGeometry, InKeyEvent))
+		{
+			return FReply::Handled();
+		}
+	}
+	return ToolCompatibleWidget::OnKeyUp(MyGeometry, InKeyEvent);
+}
+
+template <typename ToolCompatibleWidget>
+FReply TToolCompatibleMixin<ToolCompatibleWidget>::OnMouseButtonDown(const FGeometry& MyGeometry, const FPointerEvent& MouseEvent)
+{
+	if (FWidgetModeManager* WidgetModeManager = GetWidgetModeManger())
+	{
+		if (WidgetModeManager->OnMouseButtonDown(MyGeometry, MouseEvent))
+		{
+			return FReply::Handled();
+		}
+	}
+	return ToolCompatibleWidget::OnMouseButtonDown(MyGeometry, MouseEvent);
+}
+
+template <typename ToolCompatibleWidget>
+FReply TToolCompatibleMixin<ToolCompatibleWidget>::OnMouseButtonUp(const FGeometry& MyGeometry, const FPointerEvent& MouseEvent)
+{
+	if (FWidgetModeManager* WidgetModeManager = GetWidgetModeManger())
+	{
+		if (WidgetModeManager->OnMouseButtonUp(MyGeometry, MouseEvent))
+		{
+			return FReply::Handled();
+		}
+	}
+	return ToolCompatibleWidget::OnMouseButtonUp(MyGeometry, MouseEvent);
+}
+
+template <typename ToolCompatibleWidget>
+FReply TToolCompatibleMixin<ToolCompatibleWidget>::OnMouseMove(const FGeometry& MyGeometry, const FPointerEvent& MouseEvent)
+{
+	FReply Reply = ToolCompatibleWidget::OnMouseMove(MyGeometry, MouseEvent);
+	if (FWidgetModeManager* WidgetModeManager = GetWidgetModeManger())
+	{
+		WidgetModeManager->OnMouseMove(MyGeometry, MouseEvent);
+	}
+	return Reply;
+}
+
+template <typename ToolCompatibleWidget>
+void TToolCompatibleMixin<ToolCompatibleWidget>::OnMouseEnter(const FGeometry& MyGeometry, const FPointerEvent& MouseEvent)
+{
+	ToolCompatibleWidget::OnMouseEnter(MyGeometry, MouseEvent);
+	if (FWidgetModeManager* WidgetModeManager = GetWidgetModeManger())
+	{
+		WidgetModeManager->MouseEnter(nullptr, nullptr, MouseEvent.GetScreenSpacePosition().X, MouseEvent.GetScreenSpacePosition().Y);
+	}
+}
+
+template <typename ToolCompatibleWidget>
+void TToolCompatibleMixin<ToolCompatibleWidget>::OnMouseLeave(const FPointerEvent& MouseEvent)
+{
+	ToolCompatibleWidget::OnMouseLeave(MouseEvent);
+	if (FWidgetModeManager* WidgetModeManager = GetWidgetModeManger())
+	{
+		WidgetModeManager->MouseLeave(nullptr, nullptr);
+	}
+}
+
+template <typename ToolCompatibleWidget>
+FReply TToolCompatibleMixin<ToolCompatibleWidget>::OnMouseWheel(const FGeometry& MyGeometry, const FPointerEvent& MouseEvent)
+{
+	if (FWidgetModeManager* WidgetModeManager = GetWidgetModeManger())
+	{
+		if (WidgetModeManager->OnMouseWheel(MyGeometry, MouseEvent))
+		{
+			return FReply::Handled();
+		}
+	}
+	return ToolCompatibleWidget::OnMouseWheel(MyGeometry, MouseEvent);
+}
+
+template <typename ToolCompatibleWidget>
+FReply TToolCompatibleMixin<ToolCompatibleWidget>::OnMouseButtonDoubleClick(const FGeometry& InMyGeometry, const FPointerEvent& InMouseEvent)
+{
+	if (FWidgetModeManager* WidgetModeManager = GetWidgetModeManger())
+	{
+		if (WidgetModeManager->OnMouseButtonDoubleClick(InMyGeometry, InMouseEvent))
+		{
+			return FReply::Handled();
+		}
+	}
+	return ToolCompatibleWidget::OnMouseButtonDoubleClick(InMyGeometry, InMouseEvent);
+}
+
+template <typename ToolCompatibleWidget>
+FReply TToolCompatibleMixin<ToolCompatibleWidget>::OnDragDetected(const FGeometry& MyGeometry, const FPointerEvent& MouseEvent)
+{
+	if (FWidgetModeManager* WidgetModeManager = GetWidgetModeManger())
+	{
+		if (WidgetModeManager->OnDragDetected(MyGeometry, MouseEvent))
+		{
+			return FReply::Handled();
+		}
+	}
+	return ToolCompatibleWidget::OnDragDetected(MyGeometry, MouseEvent);
+}
+
+template <typename ToolCompatibleWidget>
+void TToolCompatibleMixin<ToolCompatibleWidget>::OnMouseCaptureLost(const FCaptureLostEvent& CaptureLostEvent)
+{
+	ToolCompatibleWidget::OnMouseCaptureLost(CaptureLostEvent);
+	if (FWidgetModeManager* WidgetModeManager = GetWidgetModeManger())
+	{
+		WidgetModeManager->OnMouseCaptureLost(CaptureLostEvent);
+	}
+}
+
+template <typename ToolCompatibleWidget>
+TOptional<EMouseCursor::Type> TToolCompatibleMixin<ToolCompatibleWidget>::GetCursor() const
+{
+	if (FWidgetModeManager* WidgetModeManager = GetWidgetModeManger())
+	{
+		EMouseCursor::Type EditorModeCursor = EMouseCursor::Default;
+		if (WidgetModeManager->GetCursor(EditorModeCursor))
+		{
+			return EditorModeCursor;
+		}
+	}
+
+	return ToolCompatibleWidget::GetCursor();
+}
+
+template <typename ToolCompatibleWidget>
+int32 TToolCompatibleMixin<ToolCompatibleWidget>::OnPaint(const FPaintArgs& Args, const FGeometry& AllottedGeometry, const FSlateRect& MyCullingRect, FSlateWindowElementList& OutDrawElements, int32 LayerId, const FWidgetStyle& InWidgetStyle, bool bParentEnabled) const
+{
+	int32 OutLayerId = ToolCompatibleWidget::OnPaint(Args, AllottedGeometry, MyCullingRect, OutDrawElements, LayerId, InWidgetStyle, bParentEnabled);
+	if (FWidgetModeManager* WidgetModeManager = GetWidgetModeManger())
+	{
+		return WidgetModeManager->OnPaint(Args, AllottedGeometry, MyCullingRect, OutDrawElements, OutLayerId, InWidgetStyle, bParentEnabled);
+	}
+
+	return OutLayerId;
+}
+
+template <typename ToolCompatibleWidget>
+void TToolCompatibleMixin<ToolCompatibleWidget>::SetParentToolkit(TSharedPtr<IToolkit> InParentToolkit)
+{
+	ParentToolkit = InParentToolkit; 
+}
+
+template <typename ToolCompatibleWidget>
+TSharedPtr<IToolkit> TToolCompatibleMixin<ToolCompatibleWidget>::GetParentToolkit()
+{
+	return ParentToolkit.Pin();
+}
+
+template <typename ToolCompatibleWidget>
+FWidgetModeManager* TToolCompatibleMixin<ToolCompatibleWidget>::GetWidgetModeManger() const
+{
+	if (TSharedPtr<IToolkit> ParentToolkitPtr = ParentToolkit.Pin())
+	{
+		return static_cast<FWidgetModeManager*>(&ParentToolkitPtr->GetEditorModeManager());
+	}
+
+	return nullptr;
+}
\ No newline at end of file
diff --git a/Engine/Source/Editor/UnrealEd/Public/WidgetModeManager.h b/Engine/Source/Editor/UnrealEd/Public/WidgetModeManager.h
new file mode 100644
--- /dev/null
+++ b/Engine/Source/Editor/UnrealEd/Public/WidgetModeManager.h
@@ -0,0 +1,49 @@
+// Copyright Epic Games, Inc. All Rights Reserved.
+
+#pragma once
+
+#include "CoreMinimal.h"
+#include "Toolkits/IToolkit.h"
+#include "EditorModeManager.h"
+
+/**
+ * Tool manager that focuses on tools which the behavior of a certain widget (Ex: Design / node panels)
+ * We only provide an interface where the tool tool needs extra-arg info (Ex: paint geometry) 
+ * to perform it's specialized behavior.
+ * 
+ * In all other cases (Tick, FocusLost, etc), the default editor mode tool behavior is used and sufficient.
+ */
+class UNREALED_API FWidgetModeManager : public FEditorModeTools
+{
+public:
+	FWidgetModeManager();
+
+	// ~Begin SWidget "Proxy" Interface
+	virtual bool OnKeyChar(const FGeometry& MyGeometry, const FCharacterEvent& InCharacterEvent);
+	virtual bool OnKeyDown(const FGeometry& MyGeometry, const FKeyEvent& InKeyEvent);
+	virtual bool OnKeyUp(const FGeometry& MyGeometry, const FKeyEvent& InKeyEvent);
+	virtual bool OnMouseButtonDown(const FGeometry& MyGeometry, const FPointerEvent& MouseEvent);
+	virtual bool OnMouseButtonUp(const FGeometry& MyGeometry, const FPointerEvent& MouseEvent);
+	virtual bool OnMouseMove(const FGeometry& MyGeometry, const FPointerEvent& MouseEvent);
+	virtual bool OnMouseWheel(const FGeometry& MyGeometry, const FPointerEvent& MouseEvent);
+	virtual bool OnMouseButtonDoubleClick(const FGeometry& InMyGeometry, const FPointerEvent& InMouseEvent);
+	virtual bool OnDragDetected(const FGeometry& MyGeometry, const FPointerEvent& MouseEvent);
+	virtual void OnMouseCaptureLost(const FCaptureLostEvent& CaptureLostEvent);
+	virtual int32 OnPaint(const FPaintArgs& Args, const FGeometry& AllottedGeometry, const FSlateRect& MyCullingRect, FSlateWindowElementList& OutDrawElements, int32 LayerId, const FWidgetStyle& InWidgetStyle, bool bParentEnabled);
+	// ~End SWidget "Proxy" Interface
+
+	/** Not done in constructor as managed widget typically will not exist during mode construction for editors */
+	virtual void SetManagedWidget(TSharedPtr<SWidget> InManagedWidget);
+
+	/** Not done in constructor as managed widget typically will not exist during mode construction for editors */
+	TSharedPtr<SWidget> GetManagedWidget() const;
+	
+	/** Toolkit owning the widget of interest */
+	TWeakPtr<IToolkit> OwningToolkit;
+
+protected:
+	/** Managed widget is forwarded to tool builders / tools to interact with the widget they are operating on */
+	TWeakPtr<SWidget> ManagedWidget;
+
+	class UWidgetToolsContext* CachedWidgetToolContext;
+};
diff --git a/Engine/Source/Runtime/InteractiveToolsFramework/Private/BaseBehaviors/Widgets/WidgetBaseBehavior.cpp b/Engine/Source/Runtime/InteractiveToolsFramework/Private/BaseBehaviors/Widgets/WidgetBaseBehavior.cpp
new file mode 100644
--- /dev/null
+++ b/Engine/Source/Runtime/InteractiveToolsFramework/Private/BaseBehaviors/Widgets/WidgetBaseBehavior.cpp
@@ -0,0 +1,55 @@
+#include "BaseBehaviors/Widgets/WidgetBaseBehavior.h"
+
+bool IWidgetBaseBehavior::OnKeyChar(const FGeometry& MyGeometry, const FCharacterEvent& InCharacterEvent)
+{
+	return false;
+}
+
+bool IWidgetBaseBehavior::OnKeyDown(const FGeometry& MyGeometry, const FKeyEvent& InKeyEvent)
+{
+	return false;
+}
+
+bool IWidgetBaseBehavior::OnKeyUp(const FGeometry& MyGeometry, const FKeyEvent& InKeyEvent)
+{
+	return false;
+}
+
+bool IWidgetBaseBehavior::OnMouseButtonDown(const FGeometry& MyGeometry, const FPointerEvent& MouseEvent)
+{
+	return false;
+}
+
+bool IWidgetBaseBehavior::OnMouseButtonUp(const FGeometry& MyGeometry, const FPointerEvent& MouseEvent)
+{
+	return false;
+}
+
+bool IWidgetBaseBehavior::OnMouseMove(const FGeometry& MyGeometry, const FPointerEvent& MouseEvent)
+{
+	return false;
+}
+
+bool IWidgetBaseBehavior::OnMouseWheel(const FGeometry& MyGeometry, const FPointerEvent& MouseEvent)
+{
+	return false;
+}
+
+bool IWidgetBaseBehavior::OnMouseButtonDoubleClick(const FGeometry& InMyGeometry, const FPointerEvent& InMouseEvent)
+{
+	return false;
+}
+
+bool IWidgetBaseBehavior::OnDragDetected(const FGeometry& MyGeometry, const FPointerEvent& MouseEvent)
+{
+	return false;
+}
+
+void IWidgetBaseBehavior::OnMouseCaptureLost(const FCaptureLostEvent& CaptureLostEvent)
+{
+}
+
+int32 IWidgetBaseBehavior::OnPaint(const FPaintArgs& Args, const FGeometry& AllottedGeometry, const FSlateRect& MyCullingRect, FSlateWindowElementList& OutDrawElements, int32 LayerId, const FWidgetStyle& InWidgetStyle, bool bParentEnabled) const
+{
+	return LayerId;
+}
diff --git a/Engine/Source/Runtime/InteractiveToolsFramework/Public/BaseBehaviors/Widgets/WidgetBaseBehavior.h b/Engine/Source/Runtime/InteractiveToolsFramework/Public/BaseBehaviors/Widgets/WidgetBaseBehavior.h
new file mode 100644
--- /dev/null
+++ b/Engine/Source/Runtime/InteractiveToolsFramework/Public/BaseBehaviors/Widgets/WidgetBaseBehavior.h
@@ -0,0 +1,42 @@
+// Copyright Epic Games, Inc. All Rights Reserved.
+
+#pragma once
+
+#include "CoreMinimal.h"
+#include "UObject/Interface.h"
+#include "Layout/Geometry.h"
+#include "Input/Events.h"
+#include "Rendering/DrawElements.h"
+
+#include "WidgetBaseBehavior.generated.h"
+
+/**
+ * IWidgetBaseBehavior is an interface for tools that can modify / extend various step of a widget.
+ *
+ * If implementing this interface, please consider incrementing the LayerId if needed when overriding OnPaint.
+ */
+UINTERFACE(MinimalAPI, meta = (CannotImplementInterfaceInBlueprint))
+class UWidgetBaseBehavior : public UInterface
+{
+	GENERATED_BODY()
+};
+
+class INTERACTIVETOOLSFRAMEWORK_API IWidgetBaseBehavior
+{
+	GENERATED_BODY()
+
+public:
+	// ~Begin SWidget "Proxy" Interface
+	virtual bool OnKeyChar(const FGeometry& MyGeometry, const FCharacterEvent& InCharacterEvent);
+	virtual bool OnKeyDown(const FGeometry& MyGeometry, const FKeyEvent& InKeyEvent);
+	virtual bool OnKeyUp(const FGeometry& MyGeometry, const FKeyEvent& InKeyEvent);
+	virtual bool OnMouseButtonDown(const FGeometry& MyGeometry, const FPointerEvent& MouseEvent);
+	virtual bool OnMouseButtonUp(const FGeometry& MyGeometry, const FPointerEvent& MouseEvent);
+	virtual bool OnMouseMove(const FGeometry& MyGeometry, const FPointerEvent& MouseEvent);
+	virtual bool OnMouseWheel(const FGeometry& MyGeometry, const FPointerEvent& MouseEvent);
+	virtual bool OnMouseButtonDoubleClick(const FGeometry& InMyGeometry, const FPointerEvent& InMouseEvent);
+	virtual bool OnDragDetected(const FGeometry& MyGeometry, const FPointerEvent& MouseEvent);
+	virtual void OnMouseCaptureLost(const FCaptureLostEvent& CaptureLostEvent);
+	virtual int32 OnPaint(const FPaintArgs& Args, const FGeometry& AllottedGeometry, const FSlateRect& MyCullingRect, FSlateWindowElementList& OutDrawElements, int32 LayerId, const FWidgetStyle& InWidgetStyle, bool bParentEnabled) const;
+	// ~End SWidget "Proxy" Interface
+};
\ No newline at end of file
diff --git a/Engine/Source/Runtime/SlateCore/Public/Widgets/SWidget.h b/Engine/Source/Runtime/SlateCore/Public/Widgets/SWidget.h
--- a/Engine/Source/Runtime/SlateCore/Public/Widgets/SWidget.h
+++ b/Engine/Source/Runtime/SlateCore/Public/Widgets/SWidget.h
@@ -1331,7 +1331,7 @@ protected:
 	void SetDebugInfo( const ANSICHAR* InType, const ANSICHAR* InFile, int32 OnLine, size_t InAllocSize );
 
 	/** The cursor to show when the mouse is hovering over this widget. */
-	TOptional<EMouseCursor::Type> GetCursor() const;
+	virtual TOptional<EMouseCursor::Type> GetCursor() const;
 
 public:
 
