Expose widget application mode activation / editor shutdown.
Add interface to allow for external editor widget lifetimes to be managed by UMG editor.

diff --git a/Engine/Source/Editor/Kismet/Private/BlueprintEditor.cpp b/Engine/Source/Editor/Kismet/Private/BlueprintEditor.cpp
--- a/Engine/Source/Editor/Kismet/Private/BlueprintEditor.cpp
+++ b/Engine/Source/Editor/Kismet/Private/BlueprintEditor.cpp
@@ -4397,6 +4397,11 @@ bool FBlueprintEditor::CanAccessComponentsMode() const
 	return bCanAccess;
 }
 
+bool FBlueprintEditor::IsEditorClosing() const
+{
+	return bEditorMarkedAsClosed;
+}
+
 void FBlueprintEditor::RegisterToolbarTab(const TSharedRef<class FTabManager>& InTabManager)
 {
 	FAssetEditorToolkit::RegisterTabSpawners(InTabManager);
diff --git a/Engine/Source/Editor/Kismet/Public/BlueprintEditor.h b/Engine/Source/Editor/Kismet/Public/BlueprintEditor.h
--- a/Engine/Source/Editor/Kismet/Public/BlueprintEditor.h
+++ b/Engine/Source/Editor/Kismet/Public/BlueprintEditor.h
@@ -328,6 +328,9 @@ public:
 	/**	Returns whether the edited blueprint has components */
 	bool CanAccessComponentsMode() const;
 
+	/** Returns if we are currently closing the editor */
+	bool IsEditorClosing() const;
+
 	// @todo This is a hack for now until we reconcile the default toolbar with application modes
 	void RegisterToolbarTab(const TSharedRef<class FTabManager>& TabManager);
 	/** Throw a simple message into the log */
diff --git a/Engine/Source/Editor/UMGEditor/Private/BlueprintModes/WidgetBlueprintApplicationMode.cpp b/Engine/Source/Editor/UMGEditor/Private/BlueprintModes/WidgetBlueprintApplicationMode.cpp
--- a/Engine/Source/Editor/UMGEditor/Private/BlueprintModes/WidgetBlueprintApplicationMode.cpp
+++ b/Engine/Source/Editor/UMGEditor/Private/BlueprintModes/WidgetBlueprintApplicationMode.cpp
@@ -12,6 +12,18 @@ FWidgetBlueprintApplicationMode::FWidgetBlueprintApplicationMode(TSharedPtr<FWid
 {
 }
 
+void FWidgetBlueprintApplicationMode::PreDeactivateMode()
+{
+	FBlueprintEditorApplicationMode::PreDeactivateMode();
+	OnPreDeactivateMode.Broadcast(*this);
+}
+
+void FWidgetBlueprintApplicationMode::PostActivateMode()
+{
+	OnPostActivateMode.Broadcast(*this);
+	FBlueprintEditorApplicationMode::PostActivateMode();
+}
+
 UWidgetBlueprint* FWidgetBlueprintApplicationMode::GetBlueprint() const
 {
 	if ( FWidgetBlueprintEditor* Editor = MyWidgetBlueprintEditor.Pin().Get() )
diff --git a/Engine/Source/Editor/UMGEditor/Private/BlueprintModes/WidgetDesignerApplicationMode.cpp b/Engine/Source/Editor/UMGEditor/Private/BlueprintModes/WidgetDesignerApplicationMode.cpp
--- a/Engine/Source/Editor/UMGEditor/Private/BlueprintModes/WidgetDesignerApplicationMode.cpp
+++ b/Engine/Source/Editor/UMGEditor/Private/BlueprintModes/WidgetDesignerApplicationMode.cpp
@@ -184,11 +184,14 @@ void FWidgetDesignerApplicationMode::PreDeactivateMode()
 			PaletteTab->RequestCloseTab();
 		}
 	}
+
+	OnPreDeactivateMode.Broadcast(*this);
 }
 
 void FWidgetDesignerApplicationMode::PostActivateMode()
 {
-	//FWidgetBlueprintApplicationMode::PostActivateMode();
+	// FWidgetBlueprintApplicationMode::PostActivateMode();
+	OnPostActivateMode.Broadcast(*this);
 
 	TSharedPtr<FWidgetBlueprintEditor> BP = GetBlueprintEditor();
 
diff --git a/Engine/Source/Editor/UMGEditor/Private/WidgetBlueprintEditor.cpp b/Engine/Source/Editor/UMGEditor/Private/WidgetBlueprintEditor.cpp
--- a/Engine/Source/Editor/UMGEditor/Private/WidgetBlueprintEditor.cpp
+++ b/Engine/Source/Editor/UMGEditor/Private/WidgetBlueprintEditor.cpp
@@ -1129,6 +1129,31 @@ TSharedRef<SWidget> FWidgetBlueprintEditor::OnGetWidgetAnimSequencer()
 	return AnimDrawerWidget.ToSharedRef();
 }
 
+void FWidgetBlueprintEditor::AddExternalEditorWidget(FName ID, TSharedRef<SWidget> InExternalWidget)
+{
+	if (!ExternalEditorWidgets.Contains(ID))
+	{
+		ExternalEditorWidgets.Add(ID, InExternalWidget);
+	}
+}
+
+int32 FWidgetBlueprintEditor::RemoveExternalEditorWidget(FName ID)
+{
+	return ExternalEditorWidgets.Remove(ID);
+}
+
+TSharedPtr<SWidget> FWidgetBlueprintEditor::GetExternalEditorWidget(FName ID)
+{
+	TSharedPtr<SWidget>* ExternalWidget = ExternalEditorWidgets.Find(ID);
+
+	if (ExternalWidget)
+	{
+		return *ExternalWidget;
+	}
+
+	return nullptr;
+}
+
 void FWidgetBlueprintEditor::ToggleAnimDrawer()
 {
 	GEditor->GetEditorSubsystem<UStatusBarSubsystem>()->TryToggleDrawer(FAnimationTabSummoner::WidgetAnimSequencerDrawerID);
diff --git a/Engine/Source/Editor/UMGEditor/Public/BlueprintModes/WidgetBlueprintApplicationMode.h b/Engine/Source/Editor/UMGEditor/Public/BlueprintModes/WidgetBlueprintApplicationMode.h
--- a/Engine/Source/Editor/UMGEditor/Public/BlueprintModes/WidgetBlueprintApplicationMode.h
+++ b/Engine/Source/Editor/UMGEditor/Public/BlueprintModes/WidgetBlueprintApplicationMode.h
@@ -12,11 +12,24 @@ class UWidgetBlueprint;
 /////////////////////////////////////////////////////
 // FWidgetBlueprintApplicationMode
 
+DECLARE_MULTICAST_DELEGATE_OneParam(FOnWidgetBlueprintModeTransition, class FWidgetBlueprintApplicationMode&);
+
 class UMGEDITOR_API FWidgetBlueprintApplicationMode : public FBlueprintEditorApplicationMode
 {
 public:
 	FWidgetBlueprintApplicationMode(TSharedPtr<class FWidgetBlueprintEditor> InWidgetEditor, FName InModeName);
 
+	// FApplicationMode interface
+	virtual void PreDeactivateMode() override;
+	virtual void PostActivateMode() override;
+	// End of FApplicationMode interface
+
+	/** Called at start of PostActivateMode */
+	mutable FOnWidgetBlueprintModeTransition OnPostActivateMode;
+
+	/** Called at end of PreDeactivateMode */
+	mutable FOnWidgetBlueprintModeTransition OnPreDeactivateMode;
+
 public:
 	TSharedPtr<class FWidgetBlueprintEditor> GetBlueprintEditor() const;
 	UWidgetBlueprint* GetBlueprint() const;
diff --git a/Engine/Source/Editor/UMGEditor/Public/WidgetBlueprintEditor.h b/Engine/Source/Editor/UMGEditor/Public/WidgetBlueprintEditor.h
--- a/Engine/Source/Editor/UMGEditor/Public/WidgetBlueprintEditor.h
+++ b/Engine/Source/Editor/UMGEditor/Public/WidgetBlueprintEditor.h
@@ -189,6 +189,15 @@ public:
 	/** Gets sequencer widget, for anim drawer */
 	TSharedRef<SWidget> OnGetWidgetAnimSequencer();
 
+	/** Adds external widget whose lifetime should be managed by this editor */
+	void AddExternalEditorWidget(FName ID, TSharedRef<SWidget> InExternalWidget);
+
+	/** Removes external widget whose lifetime should be managed by this editor */
+	int32 RemoveExternalEditorWidget(FName ID);
+
+	/** Gets external widget whose lifetime should be managed by this editor, NullWidget if not found */
+	TSharedPtr<SWidget> GetExternalEditorWidget(FName ID);
+
 	/** Toggles Anim Drawer, only used for keyboard shortcut */
 	void ToggleAnimDrawer();
 
@@ -401,6 +410,9 @@ private:
 	/** The preview scene that owns the preview GUI */
 	FPreviewScene PreviewScene;
 
+	/** For external widgets to have lifetimes bound to this editor instance */
+	TMap<FName, TSharedPtr<SWidget>> ExternalEditorWidgets;
+
 	/** List of created sequencers */
 	TArray<TWeakPtr<ISequencer>> Sequencers;
 
