UMG Bind menu now adds IHasPropertyBindingExtensibility extensions to the dropdown.

#rnx

diff --git a/Engine/Source/Editor/UMGEditor/Private/Customizations/UMGDetailCustomizations.cpp b/Engine/Source/Editor/UMGEditor/Private/Customizations/UMGDetailCustomizations.cpp
--- a/Engine/Source/Editor/UMGEditor/Private/Customizations/UMGDetailCustomizations.cpp
+++ b/Engine/Source/Editor/UMGEditor/Private/Customizations/UMGDetailCustomizations.cpp
@@ -13,6 +13,7 @@
 #include "K2Node_ComponentBoundEvent.h"
 #include "Kismet2/KismetEditorUtilities.h"
 
+#include "Algo/Transform.h"
 #include "BlueprintModes/WidgetBlueprintApplicationModes.h"
 #include "DetailWidgetRow.h"
 #include "PropertyHandle.h"
@@ -27,9 +28,10 @@
 #include "Widgets/Layout/SWidgetSwitcher.h"
 #include "IDetailsView.h"
 #include "IDetailPropertyExtensionHandler.h"
+#include "IHasPropertyBindingExtensibility.h"
 #include "Binding/PropertyBinding.h"
 #include "Components/WidgetComponent.h"
-#include "Algo/Transform.h"
+#include "UMGEditorModule.h"
 
 #define LOCTEXT_NAMESPACE "UMG"
 
@@ -95,7 +97,10 @@ private:
 
 TSharedRef<SWidget> FBlueprintWidgetCustomization::MakePropertyBindingWidget(TWeakPtr<FWidgetBlueprintEditor> InEditor, FDelegateProperty* InDelegateProperty, TSharedRef<IPropertyHandle> InPropertyHandle, bool bInGeneratePureBindings)
 {
-	const FName PropertyName = InPropertyHandle->GetProperty()->GetFName();
+	if (!IModularFeatures::Get().IsModularFeatureAvailable("PropertyAccessEditor"))
+	{
+		return SNullWidget::NullWidget;
+	}
 
 	TArray<TWeakObjectPtr<UObject>> Objects;
 	{
@@ -117,399 +122,444 @@ TSharedRef<SWidget> FBlueprintWidgetCustomization::MakePropertyBindingWidget(TWe
 		WidgetName = TEXT("_") + Widget->GetName() + TEXT("_");
 	}
 
-	if(IModularFeatures::Get().IsModularFeatureAvailable("PropertyAccessEditor"))
+	UWidgetBlueprint* WidgetBlueprint = InEditor.Pin()->GetWidgetBlueprintObj();
+
+	TArray<TSharedPtr<FExtender>> MenuExtenders;
+
+	IUMGEditorModule& EditorModule = FModuleManager::LoadModuleChecked<IUMGEditorModule>("UMGEditor");
+	for (const TSharedPtr<IPropertyBindingExtension>& Extension : EditorModule.GetPropertyBindingExtensibilityManager()->GetExtensions())
 	{
-		FPropertyBindingWidgetArgs Args;
-		Args.Property = InPropertyHandle->GetProperty();
-		Args.BindableSignature = InDelegateProperty->SignatureFunction;
-		Args.OnGenerateBindingName = FOnGenerateBindingName::CreateLambda([WidgetName]()
+		if (Extension->CanExtend(WidgetBlueprint, Widget, InPropertyHandle->GetProperty()))
 		{
-			return WidgetName;
-		});
+			MenuExtenders.Add(Extension->CreateMenuExtender(WidgetBlueprint, Widget, InPropertyHandle->GetProperty()));
+		}
+	}
 
-		Args.OnGotoBinding = FOnGotoBinding::CreateLambda([InEditor, Objects](FName InPropertyName)
-		{
-			UWidgetBlueprint* ThisBlueprint = InEditor.Pin()->GetWidgetBlueprintObj();
+	FPropertyBindingWidgetArgs Args;
+	Args.MenuExtender = FExtender::Combine(MenuExtenders);
+	Args.Property = InPropertyHandle->GetProperty();
+	Args.BindableSignature = InDelegateProperty->SignatureFunction;
+	Args.OnGenerateBindingName = FOnGenerateBindingName::CreateLambda([WidgetName]()
+	{
+		return WidgetName;
+	});
 
-			//TODO UMG O(N) Isn't good for this, needs to be map, but map isn't serialized, need cached runtime map for fast lookups.
+	Args.OnGotoBinding = FOnGotoBinding::CreateLambda([InEditor, Objects](FName InPropertyName)
+	{
+		TSharedPtr<FWidgetBlueprintEditor> EditorPinned = InEditor.Pin();
+		UWidgetBlueprint* ThisBlueprint = EditorPinned->GetWidgetBlueprintObj();
 
-			for (const TWeakObjectPtr<UObject>& ObjectPtr : Objects)
-			{
-				UObject* Object = ObjectPtr.Get();
+		//TODO UMG O(N) Isn't good for this, needs to be map, but map isn't serialized, need cached runtime map for fast lookups.
 
-				// Ignore null outer objects
-				if ( Object == nullptr )
-				{
-					continue;
-				}
+		for (const TWeakObjectPtr<UObject>& ObjectPtr : Objects)
+		{
+			UObject* Object = ObjectPtr.Get();
+
+			// Ignore null outer objects
+			if ( Object == nullptr )
+			{
+				continue;
+			}
 
-				for ( const FDelegateEditorBinding& Binding : ThisBlueprint->Bindings )
+			for ( const FDelegateEditorBinding& Binding : ThisBlueprint->Bindings )
+			{
+				if ( Binding.ObjectName == Object->GetName() && Binding.PropertyName == InPropertyName )
 				{
-					if ( Binding.ObjectName == Object->GetName() && Binding.PropertyName == InPropertyName )
+					if ( Binding.Kind == EBindingKind::Function )
 					{
-						if ( Binding.Kind == EBindingKind::Function )
-						{
-							TArray<UEdGraph*> AllGraphs;
-							ThisBlueprint->GetAllGraphs(AllGraphs);
+						TArray<UEdGraph*> AllGraphs;
+						ThisBlueprint->GetAllGraphs(AllGraphs);
 
-							FGuid SearchForGuid = Binding.MemberGuid;
-							if ( !Binding.SourcePath.IsEmpty() )
-							{
-								SearchForGuid = Binding.SourcePath.Segments.Last().GetMemberGuid();
-							}
+						FGuid SearchForGuid = Binding.MemberGuid;
+						if ( !Binding.SourcePath.IsEmpty() )
+						{
+							SearchForGuid = Binding.SourcePath.Segments.Last().GetMemberGuid();
+						}
 
-							for ( UEdGraph* Graph : AllGraphs )
+						for ( UEdGraph* Graph : AllGraphs )
+						{
+							if ( Graph->GraphGuid == SearchForGuid )
 							{
-								if ( Graph->GraphGuid == SearchForGuid )
-								{
-									InEditor.Pin()->SetCurrentMode(FWidgetBlueprintApplicationModes::GraphMode);
-									InEditor.Pin()->OpenDocument(Graph, FDocumentTracker::OpenNewDocument);
-								}
+								EditorPinned->SetCurrentMode(FWidgetBlueprintApplicationModes::GraphMode);
+								EditorPinned->OpenDocument(Graph, FDocumentTracker::OpenNewDocument);
 							}
-
-							// Either way return
-							return true;
 						}
+
+						// Either way return
+						return true;
 					}
 				}
 			}
+		}
 
-			return false;
-		});
+		return false;
+	});
 
-		Args.OnCanGotoBinding = FOnCanGotoBinding::CreateLambda([InEditor, Objects](FName InPropertyName)
+	Args.OnCanGotoBinding = FOnCanGotoBinding::CreateLambda([InEditor, Objects](FName InPropertyName)
+	{
+		UWidgetBlueprint* ThisBlueprint = InEditor.Pin()->GetWidgetBlueprintObj();
+
+		for (const TWeakObjectPtr<UObject>& ObjectPtr : Objects)
 		{
-			UWidgetBlueprint* ThisBlueprint = InEditor.Pin()->GetWidgetBlueprintObj();
+			UObject* Object = ObjectPtr.Get();
 
-			for (const TWeakObjectPtr<UObject>& ObjectPtr : Objects)
+			// Ignore null outer objects
+			if ( Object == nullptr )
 			{
-				UObject* Object = ObjectPtr.Get();
-
-				// Ignore null outer objects
-				if ( Object == nullptr )
-				{
-					continue;
-				}
+				continue;
+			}
 
-				for ( const FDelegateEditorBinding& Binding : ThisBlueprint->Bindings )
+			for ( const FDelegateEditorBinding& Binding : ThisBlueprint->Bindings )
+			{
+				if ( Binding.ObjectName == Object->GetName() && Binding.PropertyName == InPropertyName )
 				{
-					if ( Binding.ObjectName == Object->GetName() && Binding.PropertyName == InPropertyName )
+					if ( Binding.Kind == EBindingKind::Function )
 					{
-						if ( Binding.Kind == EBindingKind::Function )
-						{
-							return true;
-						}
+						return true;
 					}
 				}
 			}
+		}
 
-			return false;
-		});
+		return false;
+	});
 
-		Args.OnCanBindProperty = FOnCanBindProperty::CreateLambda([InDelegateProperty](FProperty* InProperty)
+	Args.OnCanBindProperty = FOnCanBindProperty::CreateLambda([InDelegateProperty](FProperty* InProperty)
+	{
+		if ( FProperty* ReturnProperty = InDelegateProperty->SignatureFunction->GetReturnProperty() )
 		{
-			if ( FProperty* ReturnProperty = InDelegateProperty->SignatureFunction->GetReturnProperty() )
+			// Find the binder that can handle the delegate return type.
+			TSubclassOf<UPropertyBinding> Binder = UWidget::FindBinderClassForDestination(ReturnProperty);
+			if ( Binder != nullptr )
 			{
-				// Find the binder that can handle the delegate return type.
-				TSubclassOf<UPropertyBinding> Binder = UWidget::FindBinderClassForDestination(ReturnProperty);
-				if ( Binder != nullptr )
-				{
-					// Ensure that the binder also can handle binding from the property we care about.
-					return ( Binder->GetDefaultObject<UPropertyBinding>()->IsSupportedSource(InProperty) );
-				}
+				// Ensure that the binder also can handle binding from the property we care about.
+				return ( Binder->GetDefaultObject<UPropertyBinding>()->IsSupportedSource(InProperty) );
 			}
+		}
 
-			return false;
-		});
+		return false;
+	});
 	
-		Args.OnCanBindFunction = FOnCanBindFunction::CreateLambda([InDelegateProperty](UFunction* InFunction)
+	Args.OnCanBindFunction = FOnCanBindFunction::CreateLambda([InDelegateProperty](UFunction* InFunction)
+	{
+		auto HasFunctionBinder = [InFunction](UFunction* InBindableSignature)
 		{
-			auto HasFunctionBinder = [InFunction](UFunction* InBindableSignature)
+			if ( InFunction->NumParms == 1 && InBindableSignature->NumParms == 1 )
 			{
-				if ( InFunction->NumParms == 1 && InBindableSignature->NumParms == 1 )
+				if ( FProperty* FunctionReturn = InFunction->GetReturnProperty() )
 				{
-					if ( FProperty* FunctionReturn = InFunction->GetReturnProperty() )
+					if ( FProperty* DelegateReturn = InBindableSignature->GetReturnProperty() )
 					{
-						if ( FProperty* DelegateReturn = InBindableSignature->GetReturnProperty() )
+						// Find the binder that can handle the delegate return type.
+						TSubclassOf<UPropertyBinding> Binder = UWidget::FindBinderClassForDestination(DelegateReturn);
+						if ( Binder != nullptr )
 						{
-							// Find the binder that can handle the delegate return type.
-							TSubclassOf<UPropertyBinding> Binder = UWidget::FindBinderClassForDestination(DelegateReturn);
-							if ( Binder != nullptr )
+							// Ensure that the binder also can handle binding from the property we care about.
+							if ( Binder->GetDefaultObject<UPropertyBinding>()->IsSupportedSource(FunctionReturn) )
 							{
-								// Ensure that the binder also can handle binding from the property we care about.
-								if ( Binder->GetDefaultObject<UPropertyBinding>()->IsSupportedSource(FunctionReturn) )
-								{
-									return true;
-								}
+								return true;
 							}
 						}
 					}
 				}
+			}
 
-				return false;
-			};
+			return false;
+		};
 
-			// We ignore CPF_ReturnParm because all that matters for binding to script functions is that the number of out parameters match.
-			return ( InFunction->IsSignatureCompatibleWith(InDelegateProperty->SignatureFunction, UFunction::GetDefaultIgnoredSignatureCompatibilityFlags() | CPF_ReturnParm) ||
-				 HasFunctionBinder(InDelegateProperty->SignatureFunction) );
-		});
+		// We ignore CPF_ReturnParm because all that matters for binding to script functions is that the number of out parameters match.
+		return ( InFunction->IsSignatureCompatibleWith(InDelegateProperty->SignatureFunction, UFunction::GetDefaultIgnoredSignatureCompatibilityFlags() | CPF_ReturnParm) ||
+				HasFunctionBinder(InDelegateProperty->SignatureFunction) );
+	});
 
-		Args.OnCanBindToClass = FOnCanBindToClass::CreateLambda([](UClass* InClass)
+	Args.OnCanBindToClass = FOnCanBindToClass::CreateLambda([](UClass* InClass)
+	{
+		if (InClass == UUserWidget::StaticClass() ||
+			InClass == AActor::StaticClass() ||
+			InClass == APawn::StaticClass() ||
+			InClass == UObject::StaticClass() ||
+			InClass == UPrimitiveComponent::StaticClass() ||
+			InClass == USceneComponent::StaticClass() ||
+			InClass == UActorComponent::StaticClass() ||
+			InClass == UWidgetComponent::StaticClass() ||
+			InClass == UStaticMeshComponent::StaticClass() ||
+			InClass == UWidgetAnimation::StaticClass() )
 		{
-			if (InClass == UUserWidget::StaticClass() ||
-				InClass == AActor::StaticClass() ||
-				InClass == APawn::StaticClass() ||
-				InClass == UObject::StaticClass() ||
-				InClass == UPrimitiveComponent::StaticClass() ||
-				InClass == USceneComponent::StaticClass() ||
-				InClass == UActorComponent::StaticClass() ||
-				InClass == UWidgetComponent::StaticClass() ||
-				InClass == UStaticMeshComponent::StaticClass() ||
-				InClass == UWidgetAnimation::StaticClass() )
-			{
-				return false;
-			}
+			return false;
+		}
 		
-			return true;
-		});
+		return true;
+	});
+
+	Args.OnCanBindToSubObjectClass = FOnCanBindToSubObjectClass::CreateLambda([](UClass* InClass)
+	{
+		// Ignore any properties that are widgets, we don't want users binding widgets to other widgets.
+		return InClass->IsChildOf(UWidget::StaticClass());
+	});
+
+	Args.OnAddBinding = FOnAddBinding::CreateLambda([InEditor, Objects](FName InPropertyName, const TArray<FBindingChainElement>& InBindingChain)
+	{
+		UWidgetBlueprint* ThisBlueprint = InEditor.Pin()->GetWidgetBlueprintObj();
+		UBlueprintGeneratedClass* SkeletonClass = Cast<UBlueprintGeneratedClass>(ThisBlueprint->SkeletonGeneratedClass);
+
+		ThisBlueprint->Modify();
 
-		Args.OnCanBindToSubObjectClass = FOnCanBindToSubObjectClass::CreateLambda([](UClass* InClass)
+		TArray<FFieldVariant> FieldChain;
+		Algo::Transform(InBindingChain, FieldChain, [](const FBindingChainElement& InElement)
 		{
-			// Ignore any properties that are widgets, we don't want users binding widgets to other widgets.
-			return InClass->IsChildOf(UWidget::StaticClass());
+			return InElement.Field;
 		});
 
-		Args.OnAddBinding = FOnAddBinding::CreateLambda([InEditor, Objects](FName InPropertyName, const TArray<FBindingChainElement>& InBindingChain)
-		{
-			UWidgetBlueprint* ThisBlueprint = InEditor.Pin()->GetWidgetBlueprintObj();
-			UBlueprintGeneratedClass* SkeletonClass = Cast<UBlueprintGeneratedClass>(ThisBlueprint->SkeletonGeneratedClass);
+		UFunction* Function = FieldChain.Last().Get<UFunction>();
+		FProperty* Property = FieldChain.Last().Get<FProperty>();
+
+		check(Function != nullptr || Property != nullptr);
 
-			ThisBlueprint->Modify();
+		for (const TWeakObjectPtr<UObject>& ObjectPtr : Objects)
+		{
+			UObject* Object = ObjectPtr.Get();
 
-			TArray<FFieldVariant> FieldChain;
-			Algo::Transform(InBindingChain, FieldChain, [](const FBindingChainElement& InElement)
+			// Ignore null outer objects
+			if ( Object == nullptr )
 			{
-				return InElement.Field;
-			});
+				continue;
+			}
+
+			FDelegateEditorBinding Binding;
+			Binding.ObjectName = Object->GetName();
+			Binding.PropertyName = InPropertyName;
+			Binding.SourcePath = FEditorPropertyPath(FieldChain);
 
-			UFunction* Function = FieldChain.Last().Get<UFunction>();
-			FProperty* Property = FieldChain.Last().Get<FProperty>();
+			if ( Function != nullptr)
+			{
+				Binding.FunctionName = Function->GetFName();
 
-			check(Function != nullptr || Property != nullptr);
+				UBlueprint::GetGuidFromClassByFieldName<UFunction>(
+					Function->GetOwnerClass(),
+					Function->GetFName(),
+					Binding.MemberGuid);
 
-			for (const TWeakObjectPtr<UObject>& ObjectPtr : Objects)
+				Binding.Kind = EBindingKind::Function;
+			}
+			else if( Property != nullptr )
 			{
-				UObject* Object = ObjectPtr.Get();
+				Binding.SourceProperty = Property->GetFName();
 
-				// Ignore null outer objects
-				if ( Object == nullptr )
-				{
-					continue;
-				}
+				UBlueprint::GetGuidFromClassByFieldName<FProperty>(
+					SkeletonClass,
+					Property->GetFName(),
+					Binding.MemberGuid);
 
-				FDelegateEditorBinding Binding;
-				Binding.ObjectName = Object->GetName();
-				Binding.PropertyName = InPropertyName;
-				Binding.SourcePath = FEditorPropertyPath(FieldChain);
+				Binding.Kind = EBindingKind::Property;
+			}
 
-				if ( Function != nullptr)
-				{
-					Binding.FunctionName = Function->GetFName();
+			ThisBlueprint->Bindings.Remove(Binding);
+			ThisBlueprint->Bindings.AddUnique(Binding);
+		}
 
-					UBlueprint::GetGuidFromClassByFieldName<UFunction>(
-						Function->GetOwnerClass(),
-						Function->GetFName(),
-						Binding.MemberGuid);
+		FBlueprintEditorUtils::MarkBlueprintAsStructurallyModified(ThisBlueprint);	
+	});
 
-					Binding.Kind = EBindingKind::Function;
-				}
-				else if( Property != nullptr )
-				{
-					Binding.SourceProperty = Property->GetFName();
+	Args.OnRemoveBinding = FOnRemoveBinding::CreateLambda([InEditor, Objects, InPropertyHandle](FName InPropertyName)
+	{
+		UWidgetBlueprint* ThisBlueprint = InEditor.Pin()->GetWidgetBlueprintObj();
 
-					UBlueprint::GetGuidFromClassByFieldName<FProperty>(
-						SkeletonClass,
-						Property->GetFName(),
-						Binding.MemberGuid);
+		ThisBlueprint->Modify();
 
-					Binding.Kind = EBindingKind::Property;
-				}
+		for (const TWeakObjectPtr<UObject>& ObjectPtr : Objects)
+		{
+			UObject* Object = ObjectPtr.Get();
 
-				ThisBlueprint->Bindings.Remove(Binding);
-				ThisBlueprint->Bindings.AddUnique(Binding);
+			// Ignore null outer objects
+			if ( Object == nullptr )
+			{
+				continue;
 			}
 
-			FBlueprintEditorUtils::MarkBlueprintAsStructurallyModified(ThisBlueprint);	
-		});
-
-		Args.OnRemoveBinding = FOnRemoveBinding::CreateLambda([InEditor, Objects](FName InPropertyName)
-		{
-			UWidgetBlueprint* ThisBlueprint = InEditor.Pin()->GetWidgetBlueprintObj();
+			FDelegateEditorBinding Binding;
+			Binding.ObjectName = Object->GetName();
+			Binding.PropertyName = InPropertyName;
 
-			ThisBlueprint->Modify();
+			ThisBlueprint->Bindings.Remove(Binding);
 
-			for (const TWeakObjectPtr<UObject>& ObjectPtr : Objects)
+			if (UWidget* Widget = Cast<UWidget>(Object))
 			{
-				UObject* Object = ObjectPtr.Get();
-
-				// Ignore null outer objects
-				if ( Object == nullptr )
+				IUMGEditorModule& EditorModule = FModuleManager::LoadModuleChecked<IUMGEditorModule>("UMGEditor");
+				for (const TSharedPtr<IPropertyBindingExtension>& Extension : EditorModule.GetPropertyBindingExtensibilityManager()->GetExtensions())
 				{
-					continue;
+					Extension->ClearCurrentValue(ThisBlueprint, Widget, InPropertyHandle->GetProperty());
 				}
+			}
+		}
 
-				FDelegateEditorBinding Binding;
-				Binding.ObjectName = Object->GetName();
-				Binding.PropertyName = InPropertyName;
+		FBlueprintEditorUtils::MarkBlueprintAsStructurallyModified(ThisBlueprint);
 
-				ThisBlueprint->Bindings.Remove(Binding);
-			}
+		
+	});
 
-			FBlueprintEditorUtils::MarkBlueprintAsStructurallyModified(ThisBlueprint);
-		});
+	Args.OnCanRemoveBinding = FOnCanRemoveBinding::CreateLambda([InEditor, Objects, InPropertyHandle](FName InPropertyName)
+	{
+		UWidgetBlueprint* ThisBlueprint = InEditor.Pin()->GetWidgetBlueprintObj();
 
-		Args.OnCanRemoveBinding = FOnCanRemoveBinding::CreateLambda([InEditor, Objects](FName InPropertyName)
+		for (const TWeakObjectPtr<UObject>& ObjectPtr : Objects)
 		{
-			UWidgetBlueprint* ThisBlueprint = InEditor.Pin()->GetWidgetBlueprintObj();
+			UObject* Object = ObjectPtr.Get();
 
-			for (const TWeakObjectPtr<UObject>& ObjectPtr : Objects)
+			// Ignore null outer objects
+			if ( Object == nullptr )
 			{
-				UObject* Object = ObjectPtr.Get();
+				continue;
+			}
 
-				// Ignore null outer objects
-				if ( Object == nullptr )
+			for ( const FDelegateEditorBinding& Binding : ThisBlueprint->Bindings )
+			{
+				if ( Binding.ObjectName == Object->GetName() && Binding.PropertyName == InPropertyName )
 				{
-					continue;
+					return true;
 				}
+			}
 
-				for ( const FDelegateEditorBinding& Binding : ThisBlueprint->Bindings )
+			if (UWidget* Widget = Cast<UWidget>(Object))
+			{
+				IUMGEditorModule& EditorModule = FModuleManager::LoadModuleChecked<IUMGEditorModule>("UMGEditor");
+				for (const TSharedPtr<IPropertyBindingExtension>& Extension : EditorModule.GetPropertyBindingExtensibilityManager()->GetExtensions())
 				{
-					if ( Binding.ObjectName == Object->GetName() && Binding.PropertyName == InPropertyName )
+					TOptional<FName> Name = Extension->GetCurrentValue(ThisBlueprint, Widget, InPropertyHandle->GetProperty());
+					if (Name.IsSet())
 					{
 						return true;
 					}
 				}
 			}
+		}
 
-			return false;
-		});
+		return false;
+	});
 
-		Args.CurrentBindingText = MakeAttributeLambda([InEditor, Objects, PropertyName]()
-		{
-			UWidgetBlueprint* ThisBlueprint = InEditor.Pin()->GetWidgetBlueprintObj();
+	Args.CurrentBindingText = MakeAttributeLambda([InEditor, Objects, InPropertyHandle]()
+	{
+		UWidgetBlueprint* ThisBlueprint = InEditor.Pin()->GetWidgetBlueprintObj();
 
-			//TODO UMG O(N) Isn't good for this, needs to be map, but map isn't serialized, need cached runtime map for fast lookups.
+		//TODO UMG O(N) Isn't good for this, needs to be map, but map isn't serialized, need cached runtime map for fast lookups.
 
-			for (const TWeakObjectPtr<UObject>& ObjectPtr : Objects)
-			{
-				UObject* Object = ObjectPtr.Get();
+		for (const TWeakObjectPtr<UObject>& ObjectPtr : Objects)
+		{
+			UObject* Object = ObjectPtr.Get();
 
-				// Ignore null outer objects
-				if ( Object == nullptr )
-				{
-					continue;
-				}
+			// Ignore null outer objects
+			if ( Object == nullptr )
+			{
+				continue;
+			}
 
-				//TODO UMG handle multiple things selected
+			//TODO UMG handle multiple things selected
 
-				for ( const FDelegateEditorBinding& Binding : ThisBlueprint->Bindings )
+			for ( const FDelegateEditorBinding& Binding : ThisBlueprint->Bindings )
+			{
+				if ( Binding.ObjectName == Object->GetName() && Binding.PropertyName == InPropertyHandle->GetProperty()->GetFName())
 				{
-					if ( Binding.ObjectName == Object->GetName() && Binding.PropertyName == PropertyName )
+					if ( !Binding.SourcePath.IsEmpty() )
 					{
-						if ( !Binding.SourcePath.IsEmpty() )
+						return Binding.SourcePath.GetDisplayText();
+					}
+					else
+					{
+						if ( Binding.Kind == EBindingKind::Function )
 						{
-							return Binding.SourcePath.GetDisplayText();
+							if ( Binding.MemberGuid.IsValid() )
+							{
+								// Graph function, look up by Guid
+								FName FoundName = ThisBlueprint->GetFieldNameFromClassByGuid<UFunction>(ThisBlueprint->GeneratedClass, Binding.MemberGuid);
+								return FText::FromString(FName::NameToDisplayString(FoundName.ToString(), false));
+							}
+							else
+							{
+								// No GUID, native function, return function name.
+								return FText::FromName(Binding.FunctionName);
+							}
 						}
-						else
+						else // Property
 						{
-							if ( Binding.Kind == EBindingKind::Function )
+							if ( Binding.MemberGuid.IsValid() )
 							{
-								if ( Binding.MemberGuid.IsValid() )
-								{
-									// Graph function, look up by Guid
-									FName FoundName = ThisBlueprint->GetFieldNameFromClassByGuid<UFunction>(ThisBlueprint->GeneratedClass, Binding.MemberGuid);
-									return FText::FromString(FName::NameToDisplayString(FoundName.ToString(), false));
-								}
-								else
-								{
-									// No GUID, native function, return function name.
-									return FText::FromName(Binding.FunctionName);
-								}
+								FName FoundName = ThisBlueprint->GetFieldNameFromClassByGuid<FProperty>(ThisBlueprint->GeneratedClass, Binding.MemberGuid);
+								return FText::FromString(FName::NameToDisplayString(FoundName.ToString(), false));
 							}
-							else // Property
+							else
 							{
-								if ( Binding.MemberGuid.IsValid() )
-								{
-									FName FoundName = ThisBlueprint->GetFieldNameFromClassByGuid<FProperty>(ThisBlueprint->GeneratedClass, Binding.MemberGuid);
-									return FText::FromString(FName::NameToDisplayString(FoundName.ToString(), false));
-								}
-								else
-								{
-									// No GUID, native property, return source property.
-									return FText::FromName(Binding.SourceProperty);
-								}
+								// No GUID, native property, return source property.
+								return FText::FromName(Binding.SourceProperty);
 							}
 						}
 					}
 				}
+			}
 
-				//TODO UMG Do something about missing functions, little exclamation points if they're missing and such.
-
-				break;
+			if (UWidget* Widget = Cast<UWidget>(Object))
+			{
+				IUMGEditorModule& EditorModule = FModuleManager::LoadModuleChecked<IUMGEditorModule>("UMGEditor");
+				for (const TSharedPtr<IPropertyBindingExtension>& Extension : EditorModule.GetPropertyBindingExtensibilityManager()->GetExtensions())
+				{
+					TOptional<FName> Name = Extension->GetCurrentValue(ThisBlueprint, Widget, InPropertyHandle->GetProperty());
+					if (Name.IsSet())
+					{
+						return FText::FromName(Name.GetValue());
+					}
+				}
 			}
 
-			return LOCTEXT("Bind", "Bind");
-		});
+			//TODO UMG Do something about missing functions, little exclamation points if they're missing and such.
 
-		Args.CurrentBindingImage = MakeAttributeLambda([InEditor, Objects, PropertyName]() -> const FSlateBrush*
-		{
-			static FName PropertyIcon(TEXT("Kismet.Tabs.Variables"));
-			static FName FunctionIcon(TEXT("GraphEditor.Function_16x"));
+			break;
+		}
 
-			UWidgetBlueprint* ThisBlueprint = InEditor.Pin()->GetWidgetBlueprintObj();
+		return LOCTEXT("Bind", "Bind");
+	});
 
-			//TODO UMG O(N) Isn't good for this, needs to be map, but map isn't serialized, need cached runtime map for fast lookups.
+	Args.CurrentBindingImage = MakeAttributeLambda([InEditor, Objects, InPropertyHandle]() -> const FSlateBrush*
+	{
+		static FName PropertyIcon(TEXT("Kismet.Tabs.Variables"));
+		static FName FunctionIcon(TEXT("GraphEditor.Function_16x"));
 
-			for (const TWeakObjectPtr<UObject>& ObjectPtr : Objects)
-			{
-				UObject* Object = ObjectPtr.Get();
+		UWidgetBlueprint* ThisBlueprint = InEditor.Pin()->GetWidgetBlueprintObj();
 
-				// Ignore null outer objects
-				if ( Object == nullptr )
-				{
-					continue;
-				}
+		//TODO UMG O(N) Isn't good for this, needs to be map, but map isn't serialized, need cached runtime map for fast lookups.
+
+		for (const TWeakObjectPtr<UObject>& ObjectPtr : Objects)
+		{
+			UObject* Object = ObjectPtr.Get();
 
-				//TODO UMG handle multiple things selected
+			// Ignore null outer objects
+			if ( Object == nullptr )
+			{
+				continue;
+			}
 
-				for ( const FDelegateEditorBinding& Binding : ThisBlueprint->Bindings )
+			//TODO UMG handle multiple things selected
+
+			for ( const FDelegateEditorBinding& Binding : ThisBlueprint->Bindings )
+			{
+				if ( Binding.ObjectName == Object->GetName() && Binding.PropertyName == InPropertyHandle->GetProperty()->GetFName() )
 				{
-					if ( Binding.ObjectName == Object->GetName() && Binding.PropertyName == PropertyName )
+					if ( Binding.Kind == EBindingKind::Function )
 					{
-						if ( Binding.Kind == EBindingKind::Function )
-						{
-							return FAppStyle::Get().GetBrush(FunctionIcon);
-						}
-						else // Property
-						{
-							return FAppStyle::Get().GetBrush(PropertyIcon);
-						}
+						return FAppStyle::Get().GetBrush(FunctionIcon);
+					}
+					else // Property
+					{
+						return FAppStyle::Get().GetBrush(PropertyIcon);
 					}
 				}
 			}
+		}
 
-			return nullptr;
-		});
+		return nullptr;
+	});
 
-		Args.bGeneratePureBindings = bInGeneratePureBindings;
+	Args.bGeneratePureBindings = bInGeneratePureBindings;
 
-		IPropertyAccessEditor& PropertyAccessEditor = IModularFeatures::Get().GetModularFeature<IPropertyAccessEditor>("PropertyAccessEditor");
-		return PropertyAccessEditor.MakePropertyBindingWidget(InEditor.Pin()->GetBlueprintObj(), Args);
-	}
-	else
-	{
-		return SNullWidget::NullWidget;
-	}
+	IPropertyAccessEditor& PropertyAccessEditor = IModularFeatures::Get().GetModularFeature<IPropertyAccessEditor>("PropertyAccessEditor");
+	return PropertyAccessEditor.MakePropertyBindingWidget(InEditor.Pin()->GetBlueprintObj(), Args);
 }
 
 bool FBlueprintWidgetCustomization::HasPropertyBindings(TWeakPtr<FWidgetBlueprintEditor> InEditor, const TSharedRef<IPropertyHandle>& InPropertyHandle)
diff --git a/Engine/Source/Editor/UMGEditor/Public/IHasPropertyBindingExtensibility.h b/Engine/Source/Editor/UMGEditor/Public/IHasPropertyBindingExtensibility.h
--- a/Engine/Source/Editor/UMGEditor/Public/IHasPropertyBindingExtensibility.h
+++ b/Engine/Source/Editor/UMGEditor/Public/IHasPropertyBindingExtensibility.h
@@ -13,8 +13,9 @@ class UMGEDITOR_API IPropertyBindingExtension
 public:
 	/** Does this extension want to extend this property in the widget? */
 	virtual bool CanExtend(const UWidgetBlueprint* WidgetBlueprint, const UWidget* Widget, const FProperty* Property) const = 0;
-	virtual TOptional<FName> GetCurrentValue(const UWidgetBlueprint* WidgetBlueprint, const UWidget* Widget, const FProperty* Property) const = 0;
 	virtual TSharedPtr<FExtender> CreateMenuExtender(const UWidgetBlueprint* WidgetBlueprint, const UWidget* Widget, const FProperty* Property) = 0;
+	virtual TOptional<FName> GetCurrentValue(const UWidgetBlueprint* WidgetBlueprint, const UWidget* Widget, const FProperty* Property) const = 0;
+	virtual void ClearCurrentValue(const UWidgetBlueprint* WidgetBlueprint, const UWidget* Widget, const FProperty* Property) = 0;
 };
 
 /**
