UMG: Add BP variable customization. Add FieldNotify as a field that can be added to a UMG variable member.

diff --git a/Engine/Source/Editor/UMGEditor/Private/FieldNotification/CustomizationHelper.cpp b/Engine/Source/Editor/UMGEditor/Private/FieldNotification/CustomizationHelper.cpp
new file mode 100644
--- /dev/null
+++ b/Engine/Source/Editor/UMGEditor/Private/FieldNotification/CustomizationHelper.cpp
@@ -0,0 +1,220 @@
+// Copyright Epic Games, Inc. All Rights Reserved.
+
+
+#include "FieldNotification/CustomizationHelper.h"
+#include "Kismet2/BlueprintEditorUtils.h"
+#include "K2Node_FunctionEntry.h"
+#include "DetailLayoutBuilder.h"
+#include "BlueprintEditorModule.h"
+#include "DetailCategoryBuilder.h"
+#include "DetailWidgetRow.h"
+#include "Widgets/Text/STextBlock.h"
+#include "Widgets/Input/SCheckBox.h"
+#include "Engine/BlueprintGeneratedClass.h"
+#include "WidgetBlueprint.h"
+
+
+#define LOCTEXT_NAMESPACE "GraphFunctionDetailsCustomization"
+
+namespace UE::FieldNotification
+{
+
+const FName FCustomizationHelper::MetaData_FieldNotify = "FieldNotify";
+
+
+
+void FCustomizationHelper::CustomizeVariableDetails(IDetailLayoutBuilder& DetailLayout)
+{
+	PropertiesBeingCustomized.Reset();
+
+	TArray<TWeakObjectPtr<UObject>> ObjectsBeingCustomized;
+	DetailLayout.GetObjectsBeingCustomized(ObjectsBeingCustomized);
+	if (ObjectsBeingCustomized.Num() > 0)
+	{
+		for (TWeakObjectPtr<UObject>& Obj : ObjectsBeingCustomized)
+		{
+			UPropertyWrapper* PropertyWrapper = Cast<UPropertyWrapper>(Obj.Get());
+			FProperty* PropertyBeingCustomized = PropertyWrapper ? PropertyWrapper->GetProperty() : nullptr;
+			if (PropertyBeingCustomized)
+			{
+				PropertiesBeingCustomized.Emplace(PropertyBeingCustomized);
+
+				const FText ToolTip = LOCTEXT("FieldNotifyToolTip", "Generate a field entry for the Field Notification system.");
+				DetailLayout.EditCategory("Variable")
+					.AddCustomRow(LOCTEXT("FieldNotify", "Field Notify"))
+					.NameContent()
+					[
+						SNew(STextBlock)
+						.Font(DetailLayout.GetDetailFont())
+						.Text(LOCTEXT("FieldNotify", "Field Notify"))
+						.ToolTipText(ToolTip)
+					]
+					.ValueContent()
+					[
+						SNew(SCheckBox)
+						.IsChecked_Raw(this, &FCustomizationHelper::IsPropertyFieldNotifyChecked)
+						.OnCheckStateChanged_Raw(this, &FCustomizationHelper::HandlePropertyFieldNotifyCheckStateChanged)
+						.ToolTipText(ToolTip)
+					];
+			}
+		}
+	}
+}
+
+
+void FCustomizationHelper::CustomizeFunctionDetails(IDetailLayoutBuilder& DetailLayout)
+{
+	FunctionsBeingCustomized.Reset();
+
+	TArray<TWeakObjectPtr<UObject>> ObjectsBeingCustomized;
+	DetailLayout.GetObjectsBeingCustomized(ObjectsBeingCustomized);
+	if (ObjectsBeingCustomized.Num() > 0)
+	{
+		for (TWeakObjectPtr<UObject>& Obj : ObjectsBeingCustomized)
+		{
+			if (UK2Node_FunctionEntry* Function = FindFunctionEntry(Obj.Get()))
+			{
+				FunctionsBeingCustomized.Emplace(Function);
+
+				const FText ToolTip = LOCTEXT("FieldNotifyToolTip", "Generate a field entry for the Field Notification system.");
+				DetailLayout.EditCategory("Graph")
+					.AddCustomRow(LOCTEXT("FieldNotify", "Field Notify"))
+					.NameContent()
+					[
+						SNew(STextBlock)
+						.Font(DetailLayout.GetDetailFont())
+						.Text(LOCTEXT("FieldNotify", "Field Notify"))
+						.ToolTipText(ToolTip)
+					]
+					.ValueContent()
+					[
+						SNew(SCheckBox)
+						.IsChecked_Raw(this, &FCustomizationHelper::IsFunctionFieldNotifyChecked)
+						.OnCheckStateChanged_Raw(this, &FCustomizationHelper::HandleFunctionFieldNotifyCheckStateChanged)
+						.ToolTipText(ToolTip)
+					];
+			}
+		}
+	}
+}
+
+
+ECheckBoxState FCustomizationHelper::IsPropertyFieldNotifyChecked() const
+{
+	TOptional<ECheckBoxState> ResultState;
+	for (const TWeakFieldPtr<FProperty>& WeakProperty : PropertiesBeingCustomized)
+	{
+		if (const FProperty* Property = WeakProperty.Get())
+		{
+			ECheckBoxState State = Property->HasMetaData(MetaData_FieldNotify) ? ECheckBoxState::Checked : ECheckBoxState::Unchecked;
+			if (ResultState.IsSet() && ResultState.GetValue() != State)
+			{
+				return ECheckBoxState::Undetermined;
+			}
+			ResultState = State;
+		}
+		else
+		{
+			return ECheckBoxState::Undetermined;
+		}
+	}
+	return ResultState.IsSet() ? ResultState.GetValue() : ECheckBoxState::Undetermined;
+}
+
+
+void FCustomizationHelper::HandlePropertyFieldNotifyCheckStateChanged(ECheckBoxState CheckBoxState)
+{
+	if (UWidgetBlueprint* WidgetBlueprintPtr = Blueprint.Get())
+	{
+		for (TWeakFieldPtr<FProperty>& WeakProperty : PropertiesBeingCustomized)
+		{
+			if (FProperty* Property = WeakProperty.Get())
+			{
+				for (FBPVariableDescription& VariableDescription : WidgetBlueprintPtr->NewVariables)
+				{
+					if (VariableDescription.VarName == Property->GetFName())
+					{
+						if (CheckBoxState == ECheckBoxState::Checked)
+						{
+							Property->SetMetaData(MetaData_FieldNotify, TEXT(""));
+							VariableDescription.SetMetaData(MetaData_FieldNotify, FString());
+						}
+						else
+						{
+							Property->RemoveMetaData(MetaData_FieldNotify);
+							VariableDescription.RemoveMetaData(MetaData_FieldNotify);
+						}
+					}
+				}
+			}
+		}
+	}
+}
+
+
+ECheckBoxState FCustomizationHelper::IsFunctionFieldNotifyChecked() const
+{
+	TOptional<ECheckBoxState> ResultState;
+	for (const TWeakObjectPtr<UK2Node_FunctionEntry>& WeakFunction : FunctionsBeingCustomized)
+	{
+		if (const UK2Node_FunctionEntry* Function = WeakFunction.Get())
+		{
+			ECheckBoxState State = Function->MetaData.HasMetaData(MetaData_FieldNotify) ? ECheckBoxState::Checked : ECheckBoxState::Unchecked;
+			if (ResultState.IsSet() && ResultState.GetValue() != State)
+			{
+				return ECheckBoxState::Undetermined;
+			}
+			ResultState = State;
+		}
+		else
+		{
+			return ECheckBoxState::Undetermined;
+		}
+	}
+	return ResultState.IsSet() ? ResultState.GetValue() : ECheckBoxState::Undetermined;
+}
+
+
+void FCustomizationHelper::HandleFunctionFieldNotifyCheckStateChanged(ECheckBoxState CheckBoxState)
+{
+	if (UWidgetBlueprint* WidgetBlueprintPtr = Blueprint.Get())
+	{
+		for (TWeakObjectPtr<UK2Node_FunctionEntry>& WeakFunction : FunctionsBeingCustomized)
+		{
+			if (UK2Node_FunctionEntry* Function = WeakFunction.Get())
+			{
+				if (CheckBoxState == ECheckBoxState::Checked)
+				{
+					Function->MetaData.SetMetaData(MetaData_FieldNotify, FString());
+				}
+				else
+				{
+					Function->MetaData.RemoveMetaData(MetaData_FieldNotify);
+				}
+			}
+		}
+	}
+}
+
+
+UK2Node_FunctionEntry* FCustomizationHelper::FindFunctionEntry(UObject* Object) const
+{
+	if (UK2Node_FunctionEntry* Function = Cast<UK2Node_FunctionEntry>(Object))
+	{
+		return Function;
+	}
+	if (UEdGraph* Graph = Cast<UEdGraph>(Object))
+	{
+		for (UEdGraphNode* GraphNode : Graph->Nodes)
+		{
+			if (UK2Node_FunctionEntry* Function = Cast<UK2Node_FunctionEntry>(GraphNode))
+			{
+				return Function;
+			}
+		}
+	}
+	return nullptr;
+}
+
+
+} //namespace
\ No newline at end of file
diff --git a/Engine/Source/Editor/UMGEditor/Private/Graph/GraphFunctionDetailsCustomization.cpp b/Engine/Source/Editor/UMGEditor/Private/Graph/GraphFunctionDetailsCustomization.cpp
new file mode 100644
--- /dev/null
+++ b/Engine/Source/Editor/UMGEditor/Private/Graph/GraphFunctionDetailsCustomization.cpp
@@ -0,0 +1,48 @@
+// Copyright Epic Games, Inc. All Rights Reserved.
+
+#include "Graph/GraphFunctionDetailsCustomization.h"
+#include "FieldNotification/CustomizationHelper.h"
+#include "DetailLayoutBuilder.h"
+#include "BlueprintEditorModule.h"
+#include "Engine/BlueprintGeneratedClass.h"
+#include "WidgetBlueprint.h"
+
+
+#define LOCTEXT_NAMESPACE "GraphFunctionDetailsCustomization"
+
+FGraphFunctionDetailsCustomization::FGraphFunctionDetailsCustomization(UWidgetBlueprint* InBlueprint)
+	: Helper(MakePimpl<UE::FieldNotification::FCustomizationHelper>(InBlueprint))
+{}
+
+TSharedPtr<IDetailCustomization> FGraphFunctionDetailsCustomization::MakeInstance(TSharedPtr<IBlueprintEditor> InBlueprintEditor)
+{
+	const TArray<UObject*>* Objects = (InBlueprintEditor.IsValid() ? InBlueprintEditor->GetObjectsCurrentlyBeingEdited() : nullptr);
+	if (Objects)
+	{
+		TOptional<UWidgetBlueprint*> FinalBlueprint;
+		for (UObject* Object : *Objects)
+		{
+			UWidgetBlueprint* Blueprint = Cast<UWidgetBlueprint>(Object);
+			if (FinalBlueprint.IsSet() && FinalBlueprint.GetValue() != Blueprint)
+			{
+				return nullptr;
+			}
+			FinalBlueprint = Blueprint;
+		}
+
+		if (FinalBlueprint.IsSet())
+		{
+			return MakeShareable(new FGraphFunctionDetailsCustomization(FinalBlueprint.GetValue()));
+		}
+	}
+
+	return nullptr;
+}
+
+void FGraphFunctionDetailsCustomization::CustomizeDetails(IDetailLayoutBuilder& DetailLayout)
+{
+	Helper->CustomizeFunctionDetails(DetailLayout);
+}
+
+
+#undef LOCTEXT_NAMESPACE
\ No newline at end of file
diff --git a/Engine/Source/Editor/UMGEditor/Private/Graph/GraphFunctionDetailsCustomization.h b/Engine/Source/Editor/UMGEditor/Private/Graph/GraphFunctionDetailsCustomization.h
new file mode 100644
--- /dev/null
+++ b/Engine/Source/Editor/UMGEditor/Private/Graph/GraphFunctionDetailsCustomization.h
@@ -0,0 +1,30 @@
+// Copyright Epic Games, Inc. All Rights Reserved.
+
+#pragma once
+
+#include "CoreMinimal.h"
+#include "IDetailCustomization.h"
+#include "Styling/SlateTypes.h"
+
+class IBlueprintEditor;
+class UWidgetBlueprint;
+
+namespace UE::FieldNotification
+{
+	class FCustomizationHelper;
+}
+
+class FGraphFunctionDetailsCustomization : public IDetailCustomization
+{
+public:
+	/** Makes a new instance of this detail layout class for a specific detail view requesting it */
+	static TSharedPtr<IDetailCustomization> MakeInstance(TSharedPtr<IBlueprintEditor> InBlueprintEditor);
+
+	FGraphFunctionDetailsCustomization(UWidgetBlueprint* InBlueprint);
+
+	//~ IDetailCustomization interface
+	virtual void CustomizeDetails(IDetailLayoutBuilder& DetailLayout) override;
+
+private:
+	TPimplPtr<UE::FieldNotification::FCustomizationHelper> Helper;
+};
diff --git a/Engine/Source/Editor/UMGEditor/Private/Graph/GraphVariableDetailsCustomization.cpp b/Engine/Source/Editor/UMGEditor/Private/Graph/GraphVariableDetailsCustomization.cpp
new file mode 100644
--- /dev/null
+++ b/Engine/Source/Editor/UMGEditor/Private/Graph/GraphVariableDetailsCustomization.cpp
@@ -0,0 +1,47 @@
+// Copyright Epic Games, Inc. All Rights Reserved.
+
+#include "Graph/GraphVariableDetailsCustomization.h"
+#include "FieldNotification/CustomizationHelper.h"
+#include "DetailLayoutBuilder.h"
+#include "BlueprintEditorModule.h"
+#include "Engine/BlueprintGeneratedClass.h"
+#include "WidgetBlueprint.h"
+
+
+#define LOCTEXT_NAMESPACE "GraphVariableDetailsCustomization"
+
+FGraphVariableDetailsCustomization::FGraphVariableDetailsCustomization(UWidgetBlueprint* InBlueprint)
+	: Helper(MakePimpl<UE::FieldNotification::FCustomizationHelper>(InBlueprint))
+{}
+
+TSharedPtr<IDetailCustomization> FGraphVariableDetailsCustomization::MakeInstance(TSharedPtr<IBlueprintEditor> InBlueprintEditor)
+{
+	const TArray<UObject*>* Objects = (InBlueprintEditor.IsValid() ? InBlueprintEditor->GetObjectsCurrentlyBeingEdited() : nullptr);
+	if (Objects)
+	{
+		TOptional<UWidgetBlueprint*> FinalBlueprint;
+		for (UObject* Object : *Objects)
+		{
+			UWidgetBlueprint* Blueprint = Cast<UWidgetBlueprint>(Object);
+			if (FinalBlueprint.IsSet() && FinalBlueprint.GetValue() != Blueprint)
+			{
+				return nullptr;
+			}
+			FinalBlueprint = Blueprint;
+		}
+
+		if (FinalBlueprint.IsSet())
+		{
+			return MakeShareable(new FGraphVariableDetailsCustomization(FinalBlueprint.GetValue()));
+		}
+	}
+
+	return nullptr;
+}
+
+void FGraphVariableDetailsCustomization::CustomizeDetails(IDetailLayoutBuilder& DetailLayout)
+{
+	Helper->CustomizeVariableDetails(DetailLayout);
+}
+
+#undef LOCTEXT_NAMESPACE
\ No newline at end of file
diff --git a/Engine/Source/Editor/UMGEditor/Private/Graph/GraphVariableDetailsCustomization.h b/Engine/Source/Editor/UMGEditor/Private/Graph/GraphVariableDetailsCustomization.h
new file mode 100644
--- /dev/null
+++ b/Engine/Source/Editor/UMGEditor/Private/Graph/GraphVariableDetailsCustomization.h
@@ -0,0 +1,30 @@
+// Copyright Epic Games, Inc. All Rights Reserved.
+
+#pragma once
+
+#include "CoreMinimal.h"
+#include "IDetailCustomization.h"
+#include "Styling/SlateTypes.h"
+
+class IBlueprintEditor;
+class UWidgetBlueprint;
+
+namespace UE::FieldNotification
+{
+	class FCustomizationHelper;
+}
+
+class FGraphVariableDetailsCustomization : public IDetailCustomization
+{
+public:
+	/** Makes a new instance of this detail layout class for a specific detail view requesting it */
+	static TSharedPtr<IDetailCustomization> MakeInstance(TSharedPtr<IBlueprintEditor> InBlueprintEditor);
+
+	FGraphVariableDetailsCustomization(UWidgetBlueprint* InBlueprint);
+
+	//~ IDetailCustomization interface
+	virtual void CustomizeDetails(IDetailLayoutBuilder& DetailLayout) override;
+
+private:
+	TPimplPtr<UE::FieldNotification::FCustomizationHelper> Helper;
+};
diff --git a/Engine/Source/Editor/UMGEditor/Private/UMGEditorModule.cpp b/Engine/Source/Editor/UMGEditor/Private/UMGEditorModule.cpp
--- a/Engine/Source/Editor/UMGEditor/Private/UMGEditorModule.cpp
+++ b/Engine/Source/Editor/UMGEditor/Private/UMGEditorModule.cpp
@@ -10,6 +10,8 @@
 #include "Settings/WidgetDesignerSettings.h"
 #include "WidgetBlueprint.h"
 #include "Animation/WidgetAnimation.h"
+#include "Graph/GraphVariableDetailsCustomization.h"
+#include "Graph/GraphFunctionDetailsCustomization.h"
 
 #include "AssetToolsModule.h"
 #include "IAssetTypeActions.h"
@@ -86,6 +88,17 @@ public:
 		IKismetCompilerInterface& KismetCompilerModule = FModuleManager::LoadModuleChecked<IKismetCompilerInterface>("KismetCompiler");
 		KismetCompilerModule.GetCompilers().Add(&WidgetBlueprintCompiler);
 
+		// Add Customization for variable in Graph editor
+		if (FBlueprintEditorModule* BlueprintEditorModule = FModuleManager::GetModulePtr<FBlueprintEditorModule>("Kismet"))
+		{
+			BlueprintVariableCustomizationHandle = BlueprintEditorModule->RegisterVariableCustomization(FProperty::StaticClass(), FOnGetVariableCustomizationInstance::CreateStatic(&FGraphVariableDetailsCustomization::MakeInstance));
+			BlueprintFunctionCustomizationHandle = BlueprintEditorModule->RegisterFunctionCustomization(UK2Node_FunctionEntry::StaticClass(), FOnGetFunctionCustomizationInstance::CreateStatic(&FGraphFunctionDetailsCustomization::MakeInstance));
+		}
+		else
+		{
+			ModuleChangedHandle = FModuleManager::Get().OnModulesChanged().AddRaw(this, &FUMGEditorModule::HandleModuleChanged);
+		}
+
 		// Register asset types
 		IAssetTools& AssetTools = FModuleManager::LoadModuleChecked<FAssetToolsModule>("AssetTools").Get();
 		RegisterAssetTypeAction(AssetTools, MakeShareable(new FAssetTypeActions_SlateVectorArtData()));
@@ -117,6 +130,7 @@ public:
 	virtual void ShutdownModule() override
 	{
 		FCoreDelegates::OnPostEngineInit.RemoveAll(this);
+		FModuleManager::Get().OnModulesChanged().Remove(ModuleChangedHandle);
 
 		if (UObjectInitialized() && bThumbnailRenderersRegistered && IConsoleManager::Get().FindConsoleVariable(TEXT("UMGEditor.ThumbnailRenderer.Enable"))->GetBool())
 		{
@@ -126,8 +140,20 @@ public:
 		MenuExtensibilityManager.Reset();
 		ToolBarExtensibilityManager.Reset();
 
-		IKismetCompilerInterface& KismetCompilerModule = FModuleManager::LoadModuleChecked<IKismetCompilerInterface>("KismetCompiler");
-		KismetCompilerModule.GetCompilers().Remove(&WidgetBlueprintCompiler); 
+		if (IKismetCompilerInterface* KismetCompilerModule = FModuleManager::GetModulePtr<IKismetCompilerInterface>("KismetCompiler"))
+		{
+			KismetCompilerModule->GetCompilers().Remove(&WidgetBlueprintCompiler);
+		}
+
+		// Remove Customization for variable in Graph editor
+		if (BlueprintVariableCustomizationHandle.IsValid() || BlueprintFunctionCustomizationHandle.IsValid())
+		{
+			if (FBlueprintEditorModule* BlueprintEditorModule = FModuleManager::GetModulePtr<FBlueprintEditorModule>("Kismet"))
+			{
+				BlueprintEditorModule->UnregisterVariableCustomization(FProperty::StaticClass(), BlueprintVariableCustomizationHandle);
+				BlueprintEditorModule->UnregisterFunctionCustomization(UK2Node_FunctionEntry::StaticClass(), BlueprintFunctionCustomizationHandle);
+			}
+		}
 
 		// Unregister all the asset types that we registered
 		if ( FModuleManager::Get().IsModuleLoaded("AssetTools") )
@@ -249,6 +275,24 @@ private:
 		CreatedAssetTypeActions.Add(Action);
 	}
 
+	void HandleModuleChanged(FName ModuleName, EModuleChangeReason ChangeReason)
+	{
+		const FName KismetModule = "Kismet";
+		if (ModuleName == KismetModule && ChangeReason == EModuleChangeReason::ModuleLoaded)
+		{
+			if (!BlueprintVariableCustomizationHandle.IsValid())
+			{
+				if (FBlueprintEditorModule* BlueprintEditorModule = FModuleManager::GetModulePtr<FBlueprintEditorModule>(KismetModule))
+				{
+					BlueprintVariableCustomizationHandle = BlueprintEditorModule->RegisterVariableCustomization(FProperty::StaticClass(), FOnGetVariableCustomizationInstance::CreateStatic(&FGraphVariableDetailsCustomization::MakeInstance));
+					BlueprintFunctionCustomizationHandle = BlueprintEditorModule->RegisterFunctionCustomization(UK2Node_FunctionEntry::StaticClass(), FOnGetFunctionCustomizationInstance::CreateStatic(&FGraphFunctionDetailsCustomization::MakeInstance));
+				}
+			}
+			FModuleManager::Get().OnModulesChanged().Remove(ModuleChangedHandle);
+			ModuleChangedHandle.Reset();
+		}
+	}
+
 	void OnPostEngineInit()
 	{
 		if (GIsEditor)
@@ -306,6 +350,13 @@ private:
 	/** Support layout extensions */
 	FOnRegisterLayoutExtensions	RegisterLayoutExtensions;
 
+	/** Handle for the FModuleManager::OnModulesChanged */
+	FDelegateHandle ModuleChangedHandle;
+	/** Handle for FBlueprintEditorModule::RegisterVariableCustomization */
+	FDelegateHandle BlueprintVariableCustomizationHandle;
+	/** Handle for FBlueprintEditorModule::RegisterFunctionCustomization */
+	FDelegateHandle BlueprintFunctionCustomizationHandle;
+
 	bool bThumbnailRenderersRegistered;
 	bool bOnPostEngineInitHandled;
 };
diff --git a/Engine/Source/Editor/UMGEditor/Private/WidgetBlueprintCompiler.cpp b/Engine/Source/Editor/UMGEditor/Private/WidgetBlueprintCompiler.cpp
--- a/Engine/Source/Editor/UMGEditor/Private/WidgetBlueprintCompiler.cpp
+++ b/Engine/Source/Editor/UMGEditor/Private/WidgetBlueprintCompiler.cpp
@@ -11,6 +11,8 @@
 #include "Animation/WidgetAnimation.h"
 #include "MovieScene.h"
 
+#include "FieldNotification/CustomizationHelper.h"
+#include "FieldNotification/FieldNotificationHelpers.h"
 #include "Kismet2/Kismet2NameValidators.h"
 #include "Kismet2/KismetReinstanceUtilities.h"
 #include "Kismet2/BlueprintEditorUtils.h"
@@ -373,6 +375,7 @@ void FWidgetBlueprintCompilerContext::CleanAndSanitizeClass(UBlueprintGeneratedC
 	NewWidgetBlueprintClass->Animations.Empty();
 	NewWidgetBlueprintClass->Bindings.Empty();
 	NewWidgetBlueprintClass->Extensions.Empty();
+	NewWidgetBlueprintClass->FieldNotifyNames.Empty();
 
 	if (UWidgetBlueprintGeneratedClass* WidgetClassToClean = Cast<UWidgetBlueprintGeneratedClass>(ClassToClean))
 	{
@@ -551,6 +554,16 @@ void FWidgetBlueprintCompilerContext::CreateClassVariablesFromBlueprint()
 		{
 			InExtension->CreateClassVariablesFromBlueprint();
 		});
+
+	//Add FieldNotifyNames
+	for (TFieldIterator<const FProperty> PropertyIt(NewClass, EFieldIterationFlags::None); PropertyIt; ++PropertyIt)
+	{
+		const FProperty* Property = *PropertyIt;
+		if (Property->HasMetaData(UE::FieldNotification::FCustomizationHelper::MetaData_FieldNotify))
+		{
+			NewWidgetBlueprintClass->FieldNotifyNames.Emplace(Property->GetFName());
+		}
+	}
 }
 
 void FWidgetBlueprintCompilerContext::CopyTermDefaultsToDefaultObject(UObject* DefaultObject)
@@ -1071,6 +1084,25 @@ void FWidgetBlueprintCompilerContext::VerifyEventReplysAreNotEmpty(FKismetFuncti
 	}
 }
 
+void FWidgetBlueprintCompilerContext::PostcompileFunction(FKismetFunctionContext& Context)
+{
+	Super::PostcompileFunction(Context);
+
+	VerifyFieldNotifyFunction(Context);
+}
+
+void FWidgetBlueprintCompilerContext::VerifyFieldNotifyFunction(FKismetFunctionContext& Context)
+{
+	if (Context.Function && Context.Function->HasMetaData(UE::FieldNotification::FCustomizationHelper::MetaData_FieldNotify))
+	{
+		if (!UE::FieldNotification::Helpers::IsValidAsField(Context.Function))
+		{
+			MessageLog.Error(*LOCTEXT("FieldNotify_IsEventGraph", "Function @@ cannot be a FieldNotify. A function needs to be const, returns a single properties, has no inputs, not be an event or a net function.").ToString(), Context.EntryPoint);
+		}
+		NewWidgetBlueprintClass->FieldNotifyNames.Emplace(Context.Function->GetFName());
+	}
+}
+
 bool FWidgetBlueprintCompilerContext::ValidateGeneratedClass(UBlueprintGeneratedClass* Class)
 {
 	const bool bSuperResult = Super::ValidateGeneratedClass(Class);
diff --git a/Engine/Source/Editor/UMGEditor/Public/FieldNotification/CustomizationHelper.h b/Engine/Source/Editor/UMGEditor/Public/FieldNotification/CustomizationHelper.h
new file mode 100644
--- /dev/null
+++ b/Engine/Source/Editor/UMGEditor/Public/FieldNotification/CustomizationHelper.h
@@ -0,0 +1,49 @@
+// Copyright Epic Games, Inc. All Rights Reserved.
+
+#pragma once
+
+#include "CoreMinimal.h"
+#include "IDetailCustomization.h"
+#include "Styling/SlateTypes.h"
+#include "UObject/WeakFieldPtr.h"
+
+class IDetailLayoutBuilder;
+class IBlueprintEditor;
+class UWidgetBlueprint;
+class UK2Node_FunctionEntry;
+
+namespace UE::FieldNotification
+{
+
+class UMGEDITOR_API FCustomizationHelper
+{
+public:
+	static const FName MetaData_FieldNotify;
+
+	FCustomizationHelper(UWidgetBlueprint* InBlueprint)
+		: Blueprint(InBlueprint)
+	{}
+
+	void CustomizeVariableDetails(IDetailLayoutBuilder& DetailLayout);
+	void CustomizeFunctionDetails(IDetailLayoutBuilder& DetailLayout);
+
+private:
+	ECheckBoxState IsPropertyFieldNotifyChecked() const;
+	void HandlePropertyFieldNotifyCheckStateChanged(ECheckBoxState CheckBoxState);
+	ECheckBoxState IsFunctionFieldNotifyChecked() const;
+	void HandleFunctionFieldNotifyCheckStateChanged(ECheckBoxState CheckBoxState);
+
+	UK2Node_FunctionEntry* FindFunctionEntry(UObject* Obj) const;
+
+private:
+	/** The blueprint we are editing */
+	TWeakObjectPtr<UWidgetBlueprint> Blueprint;
+
+	/** The property we are editing */
+	TArray<TWeakFieldPtr<FProperty>> PropertiesBeingCustomized;
+
+	/** The function we are editing */
+	TArray<TWeakObjectPtr<UK2Node_FunctionEntry>> FunctionsBeingCustomized;
+};
+
+} //namespace
\ No newline at end of file
diff --git a/Engine/Source/Editor/UMGEditor/Public/WidgetBlueprintCompiler.h b/Engine/Source/Editor/UMGEditor/Public/WidgetBlueprintCompiler.h
--- a/Engine/Source/Editor/UMGEditor/Public/WidgetBlueprintCompiler.h
+++ b/Engine/Source/Editor/UMGEditor/Public/WidgetBlueprintCompiler.h
@@ -62,6 +62,7 @@ protected:
 	virtual void SpawnNewClass(const FString& NewClassName) override;
 	virtual void OnNewClassSet(UBlueprintGeneratedClass* ClassToUse) override;
 	virtual void PrecompileFunction(FKismetFunctionContext& Context, EInternalCompilerFlags InternalFlags) override;
+	virtual void PostcompileFunction(FKismetFunctionContext& Context) override;
 	virtual void CleanAndSanitizeClass(UBlueprintGeneratedClass* ClassToClean, UObject*& InOutOldCDO) override;
 	virtual void SaveSubObjectsFromCleanAndSanitizeClass(FSubobjectCollection& SubObjectsToSave, UBlueprintGeneratedClass* ClassToClean) override;
 	virtual void EnsureProperGeneratedClass(UClass*& TargetClass) override;
@@ -75,6 +76,7 @@ protected:
 	void SanitizeBindings(UBlueprintGeneratedClass* Class);
 
 	void VerifyEventReplysAreNotEmpty(FKismetFunctionContext& Context);
+	void VerifyFieldNotifyFunction(FKismetFunctionContext& Context);
 
 public:
 	void AddExtension(UWidgetBlueprintGeneratedClass* Class, UWidgetBlueprintGeneratedClassExtension* Extension);
