Null guard parent class search in UMG editor palette to fix clang builds.

diff --git a/Engine/Source/Editor/UMGEditor/Private/Palette/SPaletteViewModel.cpp b/Engine/Source/Editor/UMGEditor/Private/Palette/SPaletteViewModel.cpp
--- a/Engine/Source/Editor/UMGEditor/Private/Palette/SPaletteViewModel.cpp
+++ b/Engine/Source/Editor/UMGEditor/Private/Palette/SPaletteViewModel.cpp
@@ -313,9 +313,12 @@ void FWidgetCatalogViewModel::BuildClassWidgetList()
 			ResolveName(Outer, ParentClassName, false, false);
 			ParentClass = FindObject<UClass>(ANY_PACKAGE, *ParentClassName);
 			// UUserWidgets have their own loading section, and we don't want to process any blueprints that don't have UWidget parents
-			if (!ParentClass->IsChildOf(UWidget::StaticClass()) || ParentClass->IsChildOf(UUserWidget::StaticClass()))
+			if (ParentClass)
 			{
-				continue;
+				if (!ParentClass->IsChildOf(UWidget::StaticClass()) || ParentClass->IsChildOf(UUserWidget::StaticClass()))
+				{
+					continue;
+				}
 			}
 		}
 
