UMG: Rework how to create variable for class extension.

diff --git a/Engine/Source/Editor/UMGEditor/Private/FieldNotification/CustomizationHelper.cpp b/Engine/Source/Editor/UMGEditor/Private/FieldNotification/CustomizationHelper.cpp
--- a/Engine/Source/Editor/UMGEditor/Private/FieldNotification/CustomizationHelper.cpp
+++ b/Engine/Source/Editor/UMGEditor/Private/FieldNotification/CustomizationHelper.cpp
@@ -35,7 +35,7 @@ void FCustomizationHelper::CustomizeVariableDetails(IDetailLayoutBuilder& Detail
 		{
 			UPropertyWrapper* PropertyWrapper = Cast<UPropertyWrapper>(Obj.Get());
 			FProperty* PropertyBeingCustomized = PropertyWrapper ? PropertyWrapper->GetProperty() : nullptr;
-			if (PropertyBeingCustomized)
+			if (PropertyBeingCustomized && !PropertyBeingCustomized->HasAnyPropertyFlags(CPF_BlueprintReadOnly))
 			{
 				PropertiesBeingCustomized.Emplace(PropertyBeingCustomized);
 
diff --git a/Engine/Source/Editor/UMGEditor/Private/WidgetBlueprintCompiler.cpp b/Engine/Source/Editor/UMGEditor/Private/WidgetBlueprintCompiler.cpp
--- a/Engine/Source/Editor/UMGEditor/Private/WidgetBlueprintCompiler.cpp
+++ b/Engine/Source/Editor/UMGEditor/Private/WidgetBlueprintCompiler.cpp
@@ -32,7 +32,20 @@
 
 extern COREUOBJECT_API bool GMinimalCompileOnLoad;
 
+//////////////////////////////////////////////////////////////////////////
+// FWidgetBlueprintCompiler::FCreateVariableContext
+FWidgetBlueprintCompilerContext::FCreateVariableContext::FCreateVariableContext(FWidgetBlueprintCompilerContext& InContext)
+	: Context(InContext)
+{}
 
+FProperty* FWidgetBlueprintCompilerContext::FCreateVariableContext::CreateVariable(const FName Name, const FEdGraphPinType& Type) const
+{
+	return Context.CreateVariable(Name, Type);
+}
+
+
+//////////////////////////////////////////////////////////////////////////
+// FWidgetBlueprintCompiler
 FWidgetBlueprintCompiler::FWidgetBlueprintCompiler()
 	: ReRegister(nullptr)
 	, CompileCount(0)
@@ -504,7 +517,7 @@ void FWidgetBlueprintCompilerContext::CreateClassVariablesFromBlueprint()
 			// Only show variables if they're explicitly marked as variables.
 			if ( Widget->bIsVariable )
 			{
-				WidgetProperty->SetPropertyFlags(CPF_BlueprintVisible);
+				WidgetProperty->SetPropertyFlags(CPF_BlueprintVisible | CPF_BlueprintReadOnly | CPF_DisableEditOnInstance);
 
 				const FString& CategoryName = Widget->GetCategoryName();
 				
@@ -550,9 +563,10 @@ void FWidgetBlueprintCompilerContext::CreateClassVariablesFromBlueprint()
 		}
 	}
 
-	UWidgetBlueprintExtension::ForEachExtension(WidgetBlueprint(), [](UWidgetBlueprintExtension* InExtension)
+	FWidgetBlueprintCompilerContext* Self = this;
+	UWidgetBlueprintExtension::ForEachExtension(WidgetBlueprint(), [Self](UWidgetBlueprintExtension* InExtension)
 		{
-			InExtension->CreateClassVariablesFromBlueprint();
+			InExtension->CreateClassVariablesFromBlueprint(FCreateVariableContext(*Self));
 		});
 
 	//Add FieldNotifyNames
@@ -632,9 +646,9 @@ void FWidgetBlueprintCompilerContext::CopyTermDefaultsToDefaultObject(UObject* D
 		}
 	}
 
-	UWidgetBlueprintExtension::ForEachExtension(WidgetBlueprint(), [](UWidgetBlueprintExtension* InExtension)
+	UWidgetBlueprintExtension::ForEachExtension(WidgetBlueprint(), [DefaultObject](UWidgetBlueprintExtension* InExtension)
 		{
-			InExtension->CreateClassVariablesFromBlueprint();
+			InExtension->CopyTermDefaultsToDefaultObject(DefaultObject);
 		});
 }
 
diff --git a/Engine/Source/Editor/UMGEditor/Private/WidgetBlueprintExtension.cpp b/Engine/Source/Editor/UMGEditor/Private/WidgetBlueprintExtension.cpp
--- a/Engine/Source/Editor/UMGEditor/Private/WidgetBlueprintExtension.cpp
+++ b/Engine/Source/Editor/UMGEditor/Private/WidgetBlueprintExtension.cpp
@@ -56,5 +56,5 @@ TArray<UWidgetBlueprintExtension*> UWidgetBlueprintExtension::GetExtensions(cons
 
 UWidgetBlueprint* UWidgetBlueprintExtension::GetWidgetBlueprint() const
 {
-	return CastChecked<UWidgetBlueprint>(GetOuter());
+	return GetOuterUWidgetBlueprint();
 }
diff --git a/Engine/Source/Editor/UMGEditor/Public/WidgetBlueprintCompiler.h b/Engine/Source/Editor/UMGEditor/Public/WidgetBlueprintCompiler.h
--- a/Engine/Source/Editor/UMGEditor/Public/WidgetBlueprintCompiler.h
+++ b/Engine/Source/Editor/UMGEditor/Public/WidgetBlueprintCompiler.h
@@ -81,6 +81,14 @@ protected:
 public:
 	void AddExtension(UWidgetBlueprintGeneratedClass* Class, UWidgetBlueprintGeneratedClassExtension* Extension);
 
+	struct UMGEDITOR_API FCreateVariableContext
+	{
+		FCreateVariableContext(FWidgetBlueprintCompilerContext& InContext);
+		FProperty* CreateVariable(const FName Name, const FEdGraphPinType& Type) const;
+	private:
+		FWidgetBlueprintCompilerContext& Context;
+	};
+
 protected:
 	void FixAbandonedWidgetTree(UWidgetBlueprint* WidgetBP);
 
diff --git a/Engine/Source/Editor/UMGEditor/Public/WidgetBlueprintExtension.h b/Engine/Source/Editor/UMGEditor/Public/WidgetBlueprintExtension.h
--- a/Engine/Source/Editor/UMGEditor/Public/WidgetBlueprintExtension.h
+++ b/Engine/Source/Editor/UMGEditor/Public/WidgetBlueprintExtension.h
@@ -6,15 +6,15 @@
 #include "UObject/ObjectMacros.h"
 #include "Blueprint/BlueprintExtension.h"
 #include "WidgetBlueprint.h"
+#include "WidgetBlueprintCompiler.h"
 
 class FSubobjectCollection;
-class FWidgetBlueprintCompilerContext;
 class UWidgetBlueprintGeneratedClass;
 
 #include "WidgetBlueprintExtension.generated.h"
 
 /** Extension that allows per-system data to be held on the widget blueprint, and per-system logic to be executed during compilation */
-UCLASS()
+UCLASS(Within=WidgetBlueprint)
 class UMGEDITOR_API UWidgetBlueprintExtension : public UBlueprintExtension
 {
 	GENERATED_BODY()
@@ -74,7 +74,7 @@ protected:
 
 	virtual void HandleCreateFunctionList() {}
 	virtual void HandleCleanAndSanitizeClass(UWidgetBlueprintGeneratedClass* ClassToClean, UObject* OldCDO) {}
-	virtual void HandleCreateClassVariablesFromBlueprint() {}
+	virtual void HandleCreateClassVariablesFromBlueprint(const FWidgetBlueprintCompilerContext::FCreateVariableContext& Context) {}
 	virtual void HandleCopyTermDefaultsToDefaultObject(UObject* DefaultObject) {}
 	virtual void HandleFinishCompilingClass(UWidgetBlueprintGeneratedClass* Class) {}
 	virtual bool HandleValidateGeneratedClass(UWidgetBlueprintGeneratedClass* Class) { return true; }
@@ -104,9 +104,9 @@ private:
 		HandleCleanAndSanitizeClass(ClassToClean, OldCDO);
 	}
 
-	void CreateClassVariablesFromBlueprint()
+	void CreateClassVariablesFromBlueprint(const FWidgetBlueprintCompilerContext::FCreateVariableContext& Context)
 	{
-		HandleCreateClassVariablesFromBlueprint();
+		HandleCreateClassVariablesFromBlueprint(Context);
 	}
 
 	void CopyTermDefaultsToDefaultObject(UObject* DefaultObject)
