IWYU: Removed unnecessary includes from IHasPropertyBindingExtension in favour of forward declarations.

diff --git a/Engine/Source/Editor/UMGEditor/Public/IHasPropertyBindingExtensibility.h b/Engine/Source/Editor/UMGEditor/Public/IHasPropertyBindingExtensibility.h
--- a/Engine/Source/Editor/UMGEditor/Public/IHasPropertyBindingExtensibility.h
+++ b/Engine/Source/Editor/UMGEditor/Public/IHasPropertyBindingExtensibility.h
@@ -2,11 +2,13 @@
 
 #pragma once
 
-#include "Components/Widget.h"
 #include "Containers/Array.h"
+#include "Framework/MultiBox/MultiBoxExtender.h"
 #include "Misc/Optional.h"
 #include "Templates/SharedPointer.h"
-#include "WidgetBlueprint.h"
+
+class UWidget;
+class UWidgetBlueprint;
 
 class UMGEDITOR_API IPropertyBindingExtension
 {
