UMG: Add context to create variable from class extension.

diff --git a/Engine/Source/Editor/UMGEditor/Private/WidgetBlueprintCompiler.cpp b/Engine/Source/Editor/UMGEditor/Private/WidgetBlueprintCompiler.cpp
--- a/Engine/Source/Editor/UMGEditor/Private/WidgetBlueprintCompiler.cpp
+++ b/Engine/Source/Editor/UMGEditor/Private/WidgetBlueprintCompiler.cpp
@@ -43,6 +43,16 @@ FProperty* FWidgetBlueprintCompilerContext::FCreateVariableContext::CreateVariab
 	return Context.CreateVariable(Name, Type);
 }
 
+UWidgetBlueprint* FWidgetBlueprintCompilerContext::FCreateVariableContext::GetWidgetBlueprint() const
+{
+	return Context.WidgetBlueprint();
+}
+
+UWidgetBlueprintGeneratedClass* FWidgetBlueprintCompilerContext::FCreateVariableContext::GetSkeletonGeneratedClass() const
+{
+	return Context.NewWidgetBlueprintClass;
+}
+
 
 //////////////////////////////////////////////////////////////////////////
 // FWidgetBlueprintCompiler
diff --git a/Engine/Source/Editor/UMGEditor/Public/WidgetBlueprintCompiler.h b/Engine/Source/Editor/UMGEditor/Public/WidgetBlueprintCompiler.h
--- a/Engine/Source/Editor/UMGEditor/Public/WidgetBlueprintCompiler.h
+++ b/Engine/Source/Editor/UMGEditor/Public/WidgetBlueprintCompiler.h
@@ -83,7 +83,11 @@ public:
 
 	struct UMGEDITOR_API FCreateVariableContext
 	{
+	public:
 		FProperty* CreateVariable(const FName Name, const FEdGraphPinType& Type) const;
+		UWidgetBlueprint* GetWidgetBlueprint() const;
+		UWidgetBlueprintGeneratedClass* GetSkeletonGeneratedClass() const;
+
 	private:
 		friend FWidgetBlueprintCompilerContext;
 		FCreateVariableContext(FWidgetBlueprintCompilerContext& InContext);
