UMG: Know what type of compilation is running on the extension class.

diff --git a/Engine/Source/Editor/UMGEditor/Private/WidgetBlueprintCompiler.cpp b/Engine/Source/Editor/UMGEditor/Private/WidgetBlueprintCompiler.cpp
--- a/Engine/Source/Editor/UMGEditor/Private/WidgetBlueprintCompiler.cpp
+++ b/Engine/Source/Editor/UMGEditor/Private/WidgetBlueprintCompiler.cpp
@@ -54,6 +54,12 @@ UWidgetBlueprintGeneratedClass* FWidgetBlueprintCompilerContext::FCreateVariable
 }
 
 
+EKismetCompileType::Type FWidgetBlueprintCompilerContext::FCreateVariableContext::GetCompileType() const
+{
+	return Context.CompileOptions.CompileType;
+}
+
+
 //////////////////////////////////////////////////////////////////////////
 // FWidgetBlueprintCompiler
 FWidgetBlueprintCompiler::FWidgetBlueprintCompiler()
@@ -147,6 +153,11 @@ UEdGraphSchema_K2* FWidgetBlueprintCompilerContext::CreateSchema()
 
 void FWidgetBlueprintCompilerContext::CreateFunctionList()
 {
+	UWidgetBlueprintExtension::ForEachExtension(WidgetBlueprint(), [this](UWidgetBlueprintExtension* InExtension)
+		{
+			InExtension->CreateFunctionList();
+		});
+
 	Super::CreateFunctionList();
 
 	for ( FDelegateEditorBinding& EditorBinding : WidgetBlueprint()->Bindings )
@@ -212,11 +223,6 @@ void FWidgetBlueprintCompilerContext::CreateFunctionList()
 			}
 		}
 	}
-
-	UWidgetBlueprintExtension::ForEachExtension(WidgetBlueprint(), [this](UWidgetBlueprintExtension* InExtension)
-		{
-			InExtension->CreateFunctionList();
-		});
 }
 
 void FWidgetBlueprintCompilerContext::ValidateWidgetNames()
diff --git a/Engine/Source/Editor/UMGEditor/Public/WidgetBlueprintCompiler.h b/Engine/Source/Editor/UMGEditor/Public/WidgetBlueprintCompiler.h
--- a/Engine/Source/Editor/UMGEditor/Public/WidgetBlueprintCompiler.h
+++ b/Engine/Source/Editor/UMGEditor/Public/WidgetBlueprintCompiler.h
@@ -3,6 +3,8 @@
 #pragma once
 
 #include "CoreMinimal.h"
+#include "EngineLogs.h"
+#include "Engine/Blueprint.h"
 #include "WidgetBlueprint.h"
 #include "KismetCompiler.h"
 #include "KismetCompilerModule.h"
@@ -87,6 +89,7 @@ public:
 		FProperty* CreateVariable(const FName Name, const FEdGraphPinType& Type) const;
 		UWidgetBlueprint* GetWidgetBlueprint() const;
 		UWidgetBlueprintGeneratedClass* GetSkeletonGeneratedClass() const;
+		EKismetCompileType::Type GetCompileType() const;
 
 	private:
 		friend FWidgetBlueprintCompilerContext;
diff --git a/Engine/Source/Runtime/UMG/Private/WidgetBlueprintGeneratedClass.cpp b/Engine/Source/Runtime/UMG/Private/WidgetBlueprintGeneratedClass.cpp
--- a/Engine/Source/Runtime/UMG/Private/WidgetBlueprintGeneratedClass.cpp
+++ b/Engine/Source/Runtime/UMG/Private/WidgetBlueprintGeneratedClass.cpp
@@ -382,9 +382,10 @@ void UWidgetBlueprintGeneratedClass::PurgeClass(bool bRecompilingOnLoad)
 		Animation->Rename(nullptr, GetTransientPackage(), RenFlags);
 		FLinkerLoad::InvalidateExport(Animation);
 	}
-	Animations.Empty();
 
+	Animations.Empty();
 	Bindings.Empty();
+	FieldNotifyNames.Empty();
 }
 
 bool UWidgetBlueprintGeneratedClass::NeedsLoadForServer() const
