UMG: Add pin for FieldNotificationId on K2 nodes.

diff --git a/Engine/Source/Editor/UMGEditor/Private/Graph/UMGGraphPanelPinFactory.h b/Engine/Source/Editor/UMGEditor/Private/Graph/UMGGraphPanelPinFactory.h
new file mode 100644
--- /dev/null
+++ b/Engine/Source/Editor/UMGEditor/Private/Graph/UMGGraphPanelPinFactory.h
@@ -0,0 +1,31 @@
+// Copyright Epic Games, Inc. All Rights Reserved.
+
+#pragma once
+
+#include "CoreMinimal.h"
+#include "EdGraphUtilities.h"
+#include "EdGraphSchema_K2.h"
+#include "FieldNotification/FieldId.h"
+#include "SGraphPin.h"
+#include "UObject/Class.h"
+#include "Widgets/SFieldNotificationGraphPin.h"
+
+
+class FUMGGraphPanelPinFactory : public FGraphPanelPinFactory
+{
+	virtual TSharedPtr<class SGraphPin> CreatePin(class UEdGraphPin* InPin) const override
+	{
+		if (InPin->PinType.PinCategory == UEdGraphSchema_K2::PC_Struct)
+		{
+			if (UScriptStruct* PinStructType = Cast<UScriptStruct>(InPin->PinType.PinSubCategoryObject.Get()))
+			{
+				if (PinStructType == FFieldNotificationId::StaticStruct())
+				{
+					return SNew(UE::FieldNotification::SFieldNotificationGraphPin, InPin);
+				}
+			}
+		}
+
+		return nullptr;
+	}
+};
diff --git a/Engine/Source/Editor/UMGEditor/Private/UMGEditorModule.cpp b/Engine/Source/Editor/UMGEditor/Private/UMGEditorModule.cpp
--- a/Engine/Source/Editor/UMGEditor/Private/UMGEditorModule.cpp
+++ b/Engine/Source/Editor/UMGEditor/Private/UMGEditorModule.cpp
@@ -10,8 +10,10 @@
 #include "Settings/WidgetDesignerSettings.h"
 #include "WidgetBlueprint.h"
 #include "Animation/WidgetAnimation.h"
-#include "Graph/GraphVariableDetailsCustomization.h"
+#include "EdGraphUtilities.h"
 #include "Graph/GraphFunctionDetailsCustomization.h"
+#include "Graph/UMGGraphPanelPinFactory.h"
+#include "Graph/GraphVariableDetailsCustomization.h"
 
 #include "AssetToolsModule.h"
 #include "IAssetTypeActions.h"
@@ -123,6 +125,9 @@ public:
 		PropertyModule.RegisterCustomClassLayout(TEXT("ListViewBase"), FOnGetDetailCustomizationInstance::CreateStatic(&FListViewBaseDetails::MakeInstance));
 		PropertyModule.RegisterCustomClassLayout(TEXT("WidgetBlueprint"), FOnGetDetailCustomizationInstance::CreateStatic(&FWidgetThumbnailCustomization::MakeInstance));
 	
+		GraphPanelPinFactory = MakeShared<FUMGGraphPanelPinFactory>();
+		FEdGraphUtilities::RegisterVisualPinFactory(GraphPanelPinFactory);
+
 		CVarThumbnailRenderEnable->AsVariable()->SetOnChangedCallback(FConsoleVariableDelegate::CreateStatic(&FUMGEditorModule::ThumbnailRenderingEnabled));
 	}
 
@@ -182,9 +187,17 @@ public:
 		FPropertyEditorModule* PropertyModule = FModuleManager::GetModulePtr<FPropertyEditorModule>("PropertyEditor");
 		if (PropertyModule != nullptr)
 		{
+			PropertyModule->UnregisterCustomClassLayout(TEXT("DynamicEntryBoxBase"));
+			PropertyModule->UnregisterCustomClassLayout(TEXT("DynamicEntryBox"));
+			PropertyModule->UnregisterCustomClassLayout(TEXT("ListViewBase"));
 			PropertyModule->UnregisterCustomClassLayout(TEXT("WidgetBlueprint"));
 		}
 
+		if (UObjectInitialized() && !IsEngineExitRequested())
+		{
+			FEdGraphUtilities::UnregisterVisualPinFactory(GraphPanelPinFactory);
+		}
+
 		//// Unregister the setting
 		//ISettingsModule* SettingsModule = FModuleManager::GetModulePtr<ISettingsModule>("Settings");
 
@@ -327,6 +340,7 @@ private:
 	TSharedPtr<FExtensibilityManager> ToolBarExtensibilityManager;
 	TSharedPtr<FDesignerExtensibilityManager> DesignerExtensibilityManager;
 	TSharedPtr<FPropertyBindingExtensibilityManager> PropertyBindingExtensibilityManager;
+	TSharedPtr<FGraphPanelPinFactory> GraphPanelPinFactory;
 
 	FDelegateHandle SequenceEditorHandle;
 	FDelegateHandle MarginTrackEditorCreateTrackEditorHandle;
diff --git a/Engine/Source/Editor/UMGEditor/Private/Widgets/SFieldNotificationGraphPin.cpp b/Engine/Source/Editor/UMGEditor/Private/Widgets/SFieldNotificationGraphPin.cpp
new file mode 100644
--- /dev/null
+++ b/Engine/Source/Editor/UMGEditor/Private/Widgets/SFieldNotificationGraphPin.cpp
@@ -0,0 +1,122 @@
+// Copyright Epic Games, Inc. All Rights Reserved.
+
+#include "Widgets/SFieldNotificationGraphPin.h"
+#include "Widgets/SFieldNotificationPicker.h"
+
+#include "EdGraphSchema_K2.h"
+#include "EdGraph/EdGraphPin.h"
+#include "Engine/Blueprint.h"
+#include "K2Node.h"
+#include "Kismet2/BlueprintEditorUtils.h"
+
+
+#define LOCTEXT_NAMESPACE "SFieldNotificationGraphPin"
+
+namespace UE::FieldNotification::Private
+{
+	UClass* GetPinClass(UEdGraphPin* Pin)
+	{
+		if (Pin == nullptr)
+		{
+			return nullptr;
+		}
+		check(Pin->PinType.PinCategory == UEdGraphSchema_K2::PC_Object);
+
+		UClass* PinClass = nullptr;
+		if (UClass* DefaultClass = Cast<UClass>(Pin->DefaultObject))
+		{
+			PinClass = DefaultClass;
+		}
+		else
+		{
+			PinClass = FBlueprintEditorUtils::FindBlueprintForNodeChecked(Pin->GetOwningNode())->GeneratedClass;
+		}
+
+		UClass* BaseClass = Cast<UClass>(Pin->PinType.PinSubCategoryObject.Get());
+		if (!PinClass || !PinClass->IsChildOf(BaseClass))
+		{
+			PinClass = BaseClass;
+		}
+
+		if (Pin->LinkedTo.Num() > 0)
+		{
+			UClass* CommonInputClass = nullptr;
+			for (UEdGraphPin* LinkedPin : Pin->LinkedTo)
+			{
+				const FEdGraphPinType& LinkedPinType = LinkedPin->PinType;
+
+				UClass* LinkClass = Cast<UClass>(LinkedPinType.PinSubCategoryObject.Get());
+				if (LinkClass == nullptr && LinkedPinType.PinSubCategory == UEdGraphSchema_K2::PSC_Self)
+				{
+					if (UK2Node* K2Node = Cast<UK2Node>(LinkedPin->GetOwningNode()))
+					{
+						LinkClass = K2Node->GetBlueprint()->GeneratedClass;
+					}
+				}
+
+				if (LinkClass != nullptr)
+				{
+					if (CommonInputClass != nullptr)
+					{
+						while (!LinkClass->IsChildOf(CommonInputClass))
+						{
+							CommonInputClass = CommonInputClass->GetSuperClass();
+						}
+					}
+					else
+					{
+						CommonInputClass = LinkClass;
+					}
+				}
+			}
+
+			PinClass = CommonInputClass;
+		}
+		return PinClass;
+	}
+} // namespace
+
+
+namespace UE::FieldNotification
+{
+
+void SFieldNotificationGraphPin::Construct(const FArguments& InArgs, UEdGraphPin* InGraphPinObj)
+{
+	SGraphPin::Construct(SGraphPin::FArguments(), InGraphPinObj);
+	OnSetValue = InArgs._OnSetValue;
+}
+
+
+TSharedRef<SWidget> SFieldNotificationGraphPin::GetDefaultValueWidget()
+{
+	UEdGraphPin* SelfPin = GraphPinObj->GetOwningNode()->FindPin(UEdGraphSchema_K2::PSC_Self);
+
+	return SNew(SFieldNotificationPicker)
+		.Value(this, &SFieldNotificationGraphPin::GetValue)
+		.OnValueChanged(this, &SFieldNotificationGraphPin::SetValue)
+		.FromClass_Static(Private::GetPinClass, SelfPin)
+		.Visibility(this, &SGraphPin::GetDefaultValueVisibility);
+}
+
+
+FFieldNotificationId SFieldNotificationGraphPin::GetValue() const
+{
+	FFieldNotificationId Result;
+	if (!GraphPinObj->GetDefaultAsString().IsEmpty())
+	{
+		FFieldNotificationId::StaticStruct()->ImportText(*GraphPinObj->GetDefaultAsString(), &Result, nullptr, EPropertyPortFlags::PPF_None, GLog, FFieldNotificationId::StaticStruct()->GetName());
+	}
+	return Result;
+}
+
+
+void SFieldNotificationGraphPin::SetValue(FFieldNotificationId NewValue)
+{
+	FString ValueString;
+	FFieldNotificationId::StaticStruct()->ExportText(ValueString, &NewValue, nullptr, nullptr, EPropertyPortFlags::PPF_None, nullptr);
+	GraphPinObj->GetSchema()->TrySetDefaultValue(*GraphPinObj, ValueString);
+}
+
+} //namespace
+
+#undef LOCTEXT_NAMESPACE
diff --git a/Engine/Source/Editor/UMGEditor/Private/Widgets/SFieldNotificationGraphPin.h b/Engine/Source/Editor/UMGEditor/Private/Widgets/SFieldNotificationGraphPin.h
new file mode 100644
--- /dev/null
+++ b/Engine/Source/Editor/UMGEditor/Private/Widgets/SFieldNotificationGraphPin.h
@@ -0,0 +1,40 @@
+// Copyright Epic Games, Inc. All Rights Reserved.
+
+#pragma once
+
+#include "CoreMinimal.h"
+#include "EdGraph/EdGraphPin.h"
+#include "FieldNotification/FieldId.h"
+#include "SGraphPin.h"
+#include "Widgets/SWidget.h"
+#include "Widgets/DeclarativeSyntaxSupport.h"
+
+class UEdGraphPin;
+
+namespace UE::FieldNotification
+{
+
+class SFieldNotificationGraphPin : public SGraphPin
+{
+public:
+	DECLARE_DELEGATE_OneParam(FFieldNotificationGraphPinSetValue, FFieldNotificationId);
+
+	SLATE_BEGIN_ARGS(SFieldNotificationGraphPin){}
+		SLATE_EVENT(FFieldNotificationGraphPinSetValue, OnSetValue)
+	SLATE_END_ARGS()
+
+	void Construct(const FArguments& InArgs, UEdGraphPin* InGraphPinObj);
+
+	//~ Begin SGraphPin Interface
+	virtual TSharedRef<SWidget>	GetDefaultValueWidget() override;
+	//~ End SGraphPin Interface
+
+private:
+	FFieldNotificationId GetValue() const;
+	void SetValue(FFieldNotificationId NewValue);
+
+private:
+	FFieldNotificationGraphPinSetValue OnSetValue;
+};
+
+} //namesapce
\ No newline at end of file
diff --git a/Engine/Source/Editor/UMGEditor/Private/Widgets/SFieldNotificationPicker.cpp b/Engine/Source/Editor/UMGEditor/Private/Widgets/SFieldNotificationPicker.cpp
new file mode 100644
--- /dev/null
+++ b/Engine/Source/Editor/UMGEditor/Private/Widgets/SFieldNotificationPicker.cpp
@@ -0,0 +1,103 @@
+// Copyright Epic Games, Inc. All Rights Reserved.
+
+#include "Widgets/SFieldNotificationPicker.h"
+
+#include "Widgets/Input/SComboBox.h"
+#include "Widgets/Text/STextBlock.h"
+
+namespace UE::FieldNotification
+{
+
+void SFieldNotificationPicker::Construct(const FArguments& InArgs)
+{
+	FromClassAttribute = InArgs._FromClass;
+	ValueAttribute = InArgs._Value;
+	OnValueChangedDelegate = InArgs._OnValueChanged;
+
+	FieldNotificationIdsSource.Reset();
+	ChildSlot
+	[
+		SAssignNew(PickerBox, SComboBox<TSharedPtr<FFieldNotificationId>>)
+		.OptionsSource(&FieldNotificationIdsSource)
+		.OnSelectionChanged(this, &SFieldNotificationPicker::HandleComboBoxChanged)
+		.OnGenerateWidget(this, &SFieldNotificationPicker::HandleGenerateWidget)
+		.OnComboBoxOpening(this, &SFieldNotificationPicker::HandleComboOpening)
+		.Content()
+		[
+			SNew(STextBlock)
+			.Text(this, &SFieldNotificationPicker::HandleComboBoxValueAsText)
+		]
+	];
+}
+
+
+FFieldNotificationId SFieldNotificationPicker::GetCurrentValue() const
+{
+	return ValueAttribute.Get(FFieldNotificationId());
+}
+
+
+void SFieldNotificationPicker::HandleComboBoxChanged(TSharedPtr<FFieldNotificationId> InItem, ESelectInfo::Type InSeletionInfo)
+{
+	FFieldNotificationId NewItem = InItem ? *InItem.Get() : FFieldNotificationId();
+	if (NewItem != GetCurrentValue())
+	{
+		if (!ValueAttribute.IsBound())
+		{
+			ValueAttribute.Set(NewItem);
+		}
+		OnValueChangedDelegate.ExecuteIfBound(NewItem);
+	}
+}
+
+
+TSharedRef<SWidget> SFieldNotificationPicker::HandleGenerateWidget(TSharedPtr<FFieldNotificationId> InItem)
+{
+	return SNew(STextBlock)
+		.Text(InItem ? FText::FromName(InItem->GetFieldName()) : FText::GetEmpty());
+}
+
+
+void SFieldNotificationPicker::HandleComboOpening()
+{
+	FieldNotificationIdsSource.Reset();
+	if (UClass* Class = FromClassAttribute.Get(nullptr))
+	{
+		TScriptInterface<INotifyFieldValueChanged> ScriptObject = Class->GetDefaultObject();
+		if (ScriptObject.GetInterface() && ScriptObject.GetObject())
+		{
+			SFieldNotificationPicker* Self = this;
+			const UE::FieldNotification::IClassDescriptor& Descriptor = ScriptObject->GetFieldNotificationDescriptor();
+			Descriptor.ForEachField(Class, [Self](const FFieldId& Id) ->bool
+				{
+					Self->FieldNotificationIdsSource.Add(MakeShared<FFieldNotificationId>(Id.GetName()));
+					return true;
+				});
+		}
+
+		if (FieldNotificationIdsSource.Num())
+		{
+			FieldNotificationIdsSource.Sort([](const TSharedPtr<FFieldNotificationId>& A, const TSharedPtr<FFieldNotificationId>& B)
+				{
+					return A->GetFieldName().LexicalLess(B->GetFieldName());
+				});
+			FFieldNotificationId CurrentValue = GetCurrentValue();
+			const TSharedPtr<FFieldNotificationId>* FoundValue = FieldNotificationIdsSource.FindByPredicate([CurrentValue](const TSharedPtr<FFieldNotificationId>& Other)
+				{
+					return Other->GetFieldName() == CurrentValue.GetFieldName();
+				});
+			if (FoundValue)
+			{
+				PickerBox->SetSelectedItem(*FoundValue);
+			}
+		}
+	}
+}
+
+
+FText SFieldNotificationPicker::HandleComboBoxValueAsText() const
+{
+	return FText::FromName(GetCurrentValue().GetFieldName());
+}
+
+} //namespace
\ No newline at end of file
diff --git a/Engine/Source/Editor/UMGEditor/Private/Widgets/SFieldNotificationPicker.h b/Engine/Source/Editor/UMGEditor/Private/Widgets/SFieldNotificationPicker.h
new file mode 100644
--- /dev/null
+++ b/Engine/Source/Editor/UMGEditor/Private/Widgets/SFieldNotificationPicker.h
@@ -0,0 +1,51 @@
+// Copyright Epic Games, Inc. All Rights Reserved.
+
+#pragma once
+
+#include "Widgets/DeclarativeSyntaxSupport.h"
+#include "Widgets/SCompoundWidget.h"
+
+#include "FieldNotification/FieldId.h"
+#include "FieldNotification/IFieldValueChanged.h"
+
+template<typename OptionType>
+class SComboBox;
+
+namespace UE::FieldNotification
+{
+
+/**
+ * A widget which allows the user to enter a FieldNotificationId or discover it from a drop menu.
+ */
+class SFieldNotificationPicker : public SCompoundWidget
+{
+public:
+	DECLARE_DELEGATE_OneParam(FOnValueChanged, FFieldNotificationId);
+
+	SLATE_BEGIN_ARGS(SFieldNotificationPicker)
+	{}
+		SLATE_ATTRIBUTE(FFieldNotificationId, Value)
+		SLATE_EVENT(FOnValueChanged, OnValueChanged)
+		SLATE_ATTRIBUTE(UClass*, FromClass)
+
+	SLATE_END_ARGS()
+
+	void Construct(const FArguments& InArgs);
+
+	FFieldNotificationId GetCurrentValue() const;
+
+private:
+	void HandleComboBoxChanged(TSharedPtr<FFieldNotificationId> InItem, ESelectInfo::Type InSeletionInfo);
+	TSharedRef<SWidget> HandleGenerateWidget(TSharedPtr<FFieldNotificationId> InItem);
+	void HandleComboOpening();
+	FText HandleComboBoxValueAsText() const;
+
+private:
+	TSharedPtr<SComboBox<TSharedPtr<FFieldNotificationId>>> PickerBox;
+	TArray<TSharedPtr<FFieldNotificationId>> FieldNotificationIdsSource;
+	FOnValueChanged OnValueChangedDelegate;
+	TAttribute<FFieldNotificationId> ValueAttribute;
+	TAttribute<UClass*> FromClassAttribute;
+};
+
+} // namespace
\ No newline at end of file
