FieldNotify: Only show the FieldNotify option in BP when the class is a WidgetBP.

diff --git a/Engine/Source/Editor/UMGEditor/Private/Graph/GraphFunctionDetailsCustomization.cpp b/Engine/Source/Editor/UMGEditor/Private/Graph/GraphFunctionDetailsCustomization.cpp
--- a/Engine/Source/Editor/UMGEditor/Private/Graph/GraphFunctionDetailsCustomization.cpp
+++ b/Engine/Source/Editor/UMGEditor/Private/Graph/GraphFunctionDetailsCustomization.cpp
@@ -23,6 +23,10 @@ TSharedPtr<IDetailCustomization> FGraphFunctionDetailsCustomization::MakeInstanc
 		for (UObject* Object : *Objects)
 		{
 			UWidgetBlueprint* Blueprint = Cast<UWidgetBlueprint>(Object);
+			if (Blueprint == nullptr)
+			{
+				return nullptr;
+			}
 			if (FinalBlueprint.IsSet() && FinalBlueprint.GetValue() != Blueprint)
 			{
 				return nullptr;
diff --git a/Engine/Source/Editor/UMGEditor/Private/Graph/GraphVariableDetailsCustomization.cpp b/Engine/Source/Editor/UMGEditor/Private/Graph/GraphVariableDetailsCustomization.cpp
--- a/Engine/Source/Editor/UMGEditor/Private/Graph/GraphVariableDetailsCustomization.cpp
+++ b/Engine/Source/Editor/UMGEditor/Private/Graph/GraphVariableDetailsCustomization.cpp
@@ -23,6 +23,10 @@ TSharedPtr<IDetailCustomization> FGraphVariableDetailsCustomization::MakeInstanc
 		for (UObject* Object : *Objects)
 		{
 			UWidgetBlueprint* Blueprint = Cast<UWidgetBlueprint>(Object);
+			if (Blueprint == nullptr)
+			{
+				return nullptr;
+			}
 			if (FinalBlueprint.IsSet() && FinalBlueprint.GetValue() != Blueprint)
 			{
 				return nullptr;
