Return correct toolkit context name for UMG editor

diff --git a/Engine/Source/Editor/UMGEditor/Private/WidgetBlueprintEditor.cpp b/Engine/Source/Editor/UMGEditor/Private/WidgetBlueprintEditor.cpp
--- a/Engine/Source/Editor/UMGEditor/Private/WidgetBlueprintEditor.cpp
+++ b/Engine/Source/Editor/UMGEditor/Private/WidgetBlueprintEditor.cpp
@@ -450,6 +450,11 @@ bool FWidgetBlueprintEditor::IsPreviewWidgetInitialized()
 	return GetPreview() != nullptr;
 }
 
+FName FWidgetBlueprintEditor::GetToolkitContextFName() const
+{
+	return GetToolkitFName();
+}
+
 FName FWidgetBlueprintEditor::GetToolkitFName() const
 {
 	return FName("WidgetBlueprintEditor");
diff --git a/Engine/Source/Editor/UMGEditor/Public/WidgetBlueprintEditor.h b/Engine/Source/Editor/UMGEditor/Public/WidgetBlueprintEditor.h
--- a/Engine/Source/Editor/UMGEditor/Public/WidgetBlueprintEditor.h
+++ b/Engine/Source/Editor/UMGEditor/Public/WidgetBlueprintEditor.h
@@ -78,6 +78,7 @@ public:
 	//~ End FGCObjectInterface interface
 
 	//~ Begin IToolkit interface
+	virtual FName GetToolkitContextFName() const override;
 	virtual FName GetToolkitFName() const override;
 	virtual FText GetBaseToolkitName() const override;
 	virtual FString GetWorldCentricTabPrefix() const override;
