Add option to hide graph editor and widget animation editor elements in UMG editor

diff --git a/Engine/Source/Editor/UMGEditor/Private/Animation/UMGDetailKeyframeHandler.cpp b/Engine/Source/Editor/UMGEditor/Private/Animation/UMGDetailKeyframeHandler.cpp
--- a/Engine/Source/Editor/UMGEditor/Private/Animation/UMGDetailKeyframeHandler.cpp
+++ b/Engine/Source/Editor/UMGEditor/Private/Animation/UMGDetailKeyframeHandler.cpp
@@ -2,6 +2,7 @@
 
 #include "Animation/UMGDetailKeyframeHandler.h"
 #include "Animation/WidgetAnimation.h"
+#include "UMGEditorProjectSettings.h"
 
 #include "PropertyHandle.h"
 #include "MovieScene.h"
@@ -12,11 +13,13 @@ FUMGDetailKeyframeHandler::FUMGDetailKeyframeHandler(TSharedPtr<FWidgetBlueprint
 
 bool FUMGDetailKeyframeHandler::IsPropertyKeyable(const UClass* InObjectClass, const IPropertyHandle& InPropertyHandle) const
 {
-	return BlueprintEditor.Pin()->GetSequencer()->CanKeyProperty(FCanKeyPropertyParams(InObjectClass, InPropertyHandle));
+	return !GetDefault<UUMGEditorProjectSettings>()->bHideWidgetAnimationEditor 
+			&& BlueprintEditor.Pin()->GetSequencer()->CanKeyProperty(FCanKeyPropertyParams(InObjectClass, InPropertyHandle));
 }
 
 bool FUMGDetailKeyframeHandler::IsPropertyKeyingEnabled() const
 {
+
 	UMovieSceneSequence* Sequence = BlueprintEditor.Pin()->GetSequencer()->GetRootMovieSceneSequence();
 	return Sequence != nullptr && Sequence != UWidgetAnimation::GetNullAnimation();
 }
diff --git a/Engine/Source/Editor/UMGEditor/Private/BlueprintModes/WidgetDesignerApplicationMode.cpp b/Engine/Source/Editor/UMGEditor/Private/BlueprintModes/WidgetDesignerApplicationMode.cpp
--- a/Engine/Source/Editor/UMGEditor/Private/BlueprintModes/WidgetDesignerApplicationMode.cpp
+++ b/Engine/Source/Editor/UMGEditor/Private/BlueprintModes/WidgetDesignerApplicationMode.cpp
@@ -9,6 +9,7 @@
 #include "HAL/PlatformApplicationMisc.h"
 #include "WidgetBlueprintEditorToolbar.h"
 #include "UMGEditorModule.h"
+#include "UMGEditorProjectSettings.h"
 #include "StatusBarSubsystem.h"
 #include "ToolMenus.h"
 #include "WidgetDrawerConfig.h"
@@ -195,14 +196,17 @@ void FWidgetDesignerApplicationMode::PostActivateMode()
 
 	TSharedPtr<FWidgetBlueprintEditor> BP = GetBlueprintEditor();
 
-	FWidgetDrawerConfig WidgetAnimSequencerDrawer(FAnimationTabSummoner::WidgetAnimSequencerDrawerID);
-	WidgetAnimSequencerDrawer.GetDrawerContentDelegate.BindSP(BP.Get(), &FWidgetBlueprintEditor::OnGetWidgetAnimSequencer);
-	WidgetAnimSequencerDrawer.OnDrawerOpenedDelegate.BindSP(BP.Get(), &FWidgetBlueprintEditor::OnWidgetAnimSequencerOpened);
-	WidgetAnimSequencerDrawer.OnDrawerDismissedDelegate.BindSP(BP.Get(), &FWidgetBlueprintEditor::OnWidgetAnimSequencerDismissed);
-	WidgetAnimSequencerDrawer.ButtonText = LOCTEXT("StatusBar_WidgetAnimSequencer", "Animations");
-	WidgetAnimSequencerDrawer.ToolTipText = LOCTEXT("StatusBar_WidgetAnimSequencerToolTip", "Opens animation sequencer (Ctrl+Shift+Space Bar).");
-	WidgetAnimSequencerDrawer.Icon = FAppStyle::Get().GetBrush("UMGEditor.AnimTabIcon");
-	BP->RegisterDrawer(MoveTemp(WidgetAnimSequencerDrawer), 1);
+	if (!GetDefault<UUMGEditorProjectSettings>()->bHideWidgetAnimationEditor)
+	{
+		FWidgetDrawerConfig WidgetAnimSequencerDrawer(FAnimationTabSummoner::WidgetAnimSequencerDrawerID);
+		WidgetAnimSequencerDrawer.GetDrawerContentDelegate.BindSP(BP.Get(), &FWidgetBlueprintEditor::OnGetWidgetAnimSequencer);
+		WidgetAnimSequencerDrawer.OnDrawerOpenedDelegate.BindSP(BP.Get(), &FWidgetBlueprintEditor::OnWidgetAnimSequencerOpened);
+		WidgetAnimSequencerDrawer.OnDrawerDismissedDelegate.BindSP(BP.Get(), &FWidgetBlueprintEditor::OnWidgetAnimSequencerDismissed);
+		WidgetAnimSequencerDrawer.ButtonText = LOCTEXT("StatusBar_WidgetAnimSequencer", "Animations");
+		WidgetAnimSequencerDrawer.ToolTipText = LOCTEXT("StatusBar_WidgetAnimSequencerToolTip", "Opens animation sequencer (Ctrl+Shift+Space Bar).");
+		WidgetAnimSequencerDrawer.Icon = FAppStyle::Get().GetBrush("UMGEditor.AnimTabIcon");
+		BP->RegisterDrawer(MoveTemp(WidgetAnimSequencerDrawer), 1);
+	}
 
 	// Toggle any active tool palette modes
 	for (TSharedPtr<FUICommandInfo>& Command : BP->ToolPaletteCommands)
diff --git a/Engine/Source/Editor/UMGEditor/Private/Details/SWidgetDetailsView.cpp b/Engine/Source/Editor/UMGEditor/Private/Details/SWidgetDetailsView.cpp
--- a/Engine/Source/Editor/UMGEditor/Private/Details/SWidgetDetailsView.cpp
+++ b/Engine/Source/Editor/UMGEditor/Private/Details/SWidgetDetailsView.cpp
@@ -30,6 +30,7 @@
 #include "WidgetBlueprintEditorUtils.h"
 #include "ScopedTransaction.h"
 #include "Styling/SlateIconFinder.h"
+#include "UMGEditorProjectSettings.h"
 
 #define LOCTEXT_NAMESPACE "UMG"
 
@@ -59,6 +60,17 @@ void SWidgetDetailsView::Construct(const FArguments& InArgs, TSharedPtr<FWidgetB
 	// Notify us of object selection changes so we can update the package re-mapping
 	PropertyView->SetOnObjectArrayChanged(FOnObjectArrayChanged::CreateSP(this, &SWidgetDetailsView::OnPropertyViewObjectArrayChanged));
 
+	TSharedRef<SWidget> IsVariableCheckbox =
+		SNew(SCheckBox)
+		.IsChecked(this, &SWidgetDetailsView::GetIsVariable)
+		.OnCheckStateChanged(this, &SWidgetDetailsView::HandleIsVariableChanged)
+		.Padding(FMargin(3, 1, 18, 1))
+		.Visibility(GetDefault<UUMGEditorProjectSettings>()->bGraphEditorHidden ? EVisibility::Hidden : EVisibility::Visible)
+		[
+			SNew(STextBlock)
+			.Text(LOCTEXT("IsVariable", "Is Variable"))
+		];
+
 	ChildSlot
 	[
 		SNew(SVerticalBox)
@@ -145,20 +157,12 @@ void SWidgetDetailsView::Construct(const FArguments& InArgs, TSharedPtr<FWidgetB
 					+ SHorizontalBox::Slot()
 					.AutoWidth()
 					[
-						SNew(SCheckBox)
-						.IsChecked(this, &SWidgetDetailsView::GetIsVariable)
-						.OnCheckStateChanged(this, &SWidgetDetailsView::HandleIsVariableChanged)
-						.Padding(FMargin(3,1,3,1))
-						[
-							SNew(STextBlock)
-							.Text(LOCTEXT("IsVariable", "Is Variable"))
-						]
+						GetDefault<UUMGEditorProjectSettings>()->bGraphEditorHidden ? SNullWidget::NullWidget : IsVariableCheckbox
 					]
 
 					+ SHorizontalBox::Slot()
 					.AutoWidth()
 					.VAlign(VAlign_Center)
-					.Padding(15,0,0,0)
 					[
 						SAssignNew(ClassLinkArea, SBox)
 					]
diff --git a/Engine/Source/Editor/UMGEditor/Private/Settings/UMGEditorProjectSettings.cpp b/Engine/Source/Editor/UMGEditor/Private/Settings/UMGEditorProjectSettings.cpp
--- a/Engine/Source/Editor/UMGEditor/Private/Settings/UMGEditorProjectSettings.cpp
+++ b/Engine/Source/Editor/UMGEditor/Private/Settings/UMGEditorProjectSettings.cpp
@@ -15,6 +15,8 @@ UUMGEditorProjectSettings::UUMGEditorProjectSettings()
 	CurrentVersion = 1;
 	bShowWidgetsFromEngineContent = false;
 	bShowWidgetsFromDeveloperContent = true;
+	bGraphEditorHidden = false;
+	bHideWidgetAnimationEditor = false;
 
 	bUseWidgetTemplateSelector = false;
 	DefaultRootWidget = nullptr;
diff --git a/Engine/Source/Editor/UMGEditor/Public/UMGEditorProjectSettings.h b/Engine/Source/Editor/UMGEditor/Public/UMGEditorProjectSettings.h
--- a/Engine/Source/Editor/UMGEditor/Public/UMGEditorProjectSettings.h
+++ b/Engine/Source/Editor/UMGEditor/Public/UMGEditorProjectSettings.h
@@ -170,6 +170,14 @@ public:
 	UPROPERTY(EditAnywhere, config, Category = Designer)
 	TSubclassOf<UPanelWidget> DefaultRootWidget;
 
+	/** Set true to hide UMG-graph related elements when the graph editor is hidden  */
+	UPROPERTY(EditAnywhere, Category = Designer)
+	bool bGraphEditorHidden;
+
+	/** Set true to hide widget animation related elements in the UMG editor  */
+	UPROPERTY(EditAnywhere, Category = Designer)
+	bool bHideWidgetAnimationEditor;
+
 	/**
 	 * The default parent class for all newly constructed widget blueprints.
 	 * The WidgetParentClass must have an empty widget hierarchy.
