Add support for JSON-based editor configs for categories and widgets in UMG editor

diff --git a/Engine/Source/Editor/UMGEditor/Private/Palette/SPaletteViewModel.cpp b/Engine/Source/Editor/UMGEditor/Private/Palette/SPaletteViewModel.cpp
--- a/Engine/Source/Editor/UMGEditor/Private/Palette/SPaletteViewModel.cpp
+++ b/Engine/Source/Editor/UMGEditor/Private/Palette/SPaletteViewModel.cpp
@@ -11,6 +11,8 @@
 	#include "EditorStyleSet.h"
 #endif // WITH_EDITOR
 
+#include "ClassViewerModule.h"
+#include "ClassViewerFilter.h"
 #include "DragDrop/WidgetTemplateDragDropOp.h"
 
 #include "Templates/WidgetTemplateClass.h"
@@ -18,6 +20,7 @@
 
 #include "AssetRegistryModule.h"
 #include "WidgetBlueprintEditorUtils.h"
+#include "Misc/NamePermissionList.h"
 
 #include "Settings/ContentBrowserSettings.h"
 #include "Settings/WidgetDesignerSettings.h"
@@ -214,7 +217,39 @@ void FWidgetCatalogViewModel::BuildClassWidgetList()
 	auto ActiveWidgetBlueprintClass = GetBlueprint()->GeneratedClass;
 	FName ActiveWidgetBlueprintClassName = ActiveWidgetBlueprintClass->GetFName();
 
-	TArray<FSoftClassPath> WidgetClassesToHide = GetDefault<UUMGEditorProjectSettings>()->WidgetClassesToHide;
+	UUMGEditorProjectSettings* UMGEditorProjectSettings = GetMutableDefault<UUMGEditorProjectSettings>();
+	bool bUseEditorConfigPaletteFiltering = UMGEditorProjectSettings->bUseEditorConfigPaletteFiltering;
+	const FNamePermissionList& AllowedPaletteCategories = UMGEditorProjectSettings->GetAllowedPaletteCategories();
+	const FPathPermissionList& AllowedPaletteWidgets = UMGEditorProjectSettings->GetAllowedPaletteWidgets();
+	const TArray<FSoftClassPath>& WidgetClassesToHide = UMGEditorProjectSettings->WidgetClassesToHide;
+
+	// Since allowing all widgets in a category can be excessive when using permission lists, ensure widgets pass class viewing filters
+	FClassViewerModule& ClassViewerModule = FModuleManager::LoadModuleChecked<FClassViewerModule>("ClassViewer");
+	const TSharedPtr<IClassViewerFilter>& GlobalClassFilter = ClassViewerModule.GetGlobalClassViewerFilter();
+	TSharedRef<FClassViewerFilterFuncs> ClassFilterFuncs = ClassViewerModule.CreateFilterFuncs();
+	FClassViewerInitializationOptions ClassViewerOptions = {};
+
+	auto WidgetPassesConfigFiltering = [&](const FString& InWidgetPath, const FString& InCategoryName, TWeakObjectPtr<UClass> InWidgetClass)
+	{
+		if (AllowedPaletteWidgets.PassesFilter(*InWidgetPath))
+		{
+			return true;
+		}
+
+		if (AllowedPaletteCategories.PassesFilter(*InCategoryName))
+		{
+			if (GlobalClassFilter.IsValid() && InWidgetClass.IsValid())
+			{
+				return GlobalClassFilter->IsClassAllowed(ClassViewerOptions, InWidgetClass.Get(), ClassFilterFuncs);
+			}
+			else
+			{
+				return false;
+			}
+		}
+
+		return false;
+	};
 
 	// Locate all UWidget classes from code and loaded widget BPs
 	for (TObjectIterator<UClass> ClassIt; ClassIt; ++ClassIt)
@@ -252,19 +287,29 @@ void FWidgetCatalogViewModel::BuildClassWidgetList()
 			}
 		}
 
-		// Excludes this widget if it is on the hide list
-		bool bIsOnList = false;
-		for (FSoftClassPath Widget : WidgetClassesToHide)
+		if (bUseEditorConfigPaletteFiltering)
 		{
-			if (WidgetAssetData.ObjectPath.ToString().Find(Widget.ToString()) == 0)
+			if (!WidgetPassesConfigFiltering(WidgetAssetData.ObjectPath.ToString(), WidgetClass->GetDefaultObject<UWidget>()->GetPaletteCategory().ToString(), WidgetClass))
 			{
-				bIsOnList = true;
-				break;
+				continue;
 			}
 		}
-		if (bIsOnList)
+		else
 		{
-			continue;
+			// Excludes this widget if it is on the hide list
+			bool bIsOnList = false;
+			for (FSoftClassPath Widget : WidgetClassesToHide)
+			{
+				if (WidgetAssetData.ObjectPath.ToString().Find(Widget.ToString()) == 0)
+				{
+					bIsOnList = true;
+					break;
+				}
+			}
+			if (bIsOnList)
+			{
+				continue;
+			}
 		}
 
 		const bool bIsSameClass = WidgetClass->GetFName() == ActiveWidgetBlueprintClassName;
@@ -298,37 +343,41 @@ void FWidgetCatalogViewModel::BuildClassWidgetList()
 	TArray<FAssetData> AllBPsAssetData;
 	AssetRegistryModule.Get().GetAssetsByClass(UBlueprint::StaticClass()->GetFName(), AllBPsAssetData, true);
 
-	for (FAssetData& BPAssetData : AllBPsAssetData)
+	// Ignore unloaded assets when using palette filtering, as they lack data needed to correctly pass filters
+	if (!bUseEditorConfigPaletteFiltering)
 	{
-		// Blueprints get the class type actions for their parent native class - this avoids us having to load the blueprint
-		UClass* ParentClass = nullptr;
-		FString ParentClassName;
-		if (!BPAssetData.GetTagValue(FBlueprintTags::NativeParentClassPath, ParentClassName))
+		for (FAssetData& BPAssetData : AllBPsAssetData)
 		{
-			BPAssetData.GetTagValue(FBlueprintTags::ParentClassPath, ParentClassName);
-		}
-		if (!ParentClassName.IsEmpty())
-		{
-			UObject* Outer = nullptr;
-			ResolveName(Outer, ParentClassName, false, false);
-			ParentClass = FindObject<UClass>(ANY_PACKAGE, *ParentClassName);
-			// UUserWidgets have their own loading section, and we don't want to process any blueprints that don't have UWidget parents
-			if (ParentClass)
+			// Blueprints get the class type actions for their parent native class - this avoids us having to load the blueprint
+			UClass* ParentClass = nullptr;
+			FString ParentClassName;
+			if (!BPAssetData.GetTagValue(FBlueprintTags::NativeParentClassPath, ParentClassName))
+			{
+				BPAssetData.GetTagValue(FBlueprintTags::ParentClassPath, ParentClassName);
+			}
+			if (!ParentClassName.IsEmpty())
 			{
-				if (!ParentClass->IsChildOf(UWidget::StaticClass()) || ParentClass->IsChildOf(UUserWidget::StaticClass()))
+				UObject* Outer = nullptr;
+				ResolveName(Outer, ParentClassName, false, false);
+				ParentClass = FindObject<UClass>(ANY_PACKAGE, *ParentClassName);
+				// UUserWidgets have their own loading section, and we don't want to process any blueprints that don't have UWidget parents
+				if (ParentClass)
 				{
-					continue;
+					if (!ParentClass->IsChildOf(UWidget::StaticClass()) || ParentClass->IsChildOf(UUserWidget::StaticClass()))
+					{
+						continue;
+					}
 				}
 			}
-		}
 
-		if (!FilterAssetData(BPAssetData))
-		{
-			// If this object isn't currently loaded, add it to the palette view
-			if (BPAssetData.ToSoftObjectPath().ResolveObject() == nullptr)
+			if (!FilterAssetData(BPAssetData))
 			{
-				auto Template = MakeShareable(new FWidgetTemplateClass(BPAssetData, nullptr));
-				AddWidgetTemplate(Template);
+				// If this object isn't currently loaded, add it to the palette view
+				if (BPAssetData.ToSoftObjectPath().ResolveObject() == nullptr)
+				{
+					auto Template = MakeShareable(new FWidgetTemplateClass(BPAssetData, nullptr));
+					AddWidgetTemplate(Template);
+				}
 			}
 		}
 	}
@@ -347,21 +396,6 @@ void FWidgetCatalogViewModel::BuildClassWidgetList()
 
 		if (!FilterAssetData(WidgetBPAssetData))
 		{
-			// Excludes this widget if it is on the hide list
-			bool bIsOnList = false;
-			for (FSoftClassPath Widget : WidgetClassesToHide)
-			{
-				if (Widget.ToString().Find(WidgetBPAssetData.ObjectPath.ToString()) == 0)
-				{
-					bIsOnList = true;
-					break;
-				}
-			}
-			if (bIsOnList)
-			{
-				continue;
-			}
-
 			// If the blueprint generated class was found earlier, pass it to the template
 			TSubclassOf<UUserWidget> WidgetBPClass = nullptr;
 			auto LoadedWidgetBPClass = LoadedWidgetBlueprintClassesByName.Find(WidgetBPAssetData.AssetName);
@@ -370,6 +404,32 @@ void FWidgetCatalogViewModel::BuildClassWidgetList()
 				WidgetBPClass = *LoadedWidgetBPClass;
 			}
 
+			if (bUseEditorConfigPaletteFiltering)
+			{
+				FText Category = FWidgetTemplateBlueprintClass(WidgetBPAssetData, WidgetBPClass).GetCategory();
+				if (!WidgetPassesConfigFiltering(WidgetBPAssetData.ObjectPath.ToString(), Category.ToString(), WidgetBPClass))
+				{
+					continue;
+				}
+			}
+			else
+			{
+				// Excludes this widget if it is on the hide list
+				bool bIsOnList = false;
+				for (FSoftClassPath Widget : WidgetClassesToHide)
+				{
+					if (Widget.ToString().Find(WidgetBPAssetData.ObjectPath.ToString()) == 0)
+					{
+						bIsOnList = true;
+						break;
+					}
+				}
+				if (bIsOnList)
+				{
+					continue;
+				}
+			}
+
 			uint32 BPFlags = WidgetBPAssetData.GetTagValueRef<uint32>(FBlueprintTags::ClassFlags);
 			if ((BPFlags & (CLASS_Abstract | CLASS_Deprecated | CLASS_HideDropDown)) == 0)
 			{
diff --git a/Engine/Source/Editor/UMGEditor/Private/Settings/UMGEditorProjectSettings.cpp b/Engine/Source/Editor/UMGEditor/Private/Settings/UMGEditorProjectSettings.cpp
--- a/Engine/Source/Editor/UMGEditor/Private/Settings/UMGEditorProjectSettings.cpp
+++ b/Engine/Source/Editor/UMGEditor/Private/Settings/UMGEditorProjectSettings.cpp
@@ -17,6 +17,7 @@ UUMGEditorProjectSettings::UUMGEditorProjectSettings()
 	bShowWidgetsFromDeveloperContent = true;
 	bGraphEditorHidden = false;
 	bHideWidgetAnimationEditor = false;
+	bUseEditorConfigPaletteFiltering = false;
 
 	bUseWidgetTemplateSelector = false;
 	DefaultRootWidget = nullptr;
@@ -74,6 +75,16 @@ TArray<UWidgetCompilerRule*> UUMGEditorProjectSettings::CompilerOption_Rules(con
 	return Rules;
 }
 
+FNamePermissionList& UUMGEditorProjectSettings::GetAllowedPaletteCategories()
+{
+	return AllowedPaletteCategories;
+}
+
+FPathPermissionList& UUMGEditorProjectSettings::GetAllowedPaletteWidgets()
+{
+	return AllowedPaletteWidgets;
+}
+
 void UUMGEditorProjectSettings::GetCompilerOptionsForWidget(const UWidgetBlueprint* WidgetBlueprint, TFunctionRef<bool(const FWidgetCompilerOptions&)> Operator) const
 {
 	FString AssetPath = WidgetBlueprint->GetOutermost()->GetName();
diff --git a/Engine/Source/Editor/UMGEditor/Public/UMGEditorProjectSettings.h b/Engine/Source/Editor/UMGEditor/Public/UMGEditorProjectSettings.h
--- a/Engine/Source/Editor/UMGEditor/Public/UMGEditorProjectSettings.h
+++ b/Engine/Source/Editor/UMGEditor/Public/UMGEditorProjectSettings.h
@@ -3,11 +3,13 @@
 #pragma once
 
 #include "CoreMinimal.h"
+#include "EditorConfigBase.h"
 #include "UObject/ObjectMacros.h"
 #include "UObject/Object.h"
 #include "UObject/SoftObjectPath.h"
 #include "Engine/DeveloperSettings.h"
 #include "Engine/EngineTypes.h"
+#include "Misc/NamePermissionList.h"
 #include "UMGEditorProjectSettings.generated.h"
 
 class UWidgetCompilerRule;
@@ -178,6 +180,10 @@ public:
 	UPROPERTY(EditAnywhere, Category = Designer)
 	bool bHideWidgetAnimationEditor;
 
+	/** Set true to filter all categories and widgets out in the palette, selectively enabling them later via permission lists  */
+	UPROPERTY(EditAnywhere, config, Category = Designer)
+	bool bUseEditorConfigPaletteFiltering;
+
 	/**
 	 * The default parent class for all newly constructed widget blueprints.
 	 * The WidgetParentClass must have an empty widget hierarchy.
@@ -193,6 +199,12 @@ public:
 	EPropertyBindingPermissionLevel CompilerOption_PropertyBindingRule(const class UWidgetBlueprint* WidgetBlueprint) const;
 	TArray<UWidgetCompilerRule*> CompilerOption_Rules(const class UWidgetBlueprint* WidgetBlueprint) const;
 
+	/** Get the permission list that controls which categories are exposed in config palette filtering */
+	FNamePermissionList& GetAllowedPaletteCategories();
+
+	/** Get the permission list that controls which widgets are exposed in config palette filtering */
+	FPathPermissionList& GetAllowedPaletteWidgets();
+
 private:
 	template<typename ReturnType, typename T>
 	ReturnType GetFirstCompilerOption(const class UWidgetBlueprint* WidgetBlueprint, T FWidgetCompilerOptions::* OptionMember, ReturnType bDefaultValue) const
@@ -215,4 +227,10 @@ protected:
 
 	/** This one is unsaved, we compare it on post init to see if the save matches real */
 	int32 CurrentVersion;
+
+	/** Palette categories to allow all widgets within when using permission list palette filtering */
+	FNamePermissionList AllowedPaletteCategories;
+
+	/** Individual widgets to always allow when using permission list palette filtering, regardless of category */
+	FPathPermissionList AllowedPaletteWidgets;
 };
diff --git a/Engine/Source/Editor/UMGEditor/UMGEditor.Build.cs b/Engine/Source/Editor/UMGEditor/UMGEditor.Build.cs
--- a/Engine/Source/Editor/UMGEditor/UMGEditor.Build.cs
+++ b/Engine/Source/Editor/UMGEditor/UMGEditor.Build.cs
@@ -39,6 +39,7 @@ public class UMGEditor : ModuleRules
 
 		PrivateDependencyModuleNames.AddRange(
 			new string[] {
+				"ClassViewer",
 				"Core",
 				"CoreUObject",
 				"ContentBrowser",
@@ -47,6 +48,7 @@ public class UMGEditor : ModuleRules
 				"InputCore",
 				"Engine",
 				"AssetTools",
+				"EditorConfig",
 				"EditorSubsystem",
 				"EditorFramework",
 				"InteractiveToolsFramework",
