Add support for global filtering of unloaded classes to SPaletteViewModel

diff --git a/Engine/Source/Editor/UMGEditor/Private/Palette/SPaletteViewModel.cpp b/Engine/Source/Editor/UMGEditor/Private/Palette/SPaletteViewModel.cpp
--- a/Engine/Source/Editor/UMGEditor/Private/Palette/SPaletteViewModel.cpp
+++ b/Engine/Source/Editor/UMGEditor/Private/Palette/SPaletteViewModel.cpp
@@ -13,6 +13,7 @@
 
 #include "ClassViewerModule.h"
 #include "ClassViewerFilter.h"
+#include "EditorClassUtils.h"
 #include "DragDrop/WidgetTemplateDragDropOp.h"
 
 #include "Templates/WidgetTemplateClass.h"
@@ -29,6 +30,126 @@
 
 #define LOCTEXT_NAMESPACE "UMG"
 
+namespace UE::Editor::SPaletteViewModel::Private
+{
+	/** Helper class to perform path based filtering for unloaded BP's */
+	class FUnloadedBlueprintData : public IUnloadedBlueprintData
+	{
+	public:
+		FUnloadedBlueprintData(const FAssetData& InAssetData)
+			:ClassPath(NAME_None)
+			,ParentClassPath(NAME_None)
+			,ClassFlags(CLASS_None)
+			,bIsNormalBlueprintType(false)
+		{
+			ClassName = MakeShared<FString>(InAssetData.AssetName.ToString());
+
+			FString GeneratedClassPath;
+			const UClass* AssetClass = InAssetData.GetClass();
+			if (AssetClass && AssetClass->IsChildOf(UBlueprintGeneratedClass::StaticClass()))
+			{
+				ClassPath = InAssetData.ObjectPath;
+			}
+			else if (InAssetData.GetTagValue(FBlueprintTags::GeneratedClassPath, GeneratedClassPath))
+			{
+				ClassPath = FName(*FPackageName::ExportTextPathToObjectPath(GeneratedClassPath));
+			}
+
+			FString ParentClassPathString;
+			if (InAssetData.GetTagValue(FBlueprintTags::ParentClassPath, ParentClassPathString))
+			{
+				ParentClassPath = FName(*FPackageName::ExportTextPathToObjectPath(ParentClassPathString));
+			}
+
+			FEditorClassUtils::GetImplementedInterfaceClassPathsFromAsset(InAssetData, ImplementedInterfaces);
+		}
+
+		virtual ~FUnloadedBlueprintData()
+		{
+		}
+
+		// Begin IUnloadedBlueprintData interface
+		virtual bool HasAnyClassFlags(uint32 InFlagsToCheck) const
+		{
+			return (ClassFlags & InFlagsToCheck) != 0;
+		}
+
+		virtual bool HasAllClassFlags(uint32 InFlagsToCheck) const
+		{
+			return ((ClassFlags & InFlagsToCheck) == InFlagsToCheck);
+		}
+
+		virtual void SetClassFlags(uint32 InFlags)
+		{
+			ClassFlags = InFlags;
+		}
+
+		virtual bool ImplementsInterface(const UClass* InInterface) const
+		{
+			FString InterfacePath = InInterface->GetPathName();
+			for (const FString& ImplementedInterface : ImplementedInterfaces)
+			{
+				if (ImplementedInterface == InterfacePath)
+				{
+					return true;
+				}
+			}
+
+			return false;
+		}
+
+		virtual bool IsChildOf(const UClass* InClass) const
+		{
+			return false;
+		}
+
+		virtual bool IsA(const UClass* InClass) const
+		{
+			// Unloaded blueprint classes should always be a BPGC, so this just checks against the expected type.
+			return UBlueprintGeneratedClass::StaticClass()->UObject::IsA(InClass);
+		}
+
+		virtual const UClass* GetClassWithin() const
+		{
+			return nullptr;
+		}
+
+		virtual const UClass* GetNativeParent() const
+		{
+			return nullptr;
+		}
+
+		virtual void SetNormalBlueprintType(bool bInNormalBPType)
+		{
+			bIsNormalBlueprintType = bInNormalBPType;
+		}
+
+		virtual bool IsNormalBlueprintType() const
+		{
+			return bIsNormalBlueprintType;
+		}
+
+		virtual TSharedPtr<FString> GetClassName() const
+		{
+			return ClassName;
+		}
+
+		virtual FName GetClassPath() const
+		{
+			return ClassPath;
+		}
+		// End IUnloadedBlueprintData interface
+
+	private:
+		TSharedPtr<FString> ClassName;
+		FName ClassPath;
+		FName ParentClassPath;
+		uint32 ClassFlags;
+		TArray<FString> ImplementedInterfaces;
+		bool bIsNormalBlueprintType;
+	};
+}
+
 FWidgetTemplateViewModel::FWidgetTemplateViewModel()
 	: FavortiesViewModel(nullptr),
 	bIsFavorite(false)
@@ -229,22 +350,24 @@ void FWidgetCatalogViewModel::BuildClassWidgetList()
 	TSharedRef<FClassViewerFilterFuncs> ClassFilterFuncs = ClassViewerModule.CreateFilterFuncs();
 	FClassViewerInitializationOptions ClassViewerOptions = {};
 
-	auto WidgetPassesConfigFiltering = [&](const FString& InWidgetPath, const FString& InCategoryName, TWeakObjectPtr<UClass> InWidgetClass)
+	auto WidgetPassesConfigFiltering = [&](const FAssetData& InWidgetAssetData, const FString& InCategoryName, TWeakObjectPtr<UClass> InWidgetClass)
 	{
-		if (AllowedPaletteWidgets.PassesFilter(*InWidgetPath))
+		if (AllowedPaletteWidgets.PassesFilter(InWidgetAssetData.ObjectPath))
 		{
 			return true;
 		}
 
-		if (AllowedPaletteCategories.PassesFilter(*InCategoryName))
+		if (AllowedPaletteCategories.PassesFilter(*InCategoryName) && GlobalClassFilter.IsValid())
 		{
-			if (GlobalClassFilter.IsValid() && InWidgetClass.IsValid())
+			if (InWidgetClass.IsValid())
 			{
 				return GlobalClassFilter->IsClassAllowed(ClassViewerOptions, InWidgetClass.Get(), ClassFilterFuncs);
 			}
-			else
+			else if (InWidgetAssetData.IsValid())
 			{
-				return false;
+				using namespace UE::Editor::SPaletteViewModel::Private;
+				TSharedRef<FUnloadedBlueprintData> UnloadedBlueprint = MakeShared<FUnloadedBlueprintData>(InWidgetAssetData);
+				return GlobalClassFilter->IsUnloadedClassAllowed(ClassViewerOptions, UnloadedBlueprint, ClassFilterFuncs);
 			}
 		}
 
@@ -289,7 +412,7 @@ void FWidgetCatalogViewModel::BuildClassWidgetList()
 
 		if (bUseEditorConfigPaletteFiltering)
 		{
-			if (!WidgetPassesConfigFiltering(WidgetAssetData.ObjectPath.ToString(), WidgetClass->GetDefaultObject<UWidget>()->GetPaletteCategory().ToString(), WidgetClass))
+			if (!WidgetPassesConfigFiltering(WidgetAssetData, WidgetClass->GetDefaultObject<UWidget>()->GetPaletteCategory().ToString(), WidgetClass))
 			{
 				continue;
 			}
@@ -407,7 +530,7 @@ void FWidgetCatalogViewModel::BuildClassWidgetList()
 			if (bUseEditorConfigPaletteFiltering)
 			{
 				FText Category = FWidgetTemplateBlueprintClass(WidgetBPAssetData, WidgetBPClass).GetCategory();
-				if (!WidgetPassesConfigFiltering(WidgetBPAssetData.ObjectPath.ToString(), Category.ToString(), WidgetBPClass))
+				if (!WidgetPassesConfigFiltering(WidgetBPAssetData, Category.ToString(), WidgetBPClass))
 				{
 					continue;
 				}
