Include unloaded BP assets in palette filtering

diff --git a/Engine/Source/Editor/UMGEditor/Private/Palette/SPaletteViewModel.cpp b/Engine/Source/Editor/UMGEditor/Private/Palette/SPaletteViewModel.cpp
--- a/Engine/Source/Editor/UMGEditor/Private/Palette/SPaletteViewModel.cpp
+++ b/Engine/Source/Editor/UMGEditor/Private/Palette/SPaletteViewModel.cpp
@@ -332,6 +332,7 @@ void FWidgetCatalogViewModel::BuildWidgetList()
 void FWidgetCatalogViewModel::BuildClassWidgetList()
 {
 	static const FName DevelopmentStatusKey(TEXT("DevelopmentStatus"));
+	static const FString InvalidCategoryName("InvalidCategoryName");
 
 	TMap<FName, TSubclassOf<UUserWidget>> LoadedWidgetBlueprintClassesByName;
 
@@ -466,41 +467,46 @@ void FWidgetCatalogViewModel::BuildClassWidgetList()
 	TArray<FAssetData> AllBPsAssetData;
 	AssetRegistryModule.Get().GetAssetsByClass(UBlueprint::StaticClass()->GetFName(), AllBPsAssetData, true);
 
-	// Ignore unloaded assets when using palette filtering, as they lack data needed to correctly pass filters
-	if (!bUseEditorConfigPaletteFiltering)
+	for (FAssetData& BPAssetData : AllBPsAssetData)
 	{
-		for (FAssetData& BPAssetData : AllBPsAssetData)
+		// Blueprints get the class type actions for their parent native class - this avoids us having to load the blueprint
+		UClass* ParentClass = nullptr;
+		FString ParentClassName;
+		if (!BPAssetData.GetTagValue(FBlueprintTags::NativeParentClassPath, ParentClassName))
 		{
-			// Blueprints get the class type actions for their parent native class - this avoids us having to load the blueprint
-			UClass* ParentClass = nullptr;
-			FString ParentClassName;
-			if (!BPAssetData.GetTagValue(FBlueprintTags::NativeParentClassPath, ParentClassName))
+			BPAssetData.GetTagValue(FBlueprintTags::ParentClassPath, ParentClassName);
+		}
+		if (!ParentClassName.IsEmpty())
+		{
+			UObject* Outer = nullptr;
+			ResolveName(Outer, ParentClassName, false, false);
+			ParentClass = FindObject<UClass>(ANY_PACKAGE, *ParentClassName);
+			// UUserWidgets have their own loading section, and we don't want to process any blueprints that don't have UWidget parents
+			if (ParentClass)
 			{
-				BPAssetData.GetTagValue(FBlueprintTags::ParentClassPath, ParentClassName);
+				if (!ParentClass->IsChildOf(UWidget::StaticClass()) || ParentClass->IsChildOf(UUserWidget::StaticClass()))
+				{
+					continue;
+				}
 			}
-			if (!ParentClassName.IsEmpty())
+		}
+
+		if (!FilterAssetData(BPAssetData))
+		{
+			// If this object isn't currently loaded, add it to the palette view
+			if (BPAssetData.ToSoftObjectPath().ResolveObject() == nullptr)
 			{
-				UObject* Outer = nullptr;
-				ResolveName(Outer, ParentClassName, false, false);
-				ParentClass = FindObject<UClass>(ANY_PACKAGE, *ParentClassName);
-				// UUserWidgets have their own loading section, and we don't want to process any blueprints that don't have UWidget parents
-				if (ParentClass)
+				if (bUseEditorConfigPaletteFiltering)
 				{
-					if (!ParentClass->IsChildOf(UWidget::StaticClass()) || ParentClass->IsChildOf(UUserWidget::StaticClass()))
+					// Still check the soft object path even if we can't check category due to BP asset being unloaded
+					if (!WidgetPassesConfigFiltering(BPAssetData, InvalidCategoryName, nullptr))
 					{
 						continue;
 					}
 				}
-			}
 
-			if (!FilterAssetData(BPAssetData))
-			{
-				// If this object isn't currently loaded, add it to the palette view
-				if (BPAssetData.ToSoftObjectPath().ResolveObject() == nullptr)
-				{
-					auto Template = MakeShareable(new FWidgetTemplateClass(BPAssetData, nullptr));
-					AddWidgetTemplate(Template);
-				}
+				auto Template = MakeShareable(new FWidgetTemplateClass(BPAssetData, nullptr));
+				AddWidgetTemplate(Template);
 			}
 		}
 	}
diff --git a/Engine/Source/Editor/UMGEditor/Private/Templates/WidgetTemplateClass.cpp b/Engine/Source/Editor/UMGEditor/Private/Templates/WidgetTemplateClass.cpp
--- a/Engine/Source/Editor/UMGEditor/Private/Templates/WidgetTemplateClass.cpp
+++ b/Engine/Source/Editor/UMGEditor/Private/Templates/WidgetTemplateClass.cpp
@@ -40,6 +40,22 @@ FWidgetTemplateClass::FWidgetTemplateClass(const FAssetData& InWidgetAssetData,
 	else
 	{
 		Name = FText::FromString(FName::NameToDisplayString(WidgetAssetData.AssetName.ToString(), false));
+
+		// Blueprints get the class type actions for their parent native class - this avoids us having to load the blueprint
+		if (InWidgetAssetData.IsValid())
+		{
+			FString ParentClassName;
+			if (!WidgetAssetData.GetTagValue(FBlueprintTags::NativeParentClassPath, ParentClassName))
+			{
+				WidgetAssetData.GetTagValue(FBlueprintTags::ParentClassPath, ParentClassName);
+			}
+			if (!ParentClassName.IsEmpty())
+			{
+				UObject* Outer = nullptr;
+				ResolveName(Outer, ParentClassName, false, false);
+				CachedParentClass = FindObject<UClass>(ANY_PACKAGE, *ParentClassName);
+			}
+		}
 	}
 }
 
@@ -55,6 +71,10 @@ FText FWidgetTemplateClass::GetCategory() const
 		auto DefaultWidget = WidgetClass->GetDefaultObject<UWidget>();
 		return DefaultWidget->GetPaletteCategory();
 	}
+	else if (CachedParentClass.IsValid() && CachedParentClass->IsChildOf(UWidget::StaticClass()) && !CachedParentClass->IsChildOf(UUserWidget::StaticClass()))
+	{
+		return CachedParentClass->GetDefaultObject<UWidget>()->GetPaletteCategory();
+	}
 	else
 	{
 		auto DefaultWidget = UWidget::StaticClass()->GetDefaultObject<UWidget>();
@@ -81,6 +101,10 @@ const FSlateBrush* FWidgetTemplateClass::GetIcon() const
 	{
 		return FSlateIconFinder::FindIconBrushForClass(WidgetClass.Get());
 	}
+	else if (CachedParentClass.IsValid() && CachedParentClass->IsChildOf(UWidget::StaticClass()) && !CachedParentClass->IsChildOf(UUserWidget::StaticClass()))
+	{
+		return FSlateIconFinder::FindIconBrushForClass(CachedParentClass.Get());
+	}
 	else
 	{
 		return FSlateIconFinder::FindIconBrushForClass(UWidget::StaticClass());
diff --git a/Engine/Source/Editor/UMGEditor/Public/Templates/WidgetTemplateClass.h b/Engine/Source/Editor/UMGEditor/Public/Templates/WidgetTemplateClass.h
--- a/Engine/Source/Editor/UMGEditor/Public/Templates/WidgetTemplateClass.h
+++ b/Engine/Source/Editor/UMGEditor/Public/Templates/WidgetTemplateClass.h
@@ -62,4 +62,7 @@ protected:
 
 	/** The asset data for the widget blueprint */
 	FAssetData WidgetAssetData;
+
+	/** Parent Class of this widget template, may not be valid */
+	TWeakObjectPtr<UClass> CachedParentClass;
 };
