Fix to Widget Palette that prevent it from having more that one widget in the palette

diff --git a/Engine/Source/Editor/UMGEditor/Private/Library/SLibraryViewModel.cpp b/Engine/Source/Editor/UMGEditor/Private/Library/SLibraryViewModel.cpp
--- a/Engine/Source/Editor/UMGEditor/Private/Library/SLibraryViewModel.cpp
+++ b/Engine/Source/Editor/UMGEditor/Private/Library/SLibraryViewModel.cpp
@@ -30,7 +30,7 @@
 
 #define LOCTEXT_NAMESPACE "UMG"
 
-void FLibraryViewModel::BuildWidgetTemplateCategory(FString& Category, TArray<TSharedPtr<FWidgetTemplate>>& Templates)
+void FLibraryViewModel::BuildWidgetTemplateCategory(FString& Category, TArray<TSharedPtr<FWidgetTemplate>>& Templates, TArray<FString>& FavoritesList)
 {
 	TSharedPtr<FWidgetHeaderViewModel> Header = MakeShareable(new FWidgetHeaderViewModel());
 	Header->GroupName = FText::FromString(Category);
@@ -40,33 +40,32 @@ void FLibraryViewModel::BuildWidgetTemplateCategory(FString& Category, TArray<TS
 	Header->Children.Add(TemplateViewModel);
 	WidgetTemplateListViewModels.Add(TemplateViewModel);
 
-	// @TODO: DarenC - Reference for when implementing favorites system
-	//for ( auto& Template : Entry.Value )
-	//{
-	//	TSharedPtr<FWidgetTemplateViewModel> TemplateViewModel = MakeShareable(new FWidgetTemplateViewModel());
-	//	TemplateViewModel->Template = Template;
-	//	TemplateViewModel->FavortiesViewModel = this;
-	//	Header->Children.Add(TemplateViewModel);
-
-	//	// If it's a favorite, we also add it to the Favorite section
-	//	int32 index = FavoritesList.Find(Template->Name.ToString());
-	//	if (index != INDEX_NONE)
-	//	{
-	//		TemplateViewModel->SetFavorite();
-
-	//		// We have to create a second copy of the ViewModel for the treeview has it doesn't support to have the same element twice.
-	//		TSharedPtr<FWidgetTemplateViewModel> FavoriteTemplateViewModel = MakeShareable(new FWidgetTemplateViewModel());
-	//		FavoriteTemplateViewModel->Template = Template;
-	//		FavoriteTemplateViewModel->FavortiesViewModel = this;
-	//		FavoriteTemplateViewModel->SetFavorite();
-
-	//		FavoriteHeader->Children.Add(FavoriteTemplateViewModel);
-
-	//		// Remove the favorite from the temporary list
-	//		FavoritesList.RemoveAt(index);
-	//	}
-
-	//}
+	for ( auto& Template : Templates)
+	{
+		TSharedPtr<FWidgetTemplateViewModel> CurrentTemplateViewModel = MakeShareable(new FWidgetTemplateViewModel());
+		CurrentTemplateViewModel->Template = Template;
+		CurrentTemplateViewModel->FavortiesViewModel = this;
+		Header->Children.Add(CurrentTemplateViewModel);
+
+		// If it's a favorite, we also add it to the Favorite section
+		int32 index = FavoritesList.Find(Template->Name.ToString());
+		if (index != INDEX_NONE)
+		{
+			CurrentTemplateViewModel->SetFavorite();
+
+			// We have to create a second copy of the ViewModel for the treeview has it doesn't support to have the same element twice.
+			TSharedPtr<FWidgetTemplateViewModel> FavoriteTemplateViewModel = MakeShareable(new FWidgetTemplateViewModel());
+			FavoriteTemplateViewModel->Template = Template;
+			FavoriteTemplateViewModel->FavortiesViewModel = this;
+			FavoriteTemplateViewModel->SetFavorite();
+
+			FavoriteHeader->Children.Add(FavoriteTemplateViewModel);
+
+			// Remove the favorite from the temporary list
+			FavoritesList.RemoveAt(index);
+		}
+
+	}
 
 	Header->Children.Sort([](const TSharedPtr<FWidgetViewModel>& L, const TSharedPtr<FWidgetViewModel>& R) { return R->GetName().CompareTo(L->GetName()) > 0; });
 
diff --git a/Engine/Source/Editor/UMGEditor/Private/Library/SLibraryViewModel.h b/Engine/Source/Editor/UMGEditor/Private/Library/SLibraryViewModel.h
--- a/Engine/Source/Editor/UMGEditor/Private/Library/SLibraryViewModel.h
+++ b/Engine/Source/Editor/UMGEditor/Private/Library/SLibraryViewModel.h
@@ -74,7 +74,7 @@ public:
 	FLibraryViewModel(TSharedPtr<FWidgetBlueprintEditor> InBlueprintEditor) : FWidgetCatalogViewModel(InBlueprintEditor) { }
 
 	//~ Begin FWidgetCatalogViewModel Interface
-	virtual void BuildWidgetTemplateCategory(FString& Category, TArray<TSharedPtr<FWidgetTemplate>>& Templates) override;
+	virtual void BuildWidgetTemplateCategory(FString& Category, TArray<TSharedPtr<FWidgetTemplate>>& Templates, TArray<FString>& FavoritesList) override;
 	virtual void AddToFavorites(const FWidgetTemplateViewModel* WidgetTemplateViewModel) override;
 	virtual void RemoveFromFavorites(const FWidgetTemplateViewModel* WidgetTemplateViewModel) override;
 	virtual void SetSearchText(const FText& InSearchText) override;
diff --git a/Engine/Source/Editor/UMGEditor/Private/Palette/SPaletteViewModel.cpp b/Engine/Source/Editor/UMGEditor/Private/Palette/SPaletteViewModel.cpp
--- a/Engine/Source/Editor/UMGEditor/Private/Palette/SPaletteViewModel.cpp
+++ b/Engine/Source/Editor/UMGEditor/Private/Palette/SPaletteViewModel.cpp
@@ -298,8 +298,8 @@ void FWidgetCatalogViewModel::BuildWidgetList()
 	// For each entry in the category create a view model for the widget template
 	for ( auto& Entry : WidgetTemplateCategories )
 	{
-		BuildWidgetTemplateCategory(Entry.Key, Entry.Value);
-	}	
+		BuildWidgetTemplateCategory(Entry.Key, Entry.Value, FavoritesList);
+	}
 
 	// Remove all Favorites that may be left in the list.Typically happening when the list of favorite contains widget that were deleted since the last opening.
 	for (const FString& favoriteName : FavoritesList)
@@ -639,14 +639,13 @@ void FWidgetCatalogViewModel::HandleOnAssetsDeleted(const TArray<UClass*>& Delet
 	}
 }
 
-void FPaletteViewModel::BuildWidgetTemplateCategory(FString& Category, TArray<TSharedPtr<FWidgetTemplate>>& Templates)
+void FPaletteViewModel::BuildWidgetTemplateCategory(FString& Category, TArray<TSharedPtr<FWidgetTemplate>>& Templates, TArray<FString>& FavoritesList)
 {
 	TSharedPtr<FWidgetHeaderViewModel> Header = MakeShareable(new FWidgetHeaderViewModel());
 	Header->GroupName = FText::FromString(Category);
 
 	// Copy of the list of favorites to be able to do some cleanup in the real list
 	UWidgetPaletteFavorites* FavoritesPalette = GetDefault<UWidgetDesignerSettings>()->Favorites;
-	TArray<FString> FavoritesList = FavoritesPalette->GetFavorites();
 
 	for (auto& Template : Templates)
 	{
diff --git a/Engine/Source/Editor/UMGEditor/Private/Palette/SPaletteViewModel.h b/Engine/Source/Editor/UMGEditor/Private/Palette/SPaletteViewModel.h
--- a/Engine/Source/Editor/UMGEditor/Private/Palette/SPaletteViewModel.h
+++ b/Engine/Source/Editor/UMGEditor/Private/Palette/SPaletteViewModel.h
@@ -179,7 +179,7 @@ protected:
 	UWidgetBlueprint* GetBlueprint() const;
 
 	virtual void BuildWidgetList();
-	virtual void BuildWidgetTemplateCategory(FString& Category, TArray<TSharedPtr<FWidgetTemplate>>& Templates) = 0;
+	virtual void BuildWidgetTemplateCategory(FString& Category, TArray<TSharedPtr<FWidgetTemplate>>& Templates, TArray<FString>& FavoritesList) = 0;
 	void BuildClassWidgetList();
 
 	static bool FilterAssetData(FAssetData& BPAssetData);
@@ -223,7 +223,7 @@ public:
 	FPaletteViewModel(TSharedPtr<FWidgetBlueprintEditor> InBlueprintEditor) : FWidgetCatalogViewModel(InBlueprintEditor) { }
 
 	//~ Begin FWidgetCatalogViewModel Interface
-	virtual void BuildWidgetTemplateCategory(FString& Category, TArray<TSharedPtr<FWidgetTemplate>>& Templates) override;
+	virtual void BuildWidgetTemplateCategory(FString& Category, TArray<TSharedPtr<FWidgetTemplate>>& Templates, TArray<FString>& FavoritesList) override;
 	virtual void AddToFavorites(const FWidgetTemplateViewModel* WidgetTemplateViewModel) override;
 	virtual void RemoveFromFavorites(const FWidgetTemplateViewModel* WidgetTemplateViewModel) override;
 	//~ End FWidgetCatalogViewModel Interface
