UE-148560 - UMG property binding fixup for real numbers doesn't always work

The fixup pass for delegates in UMG bindings had faulty logic for determining the name of a function. The "FunctionName" field on the binding may not actually reflect the true name of the function. For example, a function rename can make this field stale. This would eventually cause BP compilation errors since we didn't find the matching function graph(s) that needed fixup.

Instead, we should call GetMemberName on the final path segment. This will use a GUID for determining what the actual function name is.

diff --git a/Engine/Source/Editor/UMGEditor/Private/WidgetBlueprint.cpp b/Engine/Source/Editor/UMGEditor/Private/WidgetBlueprint.cpp
--- a/Engine/Source/Editor/UMGEditor/Private/WidgetBlueprint.cpp
+++ b/Engine/Source/Editor/UMGEditor/Private/WidgetBlueprint.cpp
@@ -593,7 +593,7 @@ void UWidgetBlueprint::ReplaceDeprecatedNodes()
 		// However, UMG property bindings are a special case: the BP functions that bind to the native delegate must agree on their underlying types.
 		// Specifically, bindings used with float properties *must* use the PC_Float type as the return value in a BP function.
 		// In order to correct this behavior, we need to:
-		// * Iterate throught the property bindings.
+		// * Iterate through the property bindings.
 		// * Find the corresponding delegate signature.
 		// * Find the function graph that matches the binding.
 		// * Find the result node.
@@ -614,9 +614,18 @@ void UWidgetBlueprint::ReplaceDeprecatedNodes()
 
 					if (BindableProperty)
 					{
-						auto GraphMatchesBindingPredicate = [&Binding](const UEdGraph* Graph) {
+						FName FunctionName = Binding.FunctionName;
+
+						if (!Binding.SourcePath.IsEmpty())
+						{
+							check(Binding.SourcePath.Segments.Num() > 0);
+							const FEditorPropertyPathSegment& LastSegment = Binding.SourcePath.Segments[Binding.SourcePath.Segments.Num() - 1];
+							FunctionName = LastSegment.GetMemberName();
+						}
+
+						auto GraphMatchesBindingPredicate = [FunctionName](const UEdGraph* Graph) {
 							check(Graph);
-							return (Binding.FunctionName == Graph->GetFName());
+							return (FunctionName == Graph->GetFName());
 						};
 
 						if (UEdGraph** GraphEntry = Graphs.FindByPredicate(GraphMatchesBindingPredicate))
