UMG - Introducing named slot inheritance.  This new feature allows parent classes to insert named slots into their hierarchy, and for child classes to insert content into those named slots.  This in effect will give people the ability to build base class templates and then derive them to further refine the look and feel in those named slots for the widgets specific purpose.  For example, you could make a message box base class widget, with a few named slots for specific specialized message boxes without having to build a new widget from scratch.

diff --git a/Engine/Source/Editor/UMGEditor/Private/Hierarchy/SHierarchyViewItem.cpp b/Engine/Source/Editor/UMGEditor/Private/Hierarchy/SHierarchyViewItem.cpp
--- a/Engine/Source/Editor/UMGEditor/Private/Hierarchy/SHierarchyViewItem.cpp
+++ b/Engine/Source/Editor/UMGEditor/Private/Hierarchy/SHierarchyViewItem.cpp
@@ -700,11 +700,20 @@ void FHierarchyRoot::GetChildren(TArray< TSharedPtr<FHierarchyModel> >& Children
 	TSharedPtr<FWidgetBlueprintEditor> BPEd = BlueprintEditor.Pin();
 	UWidgetBlueprint* Blueprint = BPEd->GetWidgetBlueprintObj();
 
+	// We only show the hierarchy of the widget tree we own.  If this is a subclass of a widget with a parent tree
+	// we won't actually show those widgets here because we can't affect them in a useful inheritable way.
 	if ( Blueprint->WidgetTree->RootWidget )
 	{
 		TSharedPtr<FHierarchyWidget> RootChild = MakeShareable(new FHierarchyWidget(BPEd->GetReferenceFromTemplate(Blueprint->WidgetTree->RootWidget), BPEd));
 		Children.Add(RootChild);
 	}
+
+	// Grab any exposed named slots from the super classes CDO.  These slots can have content slotted into them by this subclass.
+	for ( const FName& SlotName : Blueprint->GetInheritedNamedSlots() )
+	{
+		TSharedPtr<FNamedSlotModelSubclass> ChildItem = MakeShareable(new FNamedSlotModelSubclass(Blueprint, SlotName, BPEd));
+		Children.Add(ChildItem);
+	}
 }
 
 void FHierarchyRoot::OnSelection()
@@ -793,28 +802,25 @@ FReply FHierarchyRoot::HandleAcceptDrop(FDragDropEvent const& DragDropEvent, EIt
 
 //////////////////////////////////////////////////////////////////////////
 
-FNamedSlotModel::FNamedSlotModel(FWidgetReference InItem, FName InSlotName, TSharedPtr<FWidgetBlueprintEditor> InBlueprintEditor)
+FNamedSlotModelBase::FNamedSlotModelBase(FName InSlotName, TSharedPtr<FWidgetBlueprintEditor> InBlueprintEditor)
 	: FHierarchyModel(InBlueprintEditor)
-	, Item(InItem)
 	, SlotName(InSlotName)
 {
 }
 
-FName FNamedSlotModel::GetUniqueName() const
+const FSlateBrush* FNamedSlotModelBase::GetImage() const
 {
-	UWidget* WidgetTemplate = Item.GetTemplate();
-	if ( WidgetTemplate )
-	{
-		FString UniqueSlot = WidgetTemplate->GetName() + TEXT(".") + SlotName.ToString();
-		return FName(*UniqueSlot);
-	}
+	return nullptr;
+}
 
-	return NAME_None;
+FSlateFontInfo FNamedSlotModelBase::GetFont() const
+{
+	return FCoreStyle::GetDefaultFontStyle("Bold", 10);
 }
 
-FText FNamedSlotModel::GetText() const
+FText FNamedSlotModelBase::GetText() const
 {
-	if ( INamedSlotInterface* NamedSlotHost = Cast<INamedSlotInterface>(Item.GetTemplate()) )
+	if (INamedSlotInterface* NamedSlotHost = GetNamedSlotHost())
 	{
 		TSet<FWidgetReference> SelectedWidgets;
 		if ( UWidget* SlotContent = NamedSlotHost->GetContentForSlot(SlotName) )
@@ -826,59 +832,34 @@ FText FNamedSlotModel::GetText() const
 	return FText::FromName(SlotName);
 }
 
-const FSlateBrush* FNamedSlotModel::GetImage() const
-{
-	return NULL;
-}
 
-FSlateFontInfo FNamedSlotModel::GetFont() const
+void FNamedSlotModelBase::GetChildren(TArray< TSharedPtr<FHierarchyModel> >& Children)
 {
-	return FCoreStyle::GetDefaultFontStyle("Bold", 10);
-}
-
-void FNamedSlotModel::GetChildren(TArray< TSharedPtr<FHierarchyModel> >& Children)
-{
-	TSharedPtr<FWidgetBlueprintEditor> BPEd = BlueprintEditor.Pin();
-	if ( INamedSlotInterface* NamedSlotHost = Cast<INamedSlotInterface>(Item.GetTemplate()) )
+	if (INamedSlotInterface* NamedSlotHost = GetNamedSlotHost())
 	{
 		TSet<FWidgetReference> SelectedWidgets;
 		if ( UWidget* TemplateSlotContent = NamedSlotHost->GetContentForSlot(SlotName) )
 		{
+			TSharedPtr<FWidgetBlueprintEditor> BPEd = BlueprintEditor.Pin();
 			TSharedPtr<FHierarchyWidget> RootChild = MakeShareable(new FHierarchyWidget(BPEd->GetReferenceFromTemplate(TemplateSlotContent), BPEd));
 			Children.Add(RootChild);
 		}
 	}
 }
 
-void FNamedSlotModel::OnSelection()
+void FNamedSlotModelBase::OnSelection()
 {
-	TSharedPtr<FWidgetBlueprintEditor> Editor = BlueprintEditor.Pin();
-
-	FNamedSlotSelection Selection;
-	Selection.NamedSlotHostWidget = Item;
-	Selection.SlotName = SlotName;
-	Editor->SetSelectedNamedSlot(Selection);
+	// No-Op intentionally.
 }
 
-void FNamedSlotModel::UpdateSelection()
+void FNamedSlotModelBase::UpdateSelection()
 {
-	//bIsSelected = false;
-
-	//const TSet<FWidgetReference>& SelectedWidgets = BlueprintEditor.Pin()->GetSelectedWidgets();
-
-	//TSharedPtr<FWidgetBlueprintEditor> BPEd = BlueprintEditor.Pin();
-	//if ( INamedSlotInterface* NamedSlotHost = Cast<INamedSlotInterface>(Item.GetTemplate()) )
-	//{
-	//	if ( UWidget* TemplateSlotContent = NamedSlotHost->GetContentForSlot(SlotName) )
-	//	{
-	//		bIsSelected = SelectedWidgets.Contains(BPEd->GetReferenceFromTemplate(TemplateSlotContent));
-	//	}
-	//}
+	// No-Op intentionally.
 }
 
-FWidgetReference FNamedSlotModel::AsDraggedWidgetReference() const
+FWidgetReference FNamedSlotModelBase::AsDraggedWidgetReference() const
 {
-	if (INamedSlotInterface* NamedSlotHost = Cast<INamedSlotInterface>(Item.GetTemplate()))
+	if (INamedSlotInterface* NamedSlotHost = GetNamedSlotHost())
 	{
 		// Only assign content to the named slot if it is null.
 		if (UWidget* Content = NamedSlotHost->GetContentForSlot(SlotName))
@@ -890,7 +871,7 @@ FWidgetReference FNamedSlotModel::AsDraggedWidgetReference() const
 	return FWidgetReference();
 }
 
-TOptional<EItemDropZone> FNamedSlotModel::HandleCanAcceptDrop(const FDragDropEvent& DragDropEvent, EItemDropZone DropZone)
+TOptional<EItemDropZone> FNamedSlotModelBase::HandleCanAcceptDrop(const FDragDropEvent& DragDropEvent, EItemDropZone DropZone)
 {
 	UWidgetBlueprint* Blueprint = BlueprintEditor.Pin()->GetWidgetBlueprintObj();
 
@@ -904,7 +885,7 @@ TOptional<EItemDropZone> FNamedSlotModel::HandleCanAcceptDrop(const FDragDropEve
 			DecoratedDragDropOp->ResetToDefaultToolTip();
 		}
 
-		if (INamedSlotInterface* NamedSlotHost = Cast<INamedSlotInterface>(Item.GetTemplate()))
+		if (INamedSlotInterface* NamedSlotHost = GetNamedSlotHost())
 		{
 			// Only assign content to the named slot if it is null.
 			if (NamedSlotHost->GetContentForSlot(SlotName) != nullptr)
@@ -919,7 +900,7 @@ TOptional<EItemDropZone> FNamedSlotModel::HandleCanAcceptDrop(const FDragDropEve
 
 			if (UWidget* Widget = FWidgetBlueprintEditorUtils::GetWidgetTemplateFromDragDrop(Blueprint, Blueprint->WidgetTree, DragDropOp))
 			{
-				bool bIsFreeFromCircularReferences = !HasCircularReferences(Blueprint, Widget, DragDropOp);
+				const bool bIsFreeFromCircularReferences = !HasCircularReferences(Blueprint, Widget, DragDropOp);
 				if (DecoratedDragDropOp.IsValid())
 				{
 					DecoratedDragDropOp->CurrentIconBrush = FAppStyle::GetBrush(TEXT("Graph.ConnectorFeedback.OK"));
@@ -934,7 +915,7 @@ TOptional<EItemDropZone> FNamedSlotModel::HandleCanAcceptDrop(const FDragDropEve
 	{
 		HierarchyDragDropOp->ResetToDefaultToolTip();
 
-		if (INamedSlotInterface* NamedSlotHost = Cast<INamedSlotInterface>(Item.GetTemplate()))
+		if (INamedSlotInterface* NamedSlotHost = GetNamedSlotHost())
 		{
 			// Only assign content to the named slot if it is null.
 			if (NamedSlotHost->GetContentForSlot(SlotName) != nullptr)
@@ -949,7 +930,7 @@ TOptional<EItemDropZone> FNamedSlotModel::HandleCanAcceptDrop(const FDragDropEve
 
 			// Verify that the new location we're placing the widget is not inside of its existing children.
 			Blueprint->WidgetTree->ForWidgetAndChildren(TemplateWidget, [&](UWidget* Widget) {
-				if (Item.GetTemplate() == Widget)
+				if (GetNamedSlotHostWidget() == Widget)
 				{
 					bFoundNewParentInChildSet = true;
 				}
@@ -970,10 +951,9 @@ TOptional<EItemDropZone> FNamedSlotModel::HandleCanAcceptDrop(const FDragDropEve
 	return TOptional<EItemDropZone>();
 }
 
-FReply FNamedSlotModel::HandleAcceptDrop(FDragDropEvent const& DragDropEvent, EItemDropZone DropZone)
+FReply FNamedSlotModelBase::HandleAcceptDrop(FDragDropEvent const& DragDropEvent, EItemDropZone DropZone)
 {
-	UWidget* SlotHostWidget = Item.GetTemplate();
-	INamedSlotInterface* NamedSlotHost = Cast<INamedSlotInterface>(SlotHostWidget);
+	INamedSlotInterface* NamedSlotHost = GetNamedSlotHost();
 	if (NamedSlotHost == nullptr)
 	{
 		return FReply::Unhandled();
@@ -994,7 +974,7 @@ FReply FNamedSlotModel::HandleAcceptDrop(FDragDropEvent const& DragDropEvent, EI
 
 		if (UWidget* DroppingWidget = FWidgetBlueprintEditorUtils::GetWidgetTemplateFromDragDrop(Blueprint, Blueprint->WidgetTree, DragDropOp))
 		{
-			DoDrop(SlotHostWidget, DroppingWidget);
+			DoDrop(NamedSlotHost, DroppingWidget);
 			return FReply::Handled();
 		}
 		else
@@ -1026,7 +1006,7 @@ FReply FNamedSlotModel::HandleAcceptDrop(FDragDropEvent const& DragDropEvent, EI
 
 		DroppingWidget->RemoveFromParent();
 
-		DoDrop(SlotHostWidget, DroppingWidget);
+		DoDrop(NamedSlotHost, DroppingWidget);
 
 		return FReply::Handled();
 	}
@@ -1034,15 +1014,17 @@ FReply FNamedSlotModel::HandleAcceptDrop(FDragDropEvent const& DragDropEvent, EI
 	return FReply::Unhandled();
 }
 
-void FNamedSlotModel::DoDrop(UWidget* NamedSlotHostWidget, UWidget* DroppingWidget)
+void FNamedSlotModelBase::DoDrop(INamedSlotInterface* NamedSlotHost, UWidget* DroppingWidget)
 {
 	UWidgetBlueprint* Blueprint = BlueprintEditor.Pin()->GetWidgetBlueprintObj();
 
-	NamedSlotHostWidget->SetFlags(RF_Transactional);
-	NamedSlotHostWidget->Modify();
-
-	TScriptInterface<INamedSlotInterface> NamedSlotInterface = TScriptInterface<INamedSlotInterface>(NamedSlotHostWidget);
-	NamedSlotInterface->SetContentForSlot(SlotName, DroppingWidget);
+	if (UObject* NamedSlotHostObject = Cast<UObject>(NamedSlotHost))
+	{
+		NamedSlotHostObject->SetFlags(RF_Transactional);
+		NamedSlotHostObject->Modify();
+	}
+	
+	NamedSlotHost->SetContentForSlot(SlotName, DroppingWidget);
 
 	FBlueprintEditorUtils::MarkBlueprintAsStructurallyModified(Blueprint);
 
@@ -1054,6 +1036,84 @@ void FNamedSlotModel::DoDrop(UWidget* NamedSlotHostWidget, UWidget* DroppingWidg
 
 //////////////////////////////////////////////////////////////////////////
 
+FNamedSlotModel::FNamedSlotModel(FWidgetReference InItem, FName InSlotName, TSharedPtr<FWidgetBlueprintEditor> InBlueprintEditor)
+	: FNamedSlotModelBase(InSlotName, InBlueprintEditor)
+	, Item(InItem)
+{
+}
+
+FName FNamedSlotModel::GetUniqueName() const
+{
+	if ( const UWidget* WidgetTemplate = Item.GetTemplate() )
+	{
+		TStringBuilder<256> StringBuilder;
+		StringBuilder.Append(WidgetTemplate->GetName());
+		StringBuilder.Append(TEXT("."));
+		SlotName.AppendString(StringBuilder);
+
+		return FName(StringBuilder);
+	}
+
+	return FName();
+}
+
+INamedSlotInterface* FNamedSlotModel::GetNamedSlotHost() const
+{
+	return Cast<INamedSlotInterface>(Item.GetTemplate());
+}
+
+UWidget* FNamedSlotModel::GetNamedSlotHostWidget() const
+{
+	return Cast<UWidget>(Item.GetTemplate());
+}
+
+void FNamedSlotModel::OnSelection()
+{
+	TSharedPtr<FWidgetBlueprintEditor> Editor = BlueprintEditor.Pin();
+	check(Editor.IsValid());
+
+	FNamedSlotSelection Selection;
+	Selection.NamedSlotHostWidget = Item;
+	Selection.SlotName = SlotName;
+	Editor->SetSelectedNamedSlot(Selection);
+}
+
+//////////////////////////////////////////////////////////////////////////
+
+FNamedSlotModelSubclass::FNamedSlotModelSubclass(UWidgetBlueprint* InBlueprint, FName InSlotName, TSharedPtr<FWidgetBlueprintEditor> InBlueprintEditor)
+	: FNamedSlotModelBase(InSlotName, InBlueprintEditor)
+	, Blueprint(InBlueprint)
+{
+}
+
+FName FNamedSlotModelSubclass::GetUniqueName() const
+{
+	const FString UniqueSlot = TEXT("This.") + SlotName.ToString();
+	return FName(*UniqueSlot);
+}
+
+INamedSlotInterface* FNamedSlotModelSubclass::GetNamedSlotHost() const
+{
+	return Blueprint->WidgetTree;
+}
+
+UWidget* FNamedSlotModelSubclass::GetNamedSlotHostWidget() const
+{
+	// CDO stored named slot elements don't have a host widget they can use/reference/talk about.
+	return nullptr;
+}
+
+void FNamedSlotModelSubclass::OnSelection()
+{
+	TSharedPtr<FWidgetBlueprintEditor> Editor = BlueprintEditor.Pin();
+
+	FNamedSlotSelection Selection;
+	Selection.SlotName = SlotName;
+	Editor->SetSelectedNamedSlot(Selection);
+}
+
+//////////////////////////////////////////////////////////////////////////
+
 FHierarchyWidget::FHierarchyWidget(FWidgetReference InItem, TSharedPtr<FWidgetBlueprintEditor> InBlueprintEditor)
 	: FHierarchyModel(InBlueprintEditor)
 	, Item(InItem)
diff --git a/Engine/Source/Editor/UMGEditor/Private/Hierarchy/SHierarchyViewItem.h b/Engine/Source/Editor/UMGEditor/Private/Hierarchy/SHierarchyViewItem.h
--- a/Engine/Source/Editor/UMGEditor/Private/Hierarchy/SHierarchyViewItem.h
+++ b/Engine/Source/Editor/UMGEditor/Private/Hierarchy/SHierarchyViewItem.h
@@ -140,7 +140,7 @@ public:
 	virtual void OnSelection() override;
 
 	virtual bool DoesWidgetOverrideFlowDirection() const override;
-	virtual bool DoesWidgetOverrideNavigation() const { return false; }
+	virtual bool DoesWidgetOverrideNavigation() const override { return false; }
 
 	virtual TOptional<EItemDropZone> HandleCanAcceptDrop(const FDragDropEvent& DragDropEvent, EItemDropZone DropZone) override;
 	virtual FReply HandleAcceptDrop(FDragDropEvent const& DragDropEvent, EItemDropZone DropZone) override;
@@ -155,20 +155,17 @@ private:
 	FText RootText;
 };
 
-class FNamedSlotModel : public FHierarchyModel
+class FNamedSlotModelBase : public FHierarchyModel
 {
 public:
-	FNamedSlotModel(FWidgetReference InItem, FName InSlotName, TSharedPtr<FWidgetBlueprintEditor> InBlueprintEditor);
+	FNamedSlotModelBase(FName InSlotName, TSharedPtr<FWidgetBlueprintEditor> InBlueprintEditor);
 
-	virtual ~FNamedSlotModel() {}
-
-	virtual FName GetUniqueName() const override;
+	virtual ~FNamedSlotModelBase() {}
 
-	/* @returns the widget name to use for the tree item */
+	virtual FName GetUniqueName() const = 0;
 	virtual FText GetText() const override;
 
 	virtual const FSlateBrush* GetImage() const override;
-
 	virtual FSlateFontInfo GetFont() const override;
 
 	virtual void OnSelection() override;
@@ -177,18 +174,57 @@ public:
 	virtual FReply HandleAcceptDrop(FDragDropEvent const& DragDropEvent, EItemDropZone DropZone) override;
 
 protected:
+	virtual INamedSlotInterface* GetNamedSlotHost() const = 0;
+	virtual UWidget* GetNamedSlotHostWidget() const = 0;
+	
 	virtual void GetChildren(TArray< TSharedPtr<FHierarchyModel> >& Children) override;
 	virtual void UpdateSelection() override;
-	virtual FWidgetReference AsDraggedWidgetReference() const;
+	virtual FWidgetReference AsDraggedWidgetReference() const override;
 
-	void DoDrop(UWidget* NamedSlotHostWidget, UWidget* DroppingWidget);
+	virtual void DoDrop(INamedSlotInterface* NamedSlotHost, UWidget* DroppingWidget);
 
-private:
+protected:
 
-	FWidgetReference Item;
 	FName SlotName;
 };
 
+class FNamedSlotModel : public FNamedSlotModelBase
+{
+public:
+	FNamedSlotModel(FWidgetReference InItem, FName InSlotName, TSharedPtr<FWidgetBlueprintEditor> InBlueprintEditor);
+
+	virtual ~FNamedSlotModel() {}
+
+	virtual FName GetUniqueName() const override;
+	virtual INamedSlotInterface* GetNamedSlotHost() const override;
+	virtual UWidget* GetNamedSlotHostWidget() const override;
+
+protected:
+	virtual void OnSelection() override;
+
+protected:
+
+	FWidgetReference Item;
+};
+
+class FNamedSlotModelSubclass : public FNamedSlotModelBase
+{
+public:
+	FNamedSlotModelSubclass(UWidgetBlueprint* Blueprint, FName InSlotName, TSharedPtr<FWidgetBlueprintEditor> InBlueprintEditor);
+
+	virtual ~FNamedSlotModelSubclass() {}
+
+	virtual FName GetUniqueName() const override;
+	virtual INamedSlotInterface* GetNamedSlotHost() const override;
+	virtual UWidget* GetNamedSlotHostWidget() const override;
+
+protected:
+	virtual void OnSelection() override;
+
+protected:
+	TWeakObjectPtr<UWidgetBlueprint> Blueprint;
+};
+
 class FHierarchyWidget : public FHierarchyModel
 {
 public:
@@ -326,7 +362,7 @@ public:
 protected:
 	virtual void GetChildren(TArray< TSharedPtr<FHierarchyModel> >& Children) override;
 	virtual void UpdateSelection() override;
-	virtual FWidgetReference AsDraggedWidgetReference() const { return Item; }
+	virtual FWidgetReference AsDraggedWidgetReference() const override { return Item; }
 
 private:
 	FWidgetReference Item;
diff --git a/Engine/Source/Editor/UMGEditor/Private/WidgetBlueprint.cpp b/Engine/Source/Editor/UMGEditor/Private/WidgetBlueprint.cpp
--- a/Engine/Source/Editor/UMGEditor/Private/WidgetBlueprint.cpp
+++ b/Engine/Source/Editor/UMGEditor/Private/WidgetBlueprint.cpp
@@ -751,7 +751,7 @@ bool UWidgetBlueprint::DetectSlateWidgetLeaks(TArray<FText>& ValidationErrors)
 
 	// Update the widget tree directly to match the blueprint tree.  That way the preview can update
 	// without needing to do a full recompile.
-	TempUserWidget->DuplicateAndInitializeFromWidgetTree(WidgetTree);
+	TempUserWidget->DuplicateAndInitializeFromWidgetTree(WidgetTree, nullptr);
 
 	// We don't want this widget doing all the normal startup and acting like it's the real deal
 	// trying to do gameplay stuff, so make sure it's in design mode.
@@ -1330,6 +1330,19 @@ bool UWidgetBlueprint::ArePropertyBindingsAllowed() const
 	return GetDefault<UUMGEditorProjectSettings>()->CompilerOption_PropertyBindingRule(this) == EPropertyBindingPermissionLevel::Allow;
 }
 
+TArray<FName> UWidgetBlueprint::GetInheritedNamedSlots() const
+{
+	TArray<FName> SlotNames;
+	if ( UWidget* SuperClassCDO = GeneratedClass->GetSuperClass()->GetDefaultObject<UWidget>() )
+	{
+		if ( const INamedSlotInterface* NamedSlotHost = Cast<INamedSlotInterface>(SuperClassCDO) )
+		{
+			NamedSlotHost->GetSlotNames(SlotNames);
+		}
+	}
+	return SlotNames;
+}
+
 #if WITH_EDITOR
 void UWidgetBlueprint::LoadModulesRequiredForCompilation()
 {
diff --git a/Engine/Source/Editor/UMGEditor/Private/WidgetBlueprintCompiler.cpp b/Engine/Source/Editor/UMGEditor/Private/WidgetBlueprintCompiler.cpp
--- a/Engine/Source/Editor/UMGEditor/Private/WidgetBlueprintCompiler.cpp
+++ b/Engine/Source/Editor/UMGEditor/Private/WidgetBlueprintCompiler.cpp
@@ -436,6 +436,17 @@
 	// We need to save the widget tree to survive the initial sub-object clean blitz, 
 	// otherwise they all get renamed, and it causes early loading errors.
 	SubObjectsToSave.AddObject(WidgetBP->WidgetTree);
+
+	if (UUserWidget* ClassDefaultWidgetToClean = Cast<UUserWidget>(ClassToClean->ClassDefaultObject))
+	{
+		// We need preserve any named slots that have been slotted into the CDO.  This can happen when someone subclasses
+		// from a widget with named slots.  Those named slots are exposed to the child classes widget tree as
+		// containers they can slot stuff into.  Those widgets need to survive recompile.
+		for (FNamedSlotBinding& CDONamedSlotBinding : ClassDefaultWidgetToClean->NamedSlotBindings)
+		{
+			SubObjectsToSave.AddObject(CDONamedSlotBinding.Content);
+		}
+	}
 }
 
 void FWidgetBlueprintCompilerContext::CreateClassVariablesFromBlueprint()
@@ -904,6 +915,48 @@
 				}
 			}
 		}
+
+		// Do validation that as we subclass trees, we never stomp the slotted content of a parent widget.
+		// doing that is not valid, as it would invalidate variables that were set?  This check could be
+		// made more complex to only worry about cases with variables being generated, but that's a whole lot
+		// extra, so for now lets just limit it to be safe.
+		{
+			TMap<FName, UWidget*> NamedSlotContentMap;
+			// Make sure that we don't have dueling widget hierarchies
+			UWidgetBlueprintGeneratedClass* NamedSlotClass = BPGClass;
+			while (NamedSlotClass)
+			{
+				UWidgetTree* Tree = NamedSlotClass->GetWidgetTreeArchetype();
+			
+				TArray<FName> SlotNames;
+				Tree->GetSlotNames(SlotNames);
+
+				for (FName SlotName : SlotNames)
+				{
+					if (UWidget* ContentInSlot = Tree->GetContentForSlot(SlotName))
+					{
+						if (!NamedSlotContentMap.Contains(SlotName))
+						{
+							NamedSlotContentMap.Add(SlotName, ContentInSlot);
+						}
+						else
+						{
+							UClass* SubClassWithSlotFilled = ContentInSlot->GetTypedOuter<UClass>();
+							UClass* ParentClassWithSlotFilled = NamedSlotClass;
+							MessageLog.Note(
+								*FText::Format(
+									LOCTEXT("ParentAndChildBothHaveWidgetTrees", "The Named Slot '{0}' already contains @@ from the class @@ but the subclass @@ tried to slot @@ into it."),
+									FText::FromName(SlotName)
+								).ToString(),
+								ContentInSlot, ParentClassWithSlotFilled,
+								SubClassWithSlotFilled, NamedSlotContentMap.FindRef(SlotName));
+						}
+					}
+				}
+			
+				NamedSlotClass = Cast<UWidgetBlueprintGeneratedClass>(NamedSlotClass->GetSuperClass());
+			}
+		}
 	}
 
 	if (bIsSkeletonOnly || WidgetBP->SkeletonGeneratedClass != Class)
diff --git a/Engine/Source/Editor/UMGEditor/Private/WidgetBlueprintEditor.cpp b/Engine/Source/Editor/UMGEditor/Private/WidgetBlueprintEditor.cpp
--- a/Engine/Source/Editor/UMGEditor/Private/WidgetBlueprintEditor.cpp
+++ b/Engine/Source/Editor/UMGEditor/Private/WidgetBlueprintEditor.cpp
@@ -653,7 +653,10 @@ void FWidgetBlueprintEditor::SetSelectedNamedSlot(TOptional<FNamedSlotSelection>
 	SelectedNamedSlot = InSelectedNamedSlot;
 	if (InSelectedNamedSlot.IsSet())
 	{
-		SelectedWidgets.Add(InSelectedNamedSlot->NamedSlotHostWidget);
+		if (InSelectedNamedSlot->NamedSlotHostWidget.IsValid())
+		{
+			SelectedWidgets.Add(InSelectedNamedSlot->NamedSlotHostWidget);
+		}
 	}
 
 	OnSelectedWidgetsChanged.Broadcast();
@@ -1609,7 +1612,7 @@ void FWidgetBlueprintEditor::UpdatePreview(UBlueprint* InBlueprint, bool bInForc
 
 			// Update the widget tree directly to match the blueprint tree.  That way the preview can update
 			// without needing to do a full recompile.
-			PreviewUserWidget->DuplicateAndInitializeFromWidgetTree(LatestWidgetTree);
+			PreviewUserWidget->DuplicateAndInitializeFromWidgetTree(LatestWidgetTree, LatestWidgetTree != PreviewBlueprint->WidgetTree ? PreviewBlueprint->WidgetTree : nullptr);
 
 			// Establish the widget as being in design time before initializing (so that IsDesignTime is reliable within Initialize)
             // We have to call it to make sure that all the WidgetTree had the DesignerFlags set correctly
diff --git a/Engine/Source/Editor/UMGEditor/Private/WidgetBlueprintEditorUtils.cpp b/Engine/Source/Editor/UMGEditor/Private/WidgetBlueprintEditorUtils.cpp
--- a/Engine/Source/Editor/UMGEditor/Private/WidgetBlueprintEditorUtils.cpp
+++ b/Engine/Source/Editor/UMGEditor/Private/WidgetBlueprintEditorUtils.cpp
@@ -1760,11 +1760,10 @@ TArray<UWidget*> FWidgetBlueprintEditorUtils::PasteWidgetsInternal(TSharedRef<FW
 			return TArray<UWidget*>();
 		}
 
-		UWidget* NamedSlotHostWidget = ParentWidgetRef.GetTemplate();
-
 		BP->WidgetTree->Modify();
 
-		if (NamedSlotHostWidget)
+		// If there's a ParentWidgetRef, then we're pasting into a named slot of a widget in the tree.
+		if (UWidget* NamedSlotHostWidget = ParentWidgetRef.GetTemplate())
 		{
 			NamedSlotHostWidget->SetFlags(RF_Transactional);
 			NamedSlotHostWidget->Modify();
@@ -1772,6 +1771,15 @@ TArray<UWidget*> FWidgetBlueprintEditorUtils::PasteWidgetsInternal(TSharedRef<FW
 			INamedSlotInterface* NamedSlotInterface = Cast<INamedSlotInterface>(NamedSlotHostWidget);
 			NamedSlotInterface->SetContentForSlot(SlotName, RootPasteWidgets[0]);
 		}
+		else
+		{
+			// If there's no ParentWidgetRef then we're pasting into the exposed named slots of the widget tree.
+			// these are the slots that our parent class is exposing for use externally, but we can also override
+			// them as a subclass.
+
+			BP->WidgetTree->Modify();
+			BP->WidgetTree->SetContentForSlot(SlotName, RootPasteWidgets[0]);
+		}
 
 		FBlueprintEditorUtils::MarkBlueprintAsStructurallyModified(BP);
 	}
diff --git a/Engine/Source/Editor/UMGEditor/Private/WidgetReference.cpp b/Engine/Source/Editor/UMGEditor/Private/WidgetReference.cpp
--- a/Engine/Source/Editor/UMGEditor/Private/WidgetReference.cpp
+++ b/Engine/Source/Editor/UMGEditor/Private/WidgetReference.cpp
@@ -51,9 +51,7 @@ UWidget* FWidgetReference::GetPreview() const
 {
 	if ( WidgetEditor.IsValid() && TemplateHandle.IsValid() )
 	{
-		UUserWidget* PreviewRoot = WidgetEditor.Pin()->GetPreview();
-
-		if ( PreviewRoot )
+		if ( UUserWidget* PreviewRoot = WidgetEditor.Pin()->GetPreview() )
 		{
 			if ( UWidget* TemplateWidget = TemplateHandle->Widget.Get() )
 			{
diff --git a/Engine/Source/Editor/UMGEditor/Public/WidgetBlueprint.h b/Engine/Source/Editor/UMGEditor/Public/WidgetBlueprint.h
--- a/Engine/Source/Editor/UMGEditor/Public/WidgetBlueprint.h
+++ b/Engine/Source/Editor/UMGEditor/Public/WidgetBlueprint.h
@@ -322,6 +322,9 @@ public:
 
 	bool ArePropertyBindingsAllowed() const;
 
+	/** Gets any named slots exposed by the parent generated class that can be slotted into by the subclass. */
+	TArray<FName> GetInheritedNamedSlots() const;
+
 protected:
 	virtual void LoadModulesRequiredForCompilation() override;
 
diff --git a/Engine/Source/Runtime/UMG/Private/UserWidget.cpp b/Engine/Source/Runtime/UMG/Private/UserWidget.cpp
--- a/Engine/Source/Runtime/UMG/Private/UserWidget.cpp
+++ b/Engine/Source/Runtime/UMG/Private/UserWidget.cpp
@@ -125,11 +125,6 @@
 		}
 
 		UWidgetBlueprintGeneratedClass* BGClass = Cast<UWidgetBlueprintGeneratedClass>(GetClass());
-		if (BGClass)
-		{
-			BGClass = GetWidgetTreeOwningClass();
-		}
-
 		// Only do this if this widget is of a blueprint class
 		if (BGClass)
 		{
@@ -148,8 +143,7 @@
 		{
 			WidgetTree->SetFlags(RF_Transient);
 
-			const bool bReparentToWidgetTree = false;
-			InitializeNamedSlots(bReparentToWidgetTree);
+			InitializeNamedSlots();
 		}
 
 		if (!IsDesignTime() && PlayerContext.IsValid())
@@ -164,7 +158,7 @@
 	return false;
 }
 
-void UUserWidget::InitializeNamedSlots(bool bReparentToWidgetTree)
+void UUserWidget::InitializeNamedSlots()
 {
 	for (const FNamedSlotBinding& Binding : NamedSlotBindings )
 	{
@@ -182,19 +176,13 @@
 				{
 					NamedSlot->ClearChildren();
 					NamedSlot->AddChild(BindingContent);
-
-					//if ( bReparentToWidgetTree )
-					//{
-					//	FName NewName = MakeUniqueObjectName(WidgetTree, BindingContent->GetClass(), BindingContent->GetFName());
-					//	BindingContent->Rename(*NewName.ToString(), WidgetTree, REN_DontCreateRedirectors | REN_DoNotDirty);
-					//}
 				}
 			}
 		}
 	}
 }
 
-void UUserWidget::DuplicateAndInitializeFromWidgetTree(UWidgetTree* InWidgetTree)
+void UUserWidget::DuplicateAndInitializeFromWidgetTree(UWidgetTree* InWidgetTree, INamedSlotInterface* NamedSlotsToMerge)
 {
 	TScopeCounter<uint32> ScopeInitializingFromWidgetTree(bInitializingFromWidgetTree);
 
@@ -221,6 +209,39 @@
 				InstancedSubUserWidget->Initialize();
 			}
 		});
+
+		// This block controls merging named slot content specified in a child class for the widget we're templated after.
+		if (NamedSlotsToMerge)
+		{
+			TArray<FName> SlotNames;
+			NamedSlotsToMerge->GetSlotNames(SlotNames);
+			for (FName SlotName : SlotNames)
+			{
+				// Don't insert the named slot content if the named slot is filled already.  This is a problematic
+				// scenario though, if someone inserted content, but we have class default instances, we sorta leave
+				// ourselves in a strange situation, because there are now potentially class variables that won't
+				// have an instance assigned.
+				if (!GetContentForSlot(SlotName))
+				{
+					if (UWidget* TemplateSlotContent = NamedSlotsToMerge->GetContentForSlot(SlotName))
+					{
+						FObjectInstancingGraph NamedSlotInstancingGraph;
+						// We need to add a mapping from the template's widget tree to the new widget tree, that way
+						// as we instance the widget hierarchy it's grafted onto the new widget tree.
+						NamedSlotInstancingGraph.AddNewObject(WidgetTree, TemplateSlotContent->GetTypedOuter<UWidgetTree>());
+
+						// Instance the new widget from the foreign tree, but do it in a way that grafts it onto the tree we're instancing.
+						UWidget* Content = NewObject<UWidget>(WidgetTree, TemplateSlotContent->GetClass(), TemplateSlotContent->GetFName(), RF_Transactional, TemplateSlotContent, false, &NamedSlotInstancingGraph);
+						Content->SetFlags(RF_Transient | RF_DuplicateTransient);
+
+						// Insert the newly constructed widget into the named slot that corresponds.  The above creates
+						// it as if it was always part of the widget tree, but this actually puts it into a widget's
+						// slot for the named slot.
+						SetContentForSlot(SlotName, Content);
+					}
+				}
+			}
+		}
 	}
 }
 
@@ -880,8 +901,9 @@
 	// Dynamically insert the new widget into the hierarchy if it exists.
 	if ( WidgetTree )
 	{
-		UNamedSlot* NamedSlot = Cast<UNamedSlot>(WidgetTree->FindWidget(SlotName));
-		if ( NamedSlot )
+		ensureMsgf(!HasAnyFlags(RF_ClassDefaultObject), TEXT("The Widget CDO is not expected to ever have a valid widget tree."));
+		
+		if ( UNamedSlot* NamedSlot = Cast<UNamedSlot>(WidgetTree->FindWidget(SlotName)))
 		{
 			NamedSlot->ClearChildren();
 
diff --git a/Engine/Source/Runtime/UMG/Private/WidgetBlueprintGeneratedClass.cpp b/Engine/Source/Runtime/UMG/Private/WidgetBlueprintGeneratedClass.cpp
--- a/Engine/Source/Runtime/UMG/Private/WidgetBlueprintGeneratedClass.cpp
+++ b/Engine/Source/Runtime/UMG/Private/WidgetBlueprintGeneratedClass.cpp
@@ -179,6 +179,7 @@ void UWidgetBlueprintGeneratedClass::InitializeBindingsStatic(UUserWidget* UserW
 void UWidgetBlueprintGeneratedClass::InitializeWidgetStatic(UUserWidget* UserWidget
 	, const UClass* InClass
 	, UWidgetTree* InWidgetTree
+	, const UClass* InWidgetTreeWidgetClass
 	, const TArrayView<UWidgetAnimation*> InAnimations
 	, const TArrayView<const FDelegateRuntimeBinding> InBindings)
 {
@@ -194,17 +195,26 @@ void UWidgetBlueprintGeneratedClass::InitializeWidgetStatic(UUserWidget* UserWid
 	UserWidget->WidgetGeneratedByClass = WidgetGeneratedByClass;
 #endif
 
-	UWidgetTree* ClonedTree = UserWidget->WidgetTree;
+	UWidgetTree* CreatedWidgetTree = UserWidget->WidgetTree;
 
 	// Normally the ClonedTree should be null - we do in the case of design time with the widget, actually
 	// clone the widget tree directly from the WidgetBlueprint so that the rebuilt preview matches the newest
 	// widget tree, without a full blueprint compile being required.  In that case, the WidgetTree on the UserWidget
 	// will have already been initialized to some value.  When that's the case, we'll avoid duplicating it from the class
 	// similar to how we use to use the DesignerWidgetTree.
-	if ( ClonedTree == nullptr )
+	if ( CreatedWidgetTree == nullptr )
 	{
-		UserWidget->DuplicateAndInitializeFromWidgetTree(InWidgetTree);
-		ClonedTree = UserWidget->WidgetTree;
+		UWidgetTree* MergeNamedSlotTree = nullptr; 
+		if (UserWidget->GetClass() != InWidgetTreeWidgetClass)
+		{
+			if (UWidgetBlueprintGeneratedClass* WidgetsActualClass = Cast<UWidgetBlueprintGeneratedClass>(UserWidget->GetClass()))
+			{
+				MergeNamedSlotTree = WidgetsActualClass->GetWidgetTreeArchetype();
+			}
+		}
+		
+		UserWidget->DuplicateAndInitializeFromWidgetTree(InWidgetTree, MergeNamedSlotTree);
+		CreatedWidgetTree = UserWidget->WidgetTree;
 	}
 
 #if !WITH_EDITOR && UE_BUILD_DEBUG
@@ -215,7 +225,7 @@ void UWidgetBlueprintGeneratedClass::InitializeWidgetStatic(UUserWidget* UserWid
 	UserWidget->WidgetGeneratedBy = InClass->ClassGeneratedBy;
 #endif
 
-	if (ClonedTree)
+	if (CreatedWidgetTree)
 	{
 		UClass* WidgetBlueprintClass = UserWidget->GetClass();
 
@@ -229,7 +239,7 @@ void UWidgetBlueprintGeneratedClass::InitializeWidgetStatic(UUserWidget* UserWid
 
 		BindAnimationsStatic(UserWidget, InAnimations, ObjectPropertiesMap);
 
-		ClonedTree->ForEachWidget([&](UWidget* Widget) {
+		CreatedWidgetTree->ForEachWidget([&](UWidget* Widget) {
 			// Not fatal if NULL, but shouldn't happen
 			if (!ensure(Widget != nullptr))
 			{
@@ -257,7 +267,7 @@ void UWidgetBlueprintGeneratedClass::InitializeWidgetStatic(UUserWidget* UserWid
 			// Initialize Navigation Data
 			if (Widget->Navigation)
 			{
-				Widget->Navigation->ResolveRules(UserWidget, ClonedTree);
+				Widget->Navigation->ResolveRules(UserWidget, CreatedWidgetTree);
 			}
 
 #if WITH_EDITOR
@@ -307,14 +317,9 @@ void UWidgetBlueprintGeneratedClass::InitializeWidget(UUserWidget* UserWidget) c
 	TArray<UWidgetAnimation*, TMemStackAllocator<>> AllAnims;
 	TArray<FDelegateRuntimeBinding, TMemStackAllocator<>> AllBindings;
 
-	// Include current class animations.
-	AllAnims.Append(Animations);
-
-	// Include current class bindings.
-	AllBindings.Append(Bindings);
-
-	// Iterate all generated classes in the widget's parent class hierarchy and include animations and bindings found on each one.
-	UClass* SuperClass = GetSuperClass();
+	// Iterate all generated classes in the widget's parent class hierarchy and include animations and bindings
+	// found on each one.
+	UClass* SuperClass = UserWidget->GetClass();
 	while (UWidgetBlueprintGeneratedClass* WBPGC = Cast<UWidgetBlueprintGeneratedClass>(SuperClass))
 	{
 		AllAnims.Append(WBPGC->Animations);
@@ -323,7 +328,14 @@ void UWidgetBlueprintGeneratedClass::InitializeWidget(UUserWidget* UserWidget) c
 		SuperClass = SuperClass->GetSuperClass();
 	}
 
-	InitializeWidgetStatic(UserWidget, this, WidgetTree, AllAnims, AllBindings);
+	UWidgetTree* PrimaryWidgetTree = WidgetTree;
+	UWidgetBlueprintGeneratedClass* PrimaryWidgetTreeClass = FindWidgetTreeOwningClass();
+	if (PrimaryWidgetTreeClass)
+	{
+		PrimaryWidgetTree = PrimaryWidgetTreeClass->WidgetTree;
+	}
+
+	InitializeWidgetStatic(UserWidget, this, PrimaryWidgetTree, PrimaryWidgetTreeClass, AllAnims, AllBindings);
 }
 
 void UWidgetBlueprintGeneratedClass::PostLoad()
@@ -414,9 +426,9 @@ void UWidgetBlueprintGeneratedClass::Serialize(FArchive& Ar)
 	Ar.UsingCustomVersion(FEditorObjectVersion::GUID);
 }
 
-UWidgetBlueprintGeneratedClass* UWidgetBlueprintGeneratedClass::FindWidgetTreeOwningClass()
+UWidgetBlueprintGeneratedClass* UWidgetBlueprintGeneratedClass::FindWidgetTreeOwningClass() const
 {
-	UWidgetBlueprintGeneratedClass* RootBGClass = this;
+	UWidgetBlueprintGeneratedClass* RootBGClass = const_cast<UWidgetBlueprintGeneratedClass*>(this);
 	UWidgetBlueprintGeneratedClass* BGClass = RootBGClass;
 
 	while (BGClass)
diff --git a/Engine/Source/Runtime/UMG/Private/WidgetTree.cpp b/Engine/Source/Runtime/UMG/Private/WidgetTree.cpp
--- a/Engine/Source/Runtime/UMG/Private/WidgetTree.cpp
+++ b/Engine/Source/Runtime/UMG/Private/WidgetTree.cpp
@@ -138,9 +138,21 @@ bool UWidgetTree::RemoveWidget(UWidget* InRemovedWidget)
 	// If the widget being removed is the root, null it out.
 	else if ( InRemovedWidget == RootWidget )
 	{
-		RootWidget = NULL;
+		RootWidget = nullptr;
 		bRemoved = true;
 	}
+	else
+	{
+		for (const auto& KVP : NamedSlotBindings)
+		{
+			if (KVP.Value == InRemovedWidget)
+			{
+				bRemoved = true;
+				SetContentForSlot(KVP.Key, nullptr);
+				break;
+			}
+		}
+	}
 
 	return bRemoved;
 }
@@ -176,34 +188,45 @@ void UWidgetTree::GetChildWidgets(UWidget* Parent, TArray<UWidget*>& Widgets)
 
 void UWidgetTree::ForEachWidget(TFunctionRef<void(UWidget*)> Predicate) const
 {
+	// Start with the root widget.
 	if ( RootWidget )
 	{
 		Predicate(RootWidget);
 
 		ForWidgetAndChildren(RootWidget, Predicate);
 	}
+
+	// Next, check the top level named slots.  Once the UserWidget that owns this widget tree is initialized, these
+	// named slot bindings will all be emptied out, and their relevant named slot content will be inserted into the
+	// hierarchy, but before that happens, or during the design time template, we need to treat them as separate
+	// top level widgets because they're not yet in the hierarchy.
+	for (const auto& NamedSlotsEntry : NamedSlotBindings)
+	{
+		if (UWidget* NamedSlotContent = NamedSlotsEntry.Value)
+		{
+			Predicate(NamedSlotContent);
+			ForWidgetAndChildren(NamedSlotContent, Predicate);
+		}
+	}
 }
 
 void UWidgetTree::ForEachWidgetAndDescendants(TFunctionRef<void(UWidget*)> Predicate) const
 {
-	if ( RootWidget )
-	{
-		Predicate(RootWidget);
-
-		ForWidgetAndChildren(RootWidget, [&Predicate] (UWidget* Child) {
-			if ( UUserWidget* UserWidgetChild = Cast<UUserWidget>(Child) )
+	ForEachWidget([this, &Predicate] (UWidget* Widget) {
+		if ( Widget != RootWidget )
+		{
+			if (const UUserWidget* UserWidgetChild = Cast<UUserWidget>(Widget))
 			{
 				if ( UserWidgetChild->WidgetTree )
 				{
 					UserWidgetChild->WidgetTree->ForEachWidgetAndDescendants(Predicate);
+					return;
 				}
 			}
-			else
-			{
-				Predicate(Child);
-			}
-		});
-	}
+		}
+
+		Predicate(Widget);
+	});
 }
 
 void UWidgetTree::ForWidgetAndChildren(UWidget* Widget, TFunctionRef<void(UWidget*)> Predicate)
@@ -240,13 +263,35 @@ void UWidgetTree::ForWidgetAndChildren(UWidget* Widget, TFunctionRef<void(UWidge
 	}
 }
 
-void UWidgetTree::PreSave(const class ITargetPlatform* TargetPlatform)
+// INamedSlotInterface
+//----------------------------------------------------------------------------------------
+
+void UWidgetTree::GetSlotNames(TArray<FName>& SlotNames) const
+{
+	// Widget trees don't know for sure how many slots the real widget tree has, they can only tell you about the
+	// slots that have content in them.
+	NamedSlotBindings.GetKeys(SlotNames);
+}
+
+UWidget* UWidgetTree::GetContentForSlot(FName SlotName) const
 {
-	PRAGMA_DISABLE_DEPRECATION_WARNINGS;
-	Super::PreSave(TargetPlatform);
-	PRAGMA_ENABLE_DEPRECATION_WARNINGS;
+	return NamedSlotBindings.FindRef(SlotName);
 }
 
+void UWidgetTree::SetContentForSlot(FName SlotName, UWidget* Content)
+{
+	if (Content)
+	{
+		NamedSlotBindings.Add(SlotName, Content);	
+	}
+	else
+	{
+		NamedSlotBindings.Remove(SlotName);	
+	}
+}
+
+//----------------------------------------------------------------------------------------
+
 void UWidgetTree::PreSave(FObjectPreSaveContext ObjectSaveContext)
 {
 	if (!ObjectSaveContext.IsCooking()) // AllWidgets is not used in cooked builds
diff --git a/Engine/Source/Runtime/UMG/Public/Blueprint/UserWidget.h b/Engine/Source/Runtime/UMG/Public/Blueprint/UserWidget.h
--- a/Engine/Source/Runtime/UMG/Public/Blueprint/UserWidget.h
+++ b/Engine/Source/Runtime/UMG/Public/Blueprint/UserWidget.h
@@ -219,7 +219,7 @@
 	virtual void PostLoad() override;
 	//~ End UObject Interface
 
-	void DuplicateAndInitializeFromWidgetTree(UWidgetTree* InWidgetTree);
+	void DuplicateAndInitializeFromWidgetTree(UWidgetTree* InWidgetTree, INamedSlotInterface* NamedSlotsToMerge);
 
 	virtual bool Initialize();
 
@@ -239,7 +239,7 @@
 	/** The function is implemented only in nativized widgets (automatically converted from BP to c++) */
 	virtual void InitializeNativeClassData() {}
 
-	void InitializeNamedSlots(bool bReparentToWidgetTree);
+	void InitializeNamedSlots();
 
 public:
 	//~ Begin UVisual interface
diff --git a/Engine/Source/Runtime/UMG/Public/Blueprint/WidgetBlueprintGeneratedClass.h b/Engine/Source/Runtime/UMG/Public/Blueprint/WidgetBlueprintGeneratedClass.h
--- a/Engine/Source/Runtime/UMG/Public/Blueprint/WidgetBlueprintGeneratedClass.h
+++ b/Engine/Source/Runtime/UMG/Public/Blueprint/WidgetBlueprintGeneratedClass.h
@@ -105,7 +105,7 @@ public:
 	void SetWidgetTreeArchetype(UWidgetTree* InWidgetTree);
 
 	// Walks up the hierarchy looking for a valid widget tree.
-	UWidgetBlueprintGeneratedClass* FindWidgetTreeOwningClass();
+	UWidgetBlueprintGeneratedClass* FindWidgetTreeOwningClass() const;
 
 	// Execute the callback for every FieldId defined in the BP class
 	void ForEachField(TFunctionRef<bool(::UE::FieldNotification::FFieldId FielId)> Callback) const;
@@ -131,6 +131,7 @@ public:
 	static void InitializeWidgetStatic(UUserWidget* UserWidget
 		, const UClass* InClass
 		, UWidgetTree* InWidgetTree
+		, const UClass* InWidgetTreeWidgetClass
 		, const TArrayView<UWidgetAnimation*> InAnimations
 		, const TArrayView<const FDelegateRuntimeBinding> InBindings);
 
diff --git a/Engine/Source/Runtime/UMG/Public/Blueprint/WidgetTree.h b/Engine/Source/Runtime/UMG/Public/Blueprint/WidgetTree.h
--- a/Engine/Source/Runtime/UMG/Public/Blueprint/WidgetTree.h
+++ b/Engine/Source/Runtime/UMG/Public/Blueprint/WidgetTree.h
@@ -16,7 +16,7 @@
 
 /** The widget tree manages the collection of widgets in a blueprint widget. */
 UCLASS()
-class UMG_API UWidgetTree : public UObject
+class UMG_API UWidgetTree : public UObject, public INamedSlotInterface
 {
 	GENERATED_UCLASS_BODY()
 
@@ -100,11 +100,21 @@ public:
 		return NewObject<WidgetT>(this, WidgetClass, WidgetName, RF_Transactional);
 	}
 
+	// INamedSlotInterface
+	//----------------------------------------------------------------------------------------
+	
+	/** Gets the names for slots that we can store widgets into. */
+	virtual void GetSlotNames(TArray<FName>& SlotNames) const override;
+
+	/** Gets the widget for a given slot by name, will return nullptr if no widget is in the slot. */
+	virtual UWidget* GetContentForSlot(FName SlotName) const override;
+
+	/** Sets the widget for a given slot by name. */
+	virtual void SetContentForSlot(FName SlotName, UWidget* Content) override;
+
+	//----------------------------------------------------------------------------------------
+
 	// UObject interface
-	PRAGMA_DISABLE_DEPRECATION_WARNINGS // Suppress compiler warning on override of deprecated function
-	UE_DEPRECATED(5.0, "Use version that takes FObjectPreSaveContext instead.")
-	virtual void PreSave(const class ITargetPlatform* TargetPlatform) override;
-	PRAGMA_ENABLE_DEPRECATION_WARNINGS
 	virtual void PreSave(FObjectPreSaveContext ObjectSaveContext) override;
 	virtual void PostLoad() override;
 	// End of UObject interface
@@ -114,6 +124,14 @@ public:
 	UPROPERTY(Instanced)
 	TObjectPtr<UWidget> RootWidget;
 
+	/**
+	 * Stores the widgets being assigned to named slots, these widgets will be slotted into the named slots of the
+	 * user widget that owns this widget tree after the user widget is constructed.  This is how we store the
+	 * template content in order to have named slot inheritance, and merging widget trees.
+	 */
+	UPROPERTY()
+	TMap<FName, TObjectPtr<UWidget>> NamedSlotBindings;
+
 protected:
 
 #if WITH_EDITORONLY_DATA
