Fix for Static Analysis warning
#p4v-preflight-copy 20391532

diff --git a/Engine/Source/Editor/UMGEditor/Private/WidgetBlueprintCompiler.cpp b/Engine/Source/Editor/UMGEditor/Private/WidgetBlueprintCompiler.cpp
--- a/Engine/Source/Editor/UMGEditor/Private/WidgetBlueprintCompiler.cpp
+++ b/Engine/Source/Editor/UMGEditor/Private/WidgetBlueprintCompiler.cpp
@@ -741,11 +741,26 @@ void FWidgetBlueprintCompilerContext::FixAbandonedWidgetTree(UWidgetBlueprint* W
 
 void FWidgetBlueprintCompilerContext::FinishCompilingClass(UClass* Class)
 {
+	if (Class == nullptr)
+		return;
+
 	UWidgetBlueprint* WidgetBP = WidgetBlueprint();
-	UWidgetBlueprintGeneratedClass* BPGClass = CastChecked<UWidgetBlueprintGeneratedClass>(Class);
+
+	if (WidgetBP == nullptr)
+		return;
+
 	UClass* ParentClass = WidgetBP->ParentClass;
+
+	if (ParentClass == nullptr)
+		return;
+	
 	const bool bIsSkeletonOnly = CompileOptions.CompileType == EKismetCompileType::SkeletonOnly;
 
+	UWidgetBlueprintGeneratedClass* BPGClass = CastChecked<UWidgetBlueprintGeneratedClass>(Class);
+
+	if (BPGClass == nullptr)
+		return;
+
 	// Don't do a bunch of extra work on the skeleton generated class.
 	if ( !bIsSkeletonOnly )
 	{
