Localization Duplicate Id fix.

diff --git a/Engine/Source/Editor/UMGEditor/Private/WidgetBlueprintCompiler.cpp b/Engine/Source/Editor/UMGEditor/Private/WidgetBlueprintCompiler.cpp
--- a/Engine/Source/Editor/UMGEditor/Private/WidgetBlueprintCompiler.cpp
+++ b/Engine/Source/Editor/UMGEditor/Private/WidgetBlueprintCompiler.cpp
@@ -960,7 +960,7 @@ void FWidgetBlueprintCompilerContext::FinishCompilingClass(UClass* Class)
 							UClass* ParentClassWithSlotFilled = NamedSlotClass;
 							MessageLog.Note(
 								*FText::Format(
-									LOCTEXT("ParentAndChildBothHaveWidgetTrees", "The Named Slot '{0}' already contains @@ from the class @@ but the subclass @@ tried to slot @@ into it."),
+									LOCTEXT("NamedSlotAlreadyFilled", "The Named Slot '{0}' already contains @@ from the class @@ but the subclass @@ tried to slot @@ into it."),
 									FText::FromName(SlotName)
 								).ToString(),
 								ContentInSlot, ParentClassWithSlotFilled,
