UMG Designer: Prevent crash accessing a nullptr when pasting Widget at the root of a UserWidget.

diff --git a/Engine/Source/Editor/UMGEditor/Private/WidgetBlueprintEditorUtils.cpp b/Engine/Source/Editor/UMGEditor/Private/WidgetBlueprintEditorUtils.cpp
--- a/Engine/Source/Editor/UMGEditor/Private/WidgetBlueprintEditorUtils.cpp
+++ b/Engine/Source/Editor/UMGEditor/Private/WidgetBlueprintEditorUtils.cpp
@@ -1603,9 +1603,14 @@ TArray<UWidget*> FWidgetBlueprintEditorUtils::PasteWidgetsInternal(TSharedRef<FW
 			// If the widget isn't a panel or we just really want it to be a sibling (ie. when duplicating), we'll try it's parent to see if the pasted widget can be a sibling (and get its index to insert at)
 			if ( bForceSibling || !ParentWidget )
 			{
-				UWidget* WidgetTemplate = ParentWidgetRef.GetTemplate();
-				ParentWidget = WidgetTemplate->GetParent();
-				IndexToInsert = ParentWidget->GetChildIndex(WidgetTemplate) + 1;
+				if (UWidget* WidgetTemplate = ParentWidgetRef.GetTemplate())
+				{
+					ParentWidget = WidgetTemplate->GetParent();
+					if (ParentWidget)
+					{
+						IndexToInsert = ParentWidget->GetChildIndex(WidgetTemplate) + 1;
+					}
+				}
 			}
 		}
 
