UE-154770 - Saving Widget Blueprints directly can put them in a bad state.

The fixup code for delegates only updated the pins on the result node. However, we also need to fixup the user pins. Otherwise, reallocation of pins can potentially lose the original double->float conversion. This leaves the BP in a broken state since it's linker version is up to date, but it still has pins that need a double->float conversion.

diff --git a/Engine/Source/Editor/UMGEditor/Private/WidgetBlueprint.cpp b/Engine/Source/Editor/UMGEditor/Private/WidgetBlueprint.cpp
--- a/Engine/Source/Editor/UMGEditor/Private/WidgetBlueprint.cpp
+++ b/Engine/Source/Editor/UMGEditor/Private/WidgetBlueprint.cpp
@@ -657,15 +657,32 @@ void UWidgetBlueprint::ReplaceDeprecatedNodes()
 												return bHasMatch;
 											};
 
+											bool bFoundMatchingParam = false;
 											for (TFieldIterator<FFloatProperty> It(DelegateFunction); It; ++It)
 											{
 												if (OutputParameterMatchesPin(*It))
 												{
 													Pin->PinType.PinSubCategory = UEdGraphSchema_K2::PC_Float;
+													bFoundMatchingParam = true;
 													break;
 												}
 											}
 
+											if (bFoundMatchingParam)
+											{
+												UK2Node_FunctionResult* FunctionResultNode = CastChecked<UK2Node_FunctionResult>(Node);
+												for (TSharedPtr<FUserPinInfo>& UserPin : FunctionResultNode->UserDefinedPins)
+												{
+													check(UserPin);
+													if (UserPin->PinName == PinName)
+													{
+														check(UserPin->PinType.PinCategory == UEdGraphSchema_K2::PC_Real);
+														UserPin->PinType.PinSubCategory = UEdGraphSchema_K2::PC_Float;
+														break;
+													}
+												}
+											}
+
 											break;
 										}
 									}
