UMG - Additional improvements to the named slot inheritance system.  Named slots use to also allow you to include content as 'example' or 'seed' data whenever you dropped the widgets into a parent.  That behavior was and is probably a little sketchy, I don't believe it saw much usage.  However, in order to maintain backwards compatability, version bumping the introduction of the inheritance, and introudcing a field to named slots, called, bExposeOnInstanceOnly.  This will be checked for any upraded named slot, effectively disabling named slot inheritance for that slot.  By default this will be false, so going forward inheritance is the preferred setup methodology.

diff --git a/Engine/Source/Editor/UMGEditor/Private/WidgetBlueprintCompiler.cpp b/Engine/Source/Editor/UMGEditor/Private/WidgetBlueprintCompiler.cpp
--- a/Engine/Source/Editor/UMGEditor/Private/WidgetBlueprintCompiler.cpp
+++ b/Engine/Source/Editor/UMGEditor/Private/WidgetBlueprintCompiler.cpp
@@ -25,6 +25,7 @@
 #include "WidgetBlueprintExtension.h"
 #include "Editor/WidgetCompilerLog.h"
 #include "Editor.h"
+#include "Algo/RemoveIf.h"
 
 #define LOCTEXT_NAMESPACE "UMG"
 
@@ -906,13 +907,19 @@ void FWidgetBlueprintCompilerContext::FinishCompilingClass(UClass* Class)
 		// Add all the names of the named slot widgets to the slot names structure.
 		{
 			BPGClass->NamedSlots.Reset();
+			BPGClass->InstanceNamedSlots.Reset();
 			UWidgetBlueprint* WidgetBPIt = WidgetBP;
 			while (WidgetBPIt)
 			{
-				WidgetBPIt->ForEachSourceWidget([&] (UWidget* Widget) {
-					if ( Widget && Widget->IsA<UNamedSlot>() )
+				WidgetBPIt->ForEachSourceWidget([&] (const UWidget* Widget) {
+					if (const UNamedSlot* NamedSlot = Cast<UNamedSlot>(Widget))
 					{
 						BPGClass->NamedSlots.Add(Widget->GetFName());
+
+						if (NamedSlot->bExposeOnInstanceOnly)
+						{
+							BPGClass->InstanceNamedSlots.Add(Widget->GetFName());
+						}
 					}
 				});
 				
@@ -927,6 +934,18 @@ void FWidgetBlueprintCompilerContext::FinishCompilingClass(UClass* Class)
 				// If we find content for this slot, remove it from the available set.
 				BPGClass->AvailableNamedSlots.Remove(SlotName);
 			});
+
+			// Remove any available subclass named slots that are marked as instance named slot.
+			for (const FName& InstanceNamedSlot : BPGClass->InstanceNamedSlots)
+			{
+				BPGClass->AvailableNamedSlots.Remove(InstanceNamedSlot);
+			}
+
+			// Now add any available named slot that doesn't have anything in it also.
+			for (const FName& AvailableNamedSlot : BPGClass->AvailableNamedSlots)
+			{
+				BPGClass->InstanceNamedSlots.AddUnique(AvailableNamedSlot);
+			}
 		}
 
 		// Make sure that we don't have dueling widget hierarchies
diff --git a/Engine/Source/Runtime/Core/Public/UObject/FortniteMainBranchObjectVersion.h b/Engine/Source/Runtime/Core/Public/UObject/FortniteMainBranchObjectVersion.h
--- a/Engine/Source/Runtime/Core/Public/UObject/FortniteMainBranchObjectVersion.h
+++ b/Engine/Source/Runtime/Core/Public/UObject/FortniteMainBranchObjectVersion.h
@@ -69,7 +69,7 @@ struct CORE_API FFortniteMainBranchObjectVersion
 		// TimelineTemplates store their derived names instead of dynamically generating
 		StoreTimelineNamesInTemplate,
 		
-		// New Pose Asset data type
+		// Avoid duplicating widget animations to save space.
 		WidgetStopDuplicatingAnimations,
 
 		// Allow reducing of the base LOD, we need to store some imported model data so we can reduce again from the same data.
@@ -220,6 +220,12 @@ struct CORE_API FFortniteMainBranchObjectVersion
 		// Water body collision settings are now those of the base UPrimitiveComponent, rather than duplicated in UWaterBodyComponent
 		WaterBodyComponentCollisionSettingsRefactor,
 
+		// Introducing widget inherited named slots.  This wouldn't have required a version bump, except in the previous
+		// version, users could make NamedSlots and then Seed them with any random widgets, as a sorta 'default' setup.
+		// In order to preserve that, we're bumping the version so that we can set a new field on UNamedSlot to control
+		// if a widget exposes its named slot to everyone (even if it has content), which by default they wont any longer.
+		WidgetInheritedNamedSlots,
+
 		// -----<new versions can be added above this line>-------------------------------------------------
 		VersionPlusOne,
 		LatestVersion = VersionPlusOne - 1
diff --git a/Engine/Source/Runtime/UMG/Private/Components/NamedSlot.cpp b/Engine/Source/Runtime/UMG/Private/Components/NamedSlot.cpp
--- a/Engine/Source/Runtime/UMG/Private/Components/NamedSlot.cpp
+++ b/Engine/Source/Runtime/UMG/Private/Components/NamedSlot.cpp
@@ -93,6 +93,27 @@ const FText UNamedSlot::GetPaletteCategory()
 	return LOCTEXT("Common", "Common");
 }
 
+void UNamedSlot::Serialize(FArchive& Ar)
+{
+	Super::Serialize(Ar);
+
+	Ar.UsingCustomVersion(FFortniteMainBranchObjectVersion::GUID);
+}
+
+void UNamedSlot::PostLoad()
+{
+	Super::PostLoad();
+
+#if WITH_EDITORONLY_DATA
+	const int32 FortniteMainBranchObjectVersion = GetLinkerCustomVersion(FFortniteMainBranchObjectVersion::GUID);
+
+	if (FortniteMainBranchObjectVersion < FFortniteMainBranchObjectVersion::WidgetInheritedNamedSlots)
+	{
+		bExposeOnInstanceOnly = true;
+	}
+#endif
+}
+
 #endif
 
 /////////////////////////////////////////////////////
diff --git a/Engine/Source/Runtime/UMG/Private/Components/Widget.cpp b/Engine/Source/Runtime/UMG/Private/Components/Widget.cpp
--- a/Engine/Source/Runtime/UMG/Private/Components/Widget.cpp
+++ b/Engine/Source/Runtime/UMG/Private/Components/Widget.cpp
@@ -1250,13 +1250,6 @@ void UWidget::DeselectByDesigner()
 #define LOCTEXT_NAMESPACE "UMG"
 #endif // WITH_EDITOR
 
-void UWidget::PreSave(const class ITargetPlatform* TargetPlatform)
-{
-	PRAGMA_DISABLE_DEPRECATION_WARNINGS;
-	Super::PreSave(TargetPlatform);
-	PRAGMA_ENABLE_DEPRECATION_WARNINGS
-}
-
 void UWidget::PreSave(FObjectPreSaveContext ObjectSaveContext)
 {
 	Super::PreSave(ObjectSaveContext);
diff --git a/Engine/Source/Runtime/UMG/Private/UserWidget.cpp b/Engine/Source/Runtime/UMG/Private/UserWidget.cpp
--- a/Engine/Source/Runtime/UMG/Private/UserWidget.cpp
+++ b/Engine/Source/Runtime/UMG/Private/UserWidget.cpp
@@ -829,12 +829,14 @@ UWidget* UUserWidget::GetWidgetFromName(const FName& Name) const
 void UUserWidget::GetSlotNames(TArray<FName>& SlotNames) const
 {
 	// Only do this if this widget is of a blueprint class
-	if (UWidgetBlueprintGeneratedClass* BGClass = GetWidgetTreeOwningClass())
+	if (const UWidgetBlueprintGeneratedClass* BGClass = Cast<UWidgetBlueprintGeneratedClass>(GetClass()))
 	{
-		SlotNames.Append(BGClass->NamedSlots);
+		SlotNames.Append(BGClass->InstanceNamedSlots);
 	}
 	else if (WidgetTree) // For non-blueprint widget blueprints we have to go through the widget tree to locate the named slots dynamically.
 	{
+		// TODO: This code is probably defunct now, that we always have a BPGC?
+		
 		WidgetTree->ForEachWidget([&SlotNames] (UWidget* Widget) {
 			if ( Widget && Widget->IsA<UNamedSlot>() )
 			{
@@ -1932,13 +1934,6 @@ bool UUserWidget::IsAsset() const
 	return false;
 }
 
-void UUserWidget::PreSave(const class ITargetPlatform* TargetPlatform)
-{
-	PRAGMA_DISABLE_DEPRECATION_WARNINGS;
-	Super::PreSave(TargetPlatform);
-	PRAGMA_ENABLE_DEPRECATION_WARNINGS;
-}
-
 void UUserWidget::PreSave(FObjectPreSaveContext ObjectSaveContext)
 {
 	if (WidgetTree)
diff --git a/Engine/Source/Runtime/UMG/Private/WidgetBlueprintGeneratedClass.cpp b/Engine/Source/Runtime/UMG/Private/WidgetBlueprintGeneratedClass.cpp
--- a/Engine/Source/Runtime/UMG/Private/WidgetBlueprintGeneratedClass.cpp
+++ b/Engine/Source/Runtime/UMG/Private/WidgetBlueprintGeneratedClass.cpp
@@ -444,6 +444,18 @@ void UWidgetBlueprintGeneratedClass::Serialize(FArchive& Ar)
 	Super::Serialize(Ar);
 
 	Ar.UsingCustomVersion(FEditorObjectVersion::GUID);
+	Ar.UsingCustomVersion(FFortniteMainBranchObjectVersion::GUID);
+
+	if (Ar.IsLoading())
+	{
+		// We've made it so the actual set of named slots we expose is AvailableNamedSlots, we need to copy
+		// the initial set we load though to ensure we don't get compile time load errors.
+		if (Ar.CustomVer(FFortniteMainBranchObjectVersion::GUID) < FFortniteMainBranchObjectVersion::WidgetInheritedNamedSlots)
+		{
+			AvailableNamedSlots = NamedSlots;
+			InstanceNamedSlots = NamedSlots;
+		}
+	}
 }
 
 UWidgetBlueprintGeneratedClass* UWidgetBlueprintGeneratedClass::FindWidgetTreeOwningClass() const
diff --git a/Engine/Source/Runtime/UMG/Public/Blueprint/UserWidget.h b/Engine/Source/Runtime/UMG/Public/Blueprint/UserWidget.h
--- a/Engine/Source/Runtime/UMG/Public/Blueprint/UserWidget.h
+++ b/Engine/Source/Runtime/UMG/Public/Blueprint/UserWidget.h
@@ -1118,10 +1118,6 @@ public:
 
 	//~ Begin UObject Interface
 	virtual bool IsAsset() const;
-	PRAGMA_DISABLE_DEPRECATION_WARNINGS // Suppress compiler warning on override of deprecated function
-	UE_DEPRECATED(5.0, "Use version that takes FObjectPreSaveContext instead.")
-	virtual void PreSave(const class ITargetPlatform* TargetPlatform) override;
-	PRAGMA_ENABLE_DEPRECATION_WARNINGS
 	virtual void PreSave(FObjectPreSaveContext ObjectSaveContext) override;
 	//~ End UObject Interface
 
diff --git a/Engine/Source/Runtime/UMG/Public/Blueprint/WidgetBlueprintGeneratedClass.h b/Engine/Source/Runtime/UMG/Public/Blueprint/WidgetBlueprintGeneratedClass.h
--- a/Engine/Source/Runtime/UMG/Public/Blueprint/WidgetBlueprintGeneratedClass.h
+++ b/Engine/Source/Runtime/UMG/Public/Blueprint/WidgetBlueprintGeneratedClass.h
@@ -99,19 +99,28 @@ public:
 	TArray< TObjectPtr<UWidgetAnimation> > Animations;
 
 	/**
-	 * Exposed Named Slots for content.  These are slots that are accumulated from all super classes on compile.
-	 * and will exclude any named slots that are filled by a parent class.
+	 * All named slots, even the ones that have content already filled into them by a parent class and are not
+	 * available for extension.
 	 **/
 	UPROPERTY()
-	TArray<FName> AvailableNamedSlots;
+	TArray<FName> NamedSlots;
 
 	/**
-	 * All named slots, even the ones that have content already filled into them by a parent class and are not
-	 * available for extension.
+	 * Available Named Slots for content in a subclass.  These are slots that are accumulated from all super
+	 * classes on compile.  They will exclude any named slots that are filled by a parent class.
 	 **/
 	UPROPERTY()
-	TArray<FName> NamedSlots;
+	TArray<FName> AvailableNamedSlots;
 
+	/**
+	 * These are the set of named slots that can be used on an instance of the widget.  This set is slightly
+	 * different from available named slots, because ones designated UNamedSlot::bExposeOnInstanceOnly == true
+	 * will also be in this list, even though they wont be in AvailableNamedSlots, if are inherited, as inherited
+	 * named slots do not have the capability to remove existing content in a named slot.
+	 **/
+	UPROPERTY()
+	TArray<FName> InstanceNamedSlots;
+	
 public:
 	UWidgetTree* GetWidgetTreeArchetype() const { return WidgetTree; }
 	void SetWidgetTreeArchetype(UWidgetTree* InWidgetTree);
diff --git a/Engine/Source/Runtime/UMG/Public/Components/NamedSlot.h b/Engine/Source/Runtime/UMG/Public/Components/NamedSlot.h
--- a/Engine/Source/Runtime/UMG/Public/Components/NamedSlot.h
+++ b/Engine/Source/Runtime/UMG/Public/Components/NamedSlot.h
@@ -34,6 +34,21 @@ public:
 	virtual const FText GetPaletteCategory() override;
 #endif
 
+public:
+#if WITH_EDITORONLY_DATA
+	/**
+	 * Named slots exposed on the instance only follow a slightly different set of rules.  For example, they can contain
+	 * some content that the user can replace after dropping it into their tree.  However, these slots can not be inherited
+	 * to be filled in by a subclass.  So if you want this named slot to be extensible in a subclass of this widget, you
+	 * should set or leave this value as false.
+	 */
+	UPROPERTY(EditAnywhere, Category="Exposing")
+	bool bExposeOnInstanceOnly = false;
+#endif
+	
+	virtual void Serialize(FArchive& Ar) override;
+	virtual void PostLoad() override;
+	
 protected:
 	// UWidget interface
 	virtual TSharedRef<SWidget> RebuildWidget() override;
diff --git a/Engine/Source/Runtime/UMG/Public/Components/Widget.h b/Engine/Source/Runtime/UMG/Public/Components/Widget.h
--- a/Engine/Source/Runtime/UMG/Public/Components/Widget.h
+++ b/Engine/Source/Runtime/UMG/Public/Components/Widget.h
@@ -951,10 +951,6 @@ public:
 	// Begin UObject
 	virtual UWorld* GetWorld() const override;
 	virtual void FinishDestroy() override;
-	PRAGMA_DISABLE_DEPRECATION_WARNINGS // Suppress compiler warning on override of deprecated function
-	UE_DEPRECATED(5.0, "Use version that takes FObjectPreSaveContext instead.")
-	virtual void PreSave(const class ITargetPlatform* TargetPlatform) override;
-	PRAGMA_ENABLE_DEPRECATION_WARNINGS
 	virtual void PreSave(FObjectPreSaveContext ObjectSaveContext) override;
 	// End UObject
 
