Fix widgets missing when using palette filtering & crashing on drop when using generated classes.

diff --git a/Engine/Source/Editor/UMGEditor/Private/Palette/SPaletteViewModel.cpp b/Engine/Source/Editor/UMGEditor/Private/Palette/SPaletteViewModel.cpp
--- a/Engine/Source/Editor/UMGEditor/Private/Palette/SPaletteViewModel.cpp
+++ b/Engine/Source/Editor/UMGEditor/Private/Palette/SPaletteViewModel.cpp
@@ -14,6 +14,7 @@
 #include "ClassViewerModule.h"
 #include "ClassViewerFilter.h"
 #include "EditorClassUtils.h"
+#include "Engine/BlueprintGeneratedClass.h"
 #include "DragDrop/WidgetTemplateDragDropOp.h"
 
 #include "Templates/WidgetTemplateClass.h"
@@ -340,6 +341,7 @@ void FWidgetCatalogViewModel::BuildClassWidgetList()
 	static const FString InvalidCategoryName("InvalidCategoryName");
 
 	TMap<FName, TSubclassOf<UUserWidget>> LoadedWidgetBlueprintClassesByName;
+	TMap<FName, TSubclassOf<UBlueprintGeneratedClass>> LoadedGeneratedBlueprintClassesByName;
 
 	auto ActiveWidgetBlueprintClass = GetBlueprint()->GeneratedClass;
 	FName ActiveWidgetBlueprintClassName = ActiveWidgetBlueprintClass->GetFName();
@@ -457,10 +459,16 @@ void FWidgetCatalogViewModel::BuildClassWidgetList()
 				LoadedWidgetBlueprintClassesByName.Add(WidgetClass->ClassGeneratedBy->GetFName()) = WidgetClass;
 			}
 		}
+		else if (Cast<UBlueprintGeneratedClass>(WidgetClass))
+		{
+			// Note: Don't check IsChildOf above, since superstruct of BPGC is the C++ class
+			LoadedGeneratedBlueprintClassesByName.Add(WidgetClass->GetFName()) = WidgetClass;
+			TSharedPtr<FWidgetTemplateClass> Template = MakeShareable(new FWidgetTemplateClass(WidgetClass));
+			AddWidgetTemplate(Template);
+		}
 		else
 		{
 			TSharedPtr<FWidgetTemplateClass> Template = MakeShareable(new FWidgetTemplateClass(WidgetClass));
-
 			AddWidgetTemplate(Template);
 		}
 
@@ -563,7 +571,7 @@ void FWidgetCatalogViewModel::BuildClassWidgetList()
 			}
 
 			uint32 BPFlags = WidgetBPAssetData.GetTagValueRef<uint32>(FBlueprintTags::ClassFlags);
-			if ((BPFlags & (CLASS_Abstract | CLASS_Deprecated | CLASS_HideDropDown)) == 0)
+			if ((BPFlags & (CLASS_Hidden | CLASS_HideDropDown | CLASS_Deprecated | CLASS_Abstract | CLASS_NewerVersionExists)) == 0)
 			{
 				auto Template = MakeShareable(new FWidgetTemplateBlueprintClass(WidgetBPAssetData, WidgetBPClass));
 
@@ -571,6 +579,37 @@ void FWidgetCatalogViewModel::BuildClassWidgetList()
 			}
 		}
 	}
+
+
+	TArray<FAssetData> AllGeneratedBPsAssetData;
+	AssetRegistryModule.Get().GetAssetsByClass(UBlueprintGeneratedClass::StaticClass()->GetClassPathName(), AllGeneratedBPsAssetData, true);
+
+	// Search generated class assets if using palette filtering, only do this with palette filtering since otherwise searching all generated BP's would be expensive
+	if (bUseEditorConfigPaletteFiltering)
+	{
+		for (FAssetData& GeneratedBPAssetData : AllGeneratedBPsAssetData)
+		{
+			// Excludes the blueprint you're currently in
+			if (GeneratedBPAssetData.AssetName == ActiveWidgetBlueprintName)
+			{
+				continue;
+			}
+
+			if (!FilterAssetData(GeneratedBPAssetData) 
+				&& AllowedPaletteWidgets.PassesFilter(GeneratedBPAssetData.ObjectPath) 
+				&& !LoadedGeneratedBlueprintClassesByName.Contains(GeneratedBPAssetData.AssetName))
+			{
+				uint32 BPFlags = GeneratedBPAssetData.GetTagValueRef<uint32>(FBlueprintTags::ClassFlags);
+				if ((BPFlags & (CLASS_Hidden | CLASS_HideDropDown | CLASS_Deprecated | CLASS_Abstract | CLASS_NewerVersionExists)) == 0)
+				{
+					bool bIsBlueprintGeneratedClass = true;
+					auto Template = MakeShareable(new FWidgetTemplateBlueprintClass(GeneratedBPAssetData, nullptr, bIsBlueprintGeneratedClass));
+
+					AddWidgetTemplate(Template);
+				}
+			}
+		}
+	}
 }
 
 bool FWidgetCatalogViewModel::FilterAssetData(FAssetData &InAssetData)
diff --git a/Engine/Source/Editor/UMGEditor/Private/Templates/WidgetTemplateBlueprintClass.cpp b/Engine/Source/Editor/UMGEditor/Private/Templates/WidgetTemplateBlueprintClass.cpp
--- a/Engine/Source/Editor/UMGEditor/Private/Templates/WidgetTemplateBlueprintClass.cpp
+++ b/Engine/Source/Editor/UMGEditor/Private/Templates/WidgetTemplateBlueprintClass.cpp
@@ -10,13 +10,28 @@
 #include "Kismet2/BlueprintEditorUtils.h"
 #include "Styling/SlateIconFinder.h"
 #include "Subsystems/AssetEditorSubsystem.h"
+#include "Engine/BlueprintGeneratedClass.h"
 #include "Editor.h"
 
 #define LOCTEXT_NAMESPACE "UMGEditor"
 
-FWidgetTemplateBlueprintClass::FWidgetTemplateBlueprintClass(const FAssetData& InWidgetAssetData, TSubclassOf<UUserWidget> InUserWidgetClass)
+FWidgetTemplateBlueprintClass::FWidgetTemplateBlueprintClass(const FAssetData& InWidgetAssetData, TSubclassOf<UUserWidget> InUserWidgetClass, bool bInIsBlueprintGeneratedClass)
 	: FWidgetTemplateClass(InWidgetAssetData, InUserWidgetClass)
+	, bIsBlueprintGeneratedClass(bInIsBlueprintGeneratedClass)
 {
+	// Blueprints get the class type actions for their parent native class - this avoids us having to load the blueprint
+	if (bIsBlueprintGeneratedClass && !InUserWidgetClass && InWidgetAssetData.IsValid())
+	{
+		FString ParentClassName;
+		if (!WidgetAssetData.GetTagValue(FBlueprintTags::NativeParentClassPath, ParentClassName))
+		{
+			WidgetAssetData.GetTagValue(FBlueprintTags::ParentClassPath, ParentClassName);
+		}
+		if (!ParentClassName.IsEmpty())
+		{
+			CachedParentClass = UClass::TryFindTypeSlow<UClass>(FPackageName::ExportTextPathToObjectPath(ParentClassName));
+		}
+	}
 }
 
 FWidgetTemplateBlueprintClass::~FWidgetTemplateBlueprintClass()
@@ -38,6 +53,10 @@ FText FWidgetTemplateBlueprintClass::GetCategory() const
 		{
 			return FoundPaletteCategoryText;
 		}
+		else if (CachedParentClass.IsValid() && CachedParentClass->IsChildOf(UWidget::StaticClass()) && !CachedParentClass->IsChildOf(UUserWidget::StaticClass()))
+		{
+			return CachedParentClass->GetDefaultObject<UWidget>()->GetPaletteCategory();
+		}
 		else
 		{
 			auto DefaultUserWidget = UUserWidget::StaticClass()->GetDefaultObject<UUserWidget>();
@@ -48,12 +67,20 @@ FText FWidgetTemplateBlueprintClass::GetCategory() const
 
 UWidget* FWidgetTemplateBlueprintClass::Create(UWidgetTree* Tree)
 {
-	// Load the blueprint asset if needed
+	// Load the blueprint asset or blueprint generated class if needed
 	if (!WidgetClass.Get())
 	{
 		FString AssetPath = WidgetAssetData.ObjectPath.ToString();
-		UWidgetBlueprint* LoadedWidget = LoadObject<UWidgetBlueprint>(nullptr, *AssetPath);
-		WidgetClass = *LoadedWidget->GeneratedClass;
+
+		if (bIsBlueprintGeneratedClass)
+		{
+			WidgetClass = LoadObject<UBlueprintGeneratedClass>(nullptr, *AssetPath);
+		}
+		else
+		{
+			UWidgetBlueprint* LoadedWidget = LoadObject<UWidgetBlueprint>(nullptr, *AssetPath);
+			WidgetClass = *LoadedWidget->GeneratedClass;
+		}
 	}
 
 	return FWidgetTemplateClass::CreateNamed(Tree, FName(*FBlueprintEditorUtils::GetClassNameWithoutSuffix(WidgetClass.Get())));
@@ -61,7 +88,14 @@ UWidget* FWidgetTemplateBlueprintClass::Create(UWidgetTree* Tree)
 
 const FSlateBrush* FWidgetTemplateBlueprintClass::GetIcon() const
 {
-	return FSlateIconFinder::FindIconBrushForClass(UUserWidget::StaticClass());
+	if (CachedParentClass.IsValid() && CachedParentClass->IsChildOf(UWidget::StaticClass()) && !CachedParentClass->IsChildOf(UUserWidget::StaticClass()))
+	{
+		return FSlateIconFinder::FindIconBrushForClass(CachedParentClass.Get());
+	}
+	else
+	{
+		return FSlateIconFinder::FindIconBrushForClass(UUserWidget::StaticClass());
+	}
 	
 }
 
diff --git a/Engine/Source/Editor/UMGEditor/Private/Templates/WidgetTemplateClass.cpp b/Engine/Source/Editor/UMGEditor/Private/Templates/WidgetTemplateClass.cpp
--- a/Engine/Source/Editor/UMGEditor/Private/Templates/WidgetTemplateClass.cpp
+++ b/Engine/Source/Editor/UMGEditor/Private/Templates/WidgetTemplateClass.cpp
@@ -39,7 +39,9 @@ FWidgetTemplateClass::FWidgetTemplateClass(const FAssetData& InWidgetAssetData,
 	}
 	else
 	{
-		Name = FText::FromString(FName::NameToDisplayString(WidgetAssetData.AssetName.ToString(), false));
+		FString AssetName = WidgetAssetData.AssetName.ToString();
+		AssetName.RemoveFromEnd(TEXT("_C"), ESearchCase::CaseSensitive);
+		Name = FText::FromString(FName::NameToDisplayString(AssetName, false));
 
 		// Blueprints get the class type actions for their parent native class - this avoids us having to load the blueprint
 		if (InWidgetAssetData.IsValid())
diff --git a/Engine/Source/Editor/UMGEditor/Public/Templates/WidgetTemplateBlueprintClass.h b/Engine/Source/Editor/UMGEditor/Public/Templates/WidgetTemplateBlueprintClass.h
--- a/Engine/Source/Editor/UMGEditor/Public/Templates/WidgetTemplateBlueprintClass.h
+++ b/Engine/Source/Editor/UMGEditor/Public/Templates/WidgetTemplateBlueprintClass.h
@@ -23,8 +23,9 @@ public:
 	* Constructor
 	* @param InWidgetAssetData The FAssetData for the widget blueprint asset
 	* @param InUserWidgetClass The user-created widget class that will be created by this template
+	* @param bCacheParentClass True means we will search for parent class and cache, this is a slow operation, but useful for BPGCs.
 	*/
-	FWidgetTemplateBlueprintClass(const FAssetData& InWidgetAssetData, TSubclassOf<UUserWidget> InUserWidgetClass = nullptr);
+	FWidgetTemplateBlueprintClass(const FAssetData& InWidgetAssetData, TSubclassOf<UUserWidget> InUserWidgetClass = nullptr, bool bInIsBlueprintGeneratedClass = false);
 
 	/** Destructor */
 	virtual ~FWidgetTemplateBlueprintClass();
@@ -47,4 +48,10 @@ public:
 	/** Returns true if the supplied class is supported by this template */
 	static bool Supports(UClass* InClass);
 
+protected:
+	/** Parent Class of this widget template, may not be valid */
+	TWeakObjectPtr<UClass> CachedParentClass;
+
+	/** True if this template is for a cooked BP generated class */
+	bool bIsBlueprintGeneratedClass;
 };
