PR #9365: Fix not display BindWidgetAnim in SBindWidgetView (Contributed by [REDACTED_PERSON])

diff --git a/Engine/Source/Editor/UMGEditor/Private/Widgets/SBindWidgetView.cpp b/Engine/Source/Editor/UMGEditor/Private/Widgets/SBindWidgetView.cpp
--- a/Engine/Source/Editor/UMGEditor/Private/Widgets/SBindWidgetView.cpp
+++ b/Engine/Source/Editor/UMGEditor/Private/Widgets/SBindWidgetView.cpp
@@ -318,7 +318,7 @@ TArray<FBindWidgetListEntryPtr> BuildSourceDataList(TSharedPtr<FWidgetBlueprintE
 		else if (WidgetProperty->PropertyClass->IsChildOf(UWidgetAnimation::StaticClass()))
 		{
 			bool bIsOptional = false;
-			if (FWidgetBlueprintEditorUtils::IsBindWidgetProperty(WidgetProperty, bIsOptional))
+			if (FWidgetBlueprintEditorUtils::IsBindWidgetAnimProperty(WidgetProperty, bIsOptional))
 			{
 				TSharedRef<FBindWidgetListEntry> Ref = MakeShared<FBindWidgetListEntry>();
 				Ref->Property = WidgetProperty;
