Added SReadOnlyHierarchyView widget to display the UWidgetTree of a UWidgetBlueprint without all of the additional features baked into UMG's SHierarchyView.

#rnx

diff --git a/Engine/Source/Editor/UMGEditor/Private/Hierarchy/SReadOnlyHierarchyView.cpp b/Engine/Source/Editor/UMGEditor/Private/Hierarchy/SReadOnlyHierarchyView.cpp
new file mode 100644
--- /dev/null
+++ b/Engine/Source/Editor/UMGEditor/Private/Hierarchy/SReadOnlyHierarchyView.cpp
@@ -0,0 +1,293 @@
+// Copyright Epic Games, Inc. All Rights Reserved.
+
+#include "Hierarchy/SReadOnlyHierarchyView.h"
+
+#include "Blueprint/WidgetTree.h"
+#include "Styling/SlateIconFinder.h"
+#include "Widgets/Input/SSearchBox.h"
+
+#define LOCTEXT_NAMESPACE "SReadOnlyHierarchyView"
+
+void SReadOnlyHierarchyView::Construct(const FArguments& InArgs, const UWidgetBlueprint* InWidgetBlueprint)
+{
+	OnSelectionChangedDelegate = InArgs._OnSelectionChanged;
+	WidgetBlueprint = InWidgetBlueprint;
+
+	SearchFilter = MakeShared<FTextFilter>(FTextFilter::FItemToStringArray::CreateSP(this, &SReadOnlyHierarchyView::GetFilterStringsForItem));
+	FilterHandler = MakeShared<FTreeFilterHandler>();
+	FilterHandler->SetFilter(SearchFilter.Get());
+	FilterHandler->SetRootItems(&RootWidgets, &FilteredRootWidgets);
+	FilterHandler->SetGetChildrenDelegate(FTreeFilterHandler::FOnGetChildren::CreateSP(this, &SReadOnlyHierarchyView::GetItemChildren));
+
+	TreeView = SNew(STreeView<TSharedPtr<FItem>>)
+		.OnGenerateRow(this, &SReadOnlyHierarchyView::GenerateRow)
+		.OnGetChildren(FilterHandler.ToSharedRef(), &FTreeFilterHandler::OnGetFilteredChildren)
+		.OnSelectionChanged(this, &SReadOnlyHierarchyView::OnSelectionChanged)
+		.SelectionMode(InArgs._SelectionMode)
+		.TreeItemsSource(&FilteredRootWidgets)
+		.ClearSelectionOnClick(false)
+		.OnSetExpansionRecursive(this, &SReadOnlyHierarchyView::SetItemExpansionRecursive);
+
+	FilterHandler->SetTreeView(TreeView.Get());
+
+	Refresh();
+
+	TSharedRef<SVerticalBox> ContentBox = SNew(SVerticalBox);
+
+	if (InArgs._ShowSearch)
+	{
+		ContentBox->AddSlot()
+			.Padding(2)
+			.AutoHeight()
+			[
+				SAssignNew(SearchBox, SSearchBox)
+				.OnTextChanged(this, &SReadOnlyHierarchyView::SetRawFilterText)
+			];
+	}
+
+	ContentBox->AddSlot()
+		[
+			TreeView.ToSharedRef()
+		];
+
+	ChildSlot
+	[
+		ContentBox
+	];
+
+	ExpandAll();
+}
+
+SReadOnlyHierarchyView::~SReadOnlyHierarchyView()
+{
+}
+
+void SReadOnlyHierarchyView::OnSelectionChanged(TSharedPtr<FItem> Selected, ESelectInfo::Type SelectionType)
+{
+	if (!Selected.IsValid())
+	{
+		OnSelectionChangedDelegate.ExecuteIfBound(FName(), SelectionType);
+		return;
+	}
+
+	if (const UWidgetBlueprint* WidgetBP = Selected->WidgetBlueprint.Get())
+	{
+		OnSelectionChangedDelegate.ExecuteIfBound(WidgetBP->GetFName(), SelectionType);
+	}
+
+	if (const UWidget* Widget = Selected->Widget.Get())
+	{
+		OnSelectionChangedDelegate.ExecuteIfBound(Widget->GetFName(), SelectionType);
+	}
+}
+
+void SReadOnlyHierarchyView::Refresh()
+{
+	RootWidgets.Reset();
+	FilteredRootWidgets.Reset();
+	RebuildTree();
+	FilterHandler->RefreshAndFilterTree();
+	ExpandAll();
+}
+
+void SReadOnlyHierarchyView::SetItemExpansionRecursive(TSharedPtr<FItem> Item, bool bShouldBeExpanded)
+{
+	TreeView->SetItemExpansion(Item, bShouldBeExpanded);
+
+	for (const TSharedPtr<FItem>& Child : Item->Children)
+	{
+		SetItemExpansionRecursive(Child, bShouldBeExpanded);
+	}
+}
+
+void SReadOnlyHierarchyView::SetRawFilterText(const FText& Text)
+{
+	FilterHandler->SetIsEnabled(!Text.IsEmpty());
+	SearchFilter->SetRawFilterText(Text);
+	FilterHandler->RefreshAndFilterTree();
+}
+
+FText SReadOnlyHierarchyView::GetItemText(TSharedPtr<FItem> Item) const
+{
+	if (const UWidget* Widget = Item->Widget.Get())
+	{
+		return Widget->GetLabelTextWithMetadata();
+	}
+
+	static const FText RootWidgetFormat = LOCTEXT("WidgetNameFormat", "[{0}]");
+	return FText::Format(RootWidgetFormat, FText::FromString(WidgetBlueprint->GetName()));
+}
+
+const FSlateBrush* SReadOnlyHierarchyView::GetIconBrush(TSharedPtr<FItem> Item) const
+{
+	if (const UWidget* Widget = Item->Widget.Get())
+	{
+		return FSlateIconFinder::FindIconBrushForClass(Widget->GetClass());
+	}
+	return nullptr;
+}
+
+FText SReadOnlyHierarchyView::GetIconToolTipText(TSharedPtr<FItem> Item) const
+{
+	if (const UWidget* Widget = Item->Widget.Get())
+	{
+		UClass* WidgetClass = Widget->GetClass();
+		if (WidgetClass->IsChildOf(UUserWidget::StaticClass()) && WidgetClass->ClassGeneratedBy)
+		{
+			const FString& Description = Cast<UWidgetBlueprint>(WidgetClass->ClassGeneratedBy)->BlueprintDescription;
+			if (Description.Len() > 0)
+			{
+				return FText::FromString(Description);
+			}
+		}
+
+		return WidgetClass->GetToolTipText();
+	}
+
+	return FText::GetEmpty();
+}
+
+FText SReadOnlyHierarchyView::GetWidgetToolTipText(TSharedPtr<FItem> Item) const
+{
+	if (const UWidget* Widget = Item->Widget.Get())
+	{
+		if (!Widget->IsGeneratedName())
+		{
+			return FText::FromString(TEXT("[") + Widget->GetClass()->GetDisplayNameText().ToString() + TEXT("]"));
+		}
+	}
+
+	return FText::GetEmpty();
+}
+
+TSharedRef<ITableRow> SReadOnlyHierarchyView::GenerateRow(TSharedPtr<FItem> Item, const TSharedRef<STableViewBase>& OwnerTable) const
+{
+	return SNew(STableRow<TSharedPtr<FItem>>, OwnerTable)
+		[
+			SNew(SHorizontalBox)
+			// Widget icon
+			+ SHorizontalBox::Slot()
+			.AutoWidth()
+			.VAlign(VAlign_Center)
+			[
+				SNew(SImage)
+				.ColorAndOpacity(FSlateColor::UseForeground())
+				.Image(GetIconBrush(Item))
+				.ToolTipText(this, &SReadOnlyHierarchyView::GetIconToolTipText, Item)
+			]
+			// Name of the widget
+			+ SHorizontalBox::Slot()
+			.FillWidth(1.0f)
+			.Padding(2, 0, 0, 0)
+			.VAlign(VAlign_Center)
+			[
+				SNew(STextBlock)
+				.Font(Item->Widget.Get() == nullptr ? FCoreStyle::GetDefaultFontStyle("Bold", 10) : FCoreStyle::Get().GetFontStyle("NormalFont"))
+				.Text(this, &SReadOnlyHierarchyView::GetItemText, Item)
+				.ToolTipText(this, &SReadOnlyHierarchyView::GetWidgetToolTipText, Item)
+				.HighlightText_Lambda([this]() { return SearchFilter->GetRawFilterText(); })
+			]
+		];
+}
+
+void SReadOnlyHierarchyView::GetFilterStringsForItem(TSharedPtr<FItem> Item, TArray<FString>& OutStrings) const
+{
+	if (const UWidget* Widget = Item->Widget.Get())
+	{
+		OutStrings.Add(Widget->GetName());
+		OutStrings.Add(Widget->GetLabelTextWithMetadata().ToString());
+	}
+	else
+	{
+		OutStrings.Add(WidgetBlueprint->GetName());
+	}
+}
+
+void SReadOnlyHierarchyView::SetSelectedWidget(FName WidgetName)
+{
+	if (TSharedPtr<FItem> Found = FindItem(RootWidgets, WidgetName))
+	{
+		TreeView->SetSelection(Found);
+	}
+}
+
+void SReadOnlyHierarchyView::GetItemChildren(TSharedPtr<FItem> Item, TArray<TSharedPtr<FItem>>& OutChildren) const
+{
+	OutChildren.Append(Item->Children);
+}
+
+void SReadOnlyHierarchyView::BuildWidgetChildren(const TSharedPtr<FItem>& CurrentItem)
+{
+	if (const UPanelWidget* PanelWidget = Cast<UPanelWidget>(CurrentItem->Widget.Get()))
+	{
+		int32 NumChildren = PanelWidget->GetChildrenCount();
+		CurrentItem->Children.Reserve(NumChildren);
+
+		for (int32 Idx = 0; Idx < NumChildren; ++Idx)
+		{
+			if (const UWidget* Child = PanelWidget->GetChildAt(Idx))
+			{
+				TSharedRef<FItem> ChildItem = MakeShared<FItem>(Child);
+				BuildWidgetChildren(ChildItem);
+				CurrentItem->Children.Add(ChildItem);
+			}
+		}
+	}
+}
+
+void SReadOnlyHierarchyView::RebuildTree()
+{
+	const UWidgetBlueprint* WidgetBP = WidgetBlueprint.Get();
+	TSharedRef<FItem> WidgetBPItem = MakeShared<FItem>(WidgetBP);
+	RootWidgets.Add(WidgetBPItem);
+
+	if (const UWidget* RootWidget = WidgetBP->WidgetTree->RootWidget.Get())
+	{
+		TSharedRef<FItem> RootWidgetItem = MakeShared<FItem>(RootWidget);
+		WidgetBPItem->Children.Add(RootWidgetItem);
+
+		BuildWidgetChildren(RootWidgetItem);
+	}
+
+	ExpandAll();
+}
+
+void SReadOnlyHierarchyView::ExpandAll()
+{
+	for (const TSharedPtr<FItem>& Item : FilteredRootWidgets)
+	{
+		SetItemExpansionRecursive(Item, true);
+	}
+}
+
+TSharedPtr<SReadOnlyHierarchyView::FItem> SReadOnlyHierarchyView::FindItem(const TArray<TSharedPtr<SReadOnlyHierarchyView::FItem>>& RootItems, FName Name) const
+{
+	TArray<TSharedPtr<FItem>> Items;
+	Items.Append(RootItems);
+
+	for (int32 Idx = 0; Idx < Items.Num(); ++Idx)
+	{
+		TSharedPtr<FItem> Item = Items[Idx];
+
+		if (const UWidgetBlueprint* WidgetBP = Item->WidgetBlueprint.Get())
+		{
+			if (WidgetBP->GetFName() == Name)
+			{
+				return Item;
+			}
+		}
+		else if (const UWidget* Widget = Item->Widget.Get())
+		{
+			if (Widget->GetFName() == Name)
+			{
+				return Item;
+			}
+		}
+
+		GetItemChildren(Item, Items);
+	}
+
+	return TSharedPtr<FItem>();
+}
+
+#undef LOCTEXT_NAMESPACE
diff --git a/Engine/Source/Editor/UMGEditor/Public/Hierarchy/SReadOnlyHierarchyView.h b/Engine/Source/Editor/UMGEditor/Public/Hierarchy/SReadOnlyHierarchyView.h
new file mode 100644
--- /dev/null
+++ b/Engine/Source/Editor/UMGEditor/Public/Hierarchy/SReadOnlyHierarchyView.h
@@ -0,0 +1,90 @@
+// Copyright Epic Games, Inc. All Rights Reserved.
+
+#pragma once
+
+#include "Components/Widget.h"
+#include "Framework/Views/TreeFilterHandler.h"
+#include "Misc/TextFilter.h"
+#include "Framework/SlateDelegates.h"
+#include "UObject/ObjectPtr.h"
+#include "WidgetBlueprint.h"
+#include "Widgets/SCompoundWidget.h"
+
+class ITableRow;
+class SSearchBox;
+class STableViewBase;
+
+template <typename ItemType> class STreeView;
+template <typename ItemType> class TreeFilterHandler;
+template <typename ItemType> class TTextFilter;
+
+class UMGEDITOR_API SReadOnlyHierarchyView : public SCompoundWidget
+{
+private:
+
+	struct FItem
+	{
+		FItem(const UWidget* InWidget) : Widget(InWidget) {}
+		FItem(const UWidgetBlueprint* InWidgetBlueprint) : WidgetBlueprint(InWidgetBlueprint) {}
+
+		TWeakObjectPtr<const UWidget> Widget;
+		TWeakObjectPtr<const UWidgetBlueprint> WidgetBlueprint;
+
+		TArray<TSharedPtr<FItem>> Children;
+	};
+
+public:
+
+	using FOnSelectionChanged = typename TSlateDelegates<FName>::FOnSelectionChanged;
+
+	SLATE_BEGIN_ARGS(SReadOnlyHierarchyView) {}
+		SLATE_ARGUMENT_DEFAULT(bool, ShowSearch) = true;
+		SLATE_ARGUMENT_DEFAULT(ESelectionMode::Type, SelectionMode) = ESelectionMode::Single;
+		SLATE_EVENT(FOnSelectionChanged, OnSelectionChanged)
+	SLATE_END_ARGS()
+
+	virtual ~SReadOnlyHierarchyView();
+
+	void Construct(const FArguments& InArgs, const UWidgetBlueprint* WidgetBlueprint);
+	void Refresh();
+
+	void SetSelectedWidget(FName WidgetName);
+	void SetRawFilterText(const FText& Text);
+
+private:
+
+	TSharedRef<ITableRow> GenerateRow(TSharedPtr<FItem> Item, const TSharedRef<STableViewBase>& OwnerTable) const;
+
+	void GetItemChildren(TSharedPtr<FItem> Item, TArray<TSharedPtr<FItem>>& OutChildren) const;
+	FText GetItemText(TSharedPtr<FItem> Item) const;
+	const FSlateBrush* GetIconBrush(TSharedPtr<FItem> Item) const;
+	FText GetIconToolTipText(TSharedPtr<FItem> Item) const;
+	FText GetWidgetToolTipText(TSharedPtr<FItem> Item) const;
+	void GetFilterStringsForItem(TSharedPtr<FItem> Item, TArray<FString>& OutStrings) const;
+	TSharedPtr<FItem> FindItem(const TArray<TSharedPtr<FItem>>& RootItems, FName Name) const;
+	void OnSelectionChanged(TSharedPtr<FItem> Selected, ESelectInfo::Type SelectionType);
+
+	void RebuildTree();
+	void BuildWidgetChildren(const TSharedPtr<FItem>& CurrentItem);
+
+	void SetItemExpansionRecursive(TSharedPtr<FItem> Item, bool bShouldBeExpanded);
+	void ExpandAll();
+
+private:
+
+	TWeakObjectPtr<const UWidgetBlueprint> WidgetBlueprint;
+
+	TArray<TSharedPtr<FItem>> RootWidgets;
+	TArray<TSharedPtr<FItem>> FilteredRootWidgets;
+
+	FOnSelectionChanged OnSelectionChangedDelegate;
+
+	TSharedPtr<SSearchBox> SearchBox;
+
+	using FTextFilter = TTextFilter<TSharedPtr<FItem>>;
+	TSharedPtr<FTextFilter> SearchFilter;
+	using FTreeFilterHandler = TreeFilterHandler<TSharedPtr<FItem>>;
+	TSharedPtr<FTreeFilterHandler> FilterHandler;
+
+	TSharedPtr<STreeView<TSharedPtr<FItem>>> TreeView;
+};
