Core/UMG

-Add support for struct properties to USinglePropertyView

diff --git a/Engine/Source/Editor/UMGEditor/Private/Components/SinglePropertyView.cpp b/Engine/Source/Editor/UMGEditor/Private/Components/SinglePropertyView.cpp
--- a/Engine/Source/Editor/UMGEditor/Private/Components/SinglePropertyView.cpp
+++ b/Engine/Source/Editor/UMGEditor/Private/Components/SinglePropertyView.cpp
@@ -67,9 +67,7 @@ void USinglePropertyView::BuildContentWidget()
 			{
 				MissingWidgetText = FPropertyViewHelper::InvalidPropertyText;
 			}
-			else if (CastField<FStructProperty>(Property) || CastField<FArrayProperty>(Property)
-				|| CastField<FMapProperty>(Property) || CastField<FSetProperty>(Property)
-				)
+			else if (CastField<FArrayProperty>(Property) || CastField<FMapProperty>(Property) || CastField<FSetProperty>(Property))
 			{
 				MissingWidgetText = FPropertyViewHelper::UnsupportedPropertyText;
 			}
@@ -90,7 +88,15 @@ void USinglePropertyView::BuildContentWidget()
 				}
 				else
 				{
-					MissingWidgetText = FPropertyViewHelper::UnknownErrorText;
+					// Some built-in structs like FColor are supported, others aren't
+					if (CastField<FStructProperty>(Property))
+					{
+						MissingWidgetText = FPropertyViewHelper::UnsupportedPropertyText;
+					}
+					else
+					{
+						MissingWidgetText = FPropertyViewHelper::UnknownErrorText;
+					}
 				}
 			}
 		}
