UMG: When replacing a binding with new widgets, create a new binding, move the content of the old to the new rather than reuse the existing binding. This fixes issues where old possessable bindings are kept around (with their old binding names and their old possessed object classes).

diff --git a/Engine/Source/Editor/UMGEditor/Private/WidgetBlueprintEditor.cpp b/Engine/Source/Editor/UMGEditor/Private/WidgetBlueprintEditor.cpp
--- a/Engine/Source/Editor/UMGEditor/Private/WidgetBlueprintEditor.cpp
+++ b/Engine/Source/Editor/UMGEditor/Private/WidgetBlueprintEditor.cpp
@@ -2220,10 +2220,18 @@ void FWidgetBlueprintEditor::ReplaceTrackWithWidgets(TArray<FWidgetReference> Wi
 
 	// Remove everything from the track
 	RemoveAllWidgetsFromTrack(ObjectId);
+	
+	// Create a new guid for the first object
+	FGuid NewGuid = ActiveSequencer->GetHandleToObject(Widgets[0].GetPreview());
 
-	AddWidgetsToTrack(Widgets, ObjectId);
+	// Move binding contents and remove possessable
+	MovieScene->MoveBindingContents(ObjectId, NewGuid);
+	MovieScene->RemovePossessable(ObjectId);
 
-	UpdateTrackName(ObjectId);
+	// Add all the remaining widgets to the new binding
+	AddWidgetsToTrack(Widgets, NewGuid);
+
+	UpdateTrackName(NewGuid);
 	SyncSequencersMovieSceneData();
 }
 
