PR #8884: EditDefaultsOnly properties on UUserWidgets show up in the Designer details panel now (Contributed by [REDACTED_PERSON])

diff --git a/Engine/Source/Editor/UMGEditor/Private/Details/SWidgetDetailsView.cpp b/Engine/Source/Editor/UMGEditor/Private/Details/SWidgetDetailsView.cpp
--- a/Engine/Source/Editor/UMGEditor/Private/Details/SWidgetDetailsView.cpp
+++ b/Engine/Source/Editor/UMGEditor/Private/Details/SWidgetDetailsView.cpp
@@ -45,7 +45,6 @@ void SWidgetDetailsView::Construct(const FArguments& InArgs, TSharedPtr<FWidgetB
 	DetailsViewArgs.NameAreaSettings = FDetailsViewArgs::HideNameArea;
 	DetailsViewArgs.bHideSelectionTip = true;
 	DetailsViewArgs.NotifyHook = this;
-	DetailsViewArgs.DefaultsOnlyVisibility = EEditDefaultsOnlyNodeVisibility::Automatic;
 
 	PropertyView = EditModule.CreateDetailView(DetailsViewArgs);
 
@@ -57,6 +56,9 @@ void SWidgetDetailsView::Construct(const FArguments& InArgs, TSharedPtr<FWidgetB
 	TSharedRef<FDetailWidgetExtensionHandler> BindingHandler = MakeShareable( new FDetailWidgetExtensionHandler( InBlueprintEditor ) );
 	PropertyView->SetExtensionHandler(BindingHandler);
 
+	// Handle EditDefaultsOnly and EditInstanceOnly flags
+	PropertyView->SetIsPropertyVisibleDelegate(FIsPropertyVisible::CreateSP(this, &SWidgetDetailsView::IsPropertyVisible));
+
 	// Notify us of object selection changes so we can update the package re-mapping
 	PropertyView->SetOnObjectArrayChanged(FOnObjectArrayChanged::CreateSP(this, &SWidgetDetailsView::OnPropertyViewObjectArrayChanged));
 
@@ -313,6 +315,24 @@ void SWidgetDetailsView::ClearFocusIfOwned()
 	}
 }
 
+bool SWidgetDetailsView::IsPropertyVisible(const FPropertyAndParent& PropertyAndParent) const
+{
+	const FProperty& Property = PropertyAndParent.Property;
+
+	const bool bIsEditDefaultsOnly = Property.HasAnyPropertyFlags(CPF_DisableEditOnInstance);
+	const bool bIsEditInstanceOnly = Property.HasAnyPropertyFlags(CPF_DisableEditOnTemplate);
+	if (bIsEditDefaultsOnly || bIsEditInstanceOnly)
+	{
+		// EditDefaultsOnly properties are only visible when the CDO/root is selected, EditInstanceOnly are only visible when the CDO/root is *not* selected
+		const bool bIsCDOSelected = IsWidgetCDOSelected();
+		if ((bIsEditDefaultsOnly && !bIsCDOSelected) || (bIsEditInstanceOnly && bIsCDOSelected))
+		{
+			return false;
+		}
+	}
+	return true;
+}
+
 bool SWidgetDetailsView::IsWidgetCDOSelected() const
 {
 	if ( SelectedObjects.Num() == 1 )
diff --git a/Engine/Source/Editor/UMGEditor/Private/Details/SWidgetDetailsView.h b/Engine/Source/Editor/UMGEditor/Private/Details/SWidgetDetailsView.h
--- a/Engine/Source/Editor/UMGEditor/Private/Details/SWidgetDetailsView.h
+++ b/Engine/Source/Editor/UMGEditor/Private/Details/SWidgetDetailsView.h
@@ -48,6 +48,8 @@ private:
 
 	void ClearFocusIfOwned();
 
+	bool IsPropertyVisible(const FPropertyAndParent& PropertyAndParent) const;
+	
 	bool IsWidgetCDOSelected() const;
 
 	EVisibility GetBorderAreaVisibility() const;
