Add display of size for current widget in UMG editor (on behalf of Devan England).

diff --git a/Engine/Source/Editor/UMGEditor/Private/Designer/SDesignerView.cpp b/Engine/Source/Editor/UMGEditor/Private/Designer/SDesignerView.cpp
--- a/Engine/Source/Editor/UMGEditor/Private/Designer/SDesignerView.cpp
+++ b/Engine/Source/Editor/UMGEditor/Private/Designer/SDesignerView.cpp
@@ -477,6 +477,19 @@ TSharedRef<SWidget> SDesignerView::CreateOverlayUI()
 			.Visibility(this, &SDesignerView::GetCursorPositionTextVisibility)
 		]
 
+		+ SHorizontalBox::Slot()
+		.AutoWidth()
+		.VAlign(VAlign_Center)
+		.Padding(40, 2, 0, 0)
+		[
+			SNew(STextBlock)
+			.TextStyle(FAppStyle::Get(), "Graph.ZoomText")
+			.Font(FCoreStyle::GetDefaultFontStyle(TEXT("BoldCondensed"), 14))
+			.Text(this, &SDesignerView::GetSelectedWidgetDimensionsText)
+			.ColorAndOpacity(FLinearColor(1.f, 1.f, 1.f, 0.25f))
+			.Visibility(this, &SDesignerView::GetSelectedWidgetDimensionsVisibility)
+		]
+
 		+ SHorizontalBox::Slot()
 		.FillWidth(1.0f)
 		[
@@ -3158,6 +3171,37 @@ EVisibility SDesignerView::GetCursorPositionTextVisibility() const
 	return IsHovered() ? EVisibility::SelfHitTestInvisible : EVisibility::Collapsed;
 }
 
+FText SDesignerView::GetSelectedWidgetDimensionsText() const
+{
+	const FWidgetReference SelectedWidget = GetSelectedWidget();
+	if ( SelectedWidget.IsValid() )
+	{
+		const UWidget* WidgetPreview = SelectedWidget.GetPreview();
+		const FVector2D& Size = WidgetPreview->GetCachedGeometry().GetLocalSize();
+
+		FNumberFormattingOptions FmtOptions;
+		FmtOptions.SetMaximumFractionalDigits(2);
+		const FText ScaleFactorText = FText::Format(
+			LOCTEXT("SelectionDimensionsScaleFormat", "(Render Scale: {0} x {1})"), 
+			FText::AsNumber(WidgetPreview->GetRenderTransform().Scale.X, &FmtOptions),
+			FText::AsNumber(WidgetPreview->GetRenderTransform().Scale.Y, &FmtOptions));
+
+		bool bShowScaleFactor = !FMath::IsNearlyEqual(WidgetPreview->GetRenderTransform().Scale.X, 1.f) || !FMath::IsNearlyEqual(WidgetPreview->GetRenderTransform().Scale.Y, 1.f);
+		return FText::Format(
+			LOCTEXT("SelectionDimensionsFormat", "Selection: {0} x {1} {2}"), 
+			FText::AsNumber(Size.X, &FmtOptions), 
+			FText::AsNumber(Size.Y, &FmtOptions), 
+			bShowScaleFactor ? ScaleFactorText : FText());
+	}
+	return FText();
+}
+
+EVisibility SDesignerView::GetSelectedWidgetDimensionsVisibility() const
+{
+	const TSet<FWidgetReference>& SelectedWidgets = BlueprintEditor.Pin()->GetSelectedWidgets();
+	return SelectedWidgets.Num() == 1 ? EVisibility::SelfHitTestInvisible : EVisibility::Collapsed;
+}
+
 FReply SDesignerView::HandleDPISettingsClicked()
 {
 	FModuleManager::LoadModuleChecked<ISettingsModule>("Settings").ShowViewer("Project", "Engine", "UI");
diff --git a/Engine/Source/Editor/UMGEditor/Public/Designer/SDesignerView.h b/Engine/Source/Editor/UMGEditor/Public/Designer/SDesignerView.h
--- a/Engine/Source/Editor/UMGEditor/Public/Designer/SDesignerView.h
+++ b/Engine/Source/Editor/UMGEditor/Public/Designer/SDesignerView.h
@@ -220,6 +220,9 @@ private:
 
 	FText GetCursorPositionText() const;
 	EVisibility GetCursorPositionTextVisibility() const;
+	
+	FText GetSelectedWidgetDimensionsText() const;
+	EVisibility GetSelectedWidgetDimensionsVisibility() const;
 
 	// Handles drawing selection and other effects a SPaintSurface widget injected into the hierarchy.
 	int32 HandleEffectsPainting(const FOnPaintHandlerParams& PaintArgs);
