Added support for multi-selection in the SReadOnlyHierarchyView.

#rnx

diff --git a/Engine/Source/Editor/UMGEditor/Private/Hierarchy/SReadOnlyHierarchyView.cpp b/Engine/Source/Editor/UMGEditor/Private/Hierarchy/SReadOnlyHierarchyView.cpp
--- a/Engine/Source/Editor/UMGEditor/Private/Hierarchy/SReadOnlyHierarchyView.cpp
+++ b/Engine/Source/Editor/UMGEditor/Private/Hierarchy/SReadOnlyHierarchyView.cpp
@@ -2,6 +2,7 @@
 
 #include "Hierarchy/SReadOnlyHierarchyView.h"
 
+#include "Algo/Transform.h"
 #include "Blueprint/WidgetTree.h"
 #include "Styling/SlateIconFinder.h"
 #include "Widgets/Input/SSearchBox.h"
@@ -62,23 +63,29 @@ SReadOnlyHierarchyView::~SReadOnlyHierarchyView()
 {
 }
 
-void SReadOnlyHierarchyView::OnSelectionChanged(TSharedPtr<FItem> Selected, ESelectInfo::Type SelectionType)
+FName SReadOnlyHierarchyView::GetItemName(const TSharedPtr<FItem>& Item) const 
 {
-	if (!Selected.IsValid())
+	if (!Item.IsValid())
 	{
-		OnSelectionChangedDelegate.ExecuteIfBound(FName(), SelectionType);
-		return;
+		return FName();
 	}
 
-	if (const UWidgetBlueprint* WidgetBP = Selected->WidgetBlueprint.Get())
+	if (const UWidgetBlueprint* WidgetBP = Item->WidgetBlueprint.Get())
 	{
-		OnSelectionChangedDelegate.ExecuteIfBound(WidgetBP->GetFName(), SelectionType);
+		return WidgetBP->GetFName();
 	}
 
-	if (const UWidget* Widget = Selected->Widget.Get())
+	if (const UWidget* Widget = Item->Widget.Get())
 	{
-		OnSelectionChangedDelegate.ExecuteIfBound(Widget->GetFName(), SelectionType);
+		return Widget->GetFName();
 	}
+
+	return FName();
+}
+
+void SReadOnlyHierarchyView::OnSelectionChanged(TSharedPtr<FItem> Selected, ESelectInfo::Type SelectionType)
+{
+	OnSelectionChangedDelegate.ExecuteIfBound(GetItemName(Selected), SelectionType);
 }
 
 void SReadOnlyHierarchyView::Refresh()
@@ -167,8 +174,9 @@ TSharedRef<ITableRow> SReadOnlyHierarchyView::GenerateRow(TSharedPtr<FItem> Item
 			SNew(SHorizontalBox)
 			// Widget icon
 			+ SHorizontalBox::Slot()
-			.AutoWidth()
 			.VAlign(VAlign_Center)
+			.Padding(4, 3, 0, 3)
+			.AutoWidth()
 			[
 				SNew(SImage)
 				.ColorAndOpacity(FSlateColor::UseForeground())
@@ -178,7 +186,7 @@ TSharedRef<ITableRow> SReadOnlyHierarchyView::GenerateRow(TSharedPtr<FItem> Item
 			// Name of the widget
 			+ SHorizontalBox::Slot()
 			.FillWidth(1.0f)
-			.Padding(2, 0, 0, 0)
+			.Padding(2, 0, 4, 0)
 			.VAlign(VAlign_Center)
 			[
 				SNew(STextBlock)
@@ -211,6 +219,24 @@ void SReadOnlyHierarchyView::SetSelectedWidget(FName WidgetName)
 	}
 }
 
+TArray<FName> SReadOnlyHierarchyView::GetSelectedWidgets() const
+{
+	TArray<TSharedPtr<FItem>> SelectedItems = TreeView->GetSelectedItems();
+
+	TArray<FName> SelectedNames;
+	Algo::TransformIf(SelectedItems, SelectedNames, 
+		[this](const TSharedPtr<FItem>& Item)
+		{
+			return !GetItemName(Item).IsNone();
+		},
+		[this](const TSharedPtr<FItem>& Item)
+		{
+			return GetItemName(Item);
+		});
+
+	return SelectedNames;
+}
+
 void SReadOnlyHierarchyView::GetItemChildren(TSharedPtr<FItem> Item, TArray<TSharedPtr<FItem>>& OutChildren) const
 {
 	OutChildren.Append(Item->Children);
diff --git a/Engine/Source/Editor/UMGEditor/Public/Hierarchy/SReadOnlyHierarchyView.h b/Engine/Source/Editor/UMGEditor/Public/Hierarchy/SReadOnlyHierarchyView.h
--- a/Engine/Source/Editor/UMGEditor/Public/Hierarchy/SReadOnlyHierarchyView.h
+++ b/Engine/Source/Editor/UMGEditor/Public/Hierarchy/SReadOnlyHierarchyView.h
@@ -51,6 +51,8 @@ public:
 	void SetSelectedWidget(FName WidgetName);
 	void SetRawFilterText(const FText& Text);
 
+	TArray<FName> GetSelectedWidgets() const;
+
 private:
 
 	TSharedRef<ITableRow> GenerateRow(TSharedPtr<FItem> Item, const TSharedRef<STableViewBase>& OwnerTable) const;
@@ -70,6 +72,8 @@ private:
 	void SetItemExpansionRecursive(TSharedPtr<FItem> Item, bool bShouldBeExpanded);
 	void ExpandAll();
 
+	FName GetItemName(const TSharedPtr<FItem>& Item) const;
+
 private:
 
 	TWeakObjectPtr<const UWidgetBlueprint> WidgetBlueprint;
