FNames containing asset paths are deprecated. FSoftObjectPath or FTopLevelAssetPath should be used instead.
Fixups for UMGEditor.

diff --git a/Engine/Source/Editor/UMGEditor/Private/Palette/SPaletteViewModel.cpp b/Engine/Source/Editor/UMGEditor/Private/Palette/SPaletteViewModel.cpp
--- a/Engine/Source/Editor/UMGEditor/Private/Palette/SPaletteViewModel.cpp
+++ b/Engine/Source/Editor/UMGEditor/Private/Palette/SPaletteViewModel.cpp
@@ -38,8 +38,7 @@ namespace UE::Editor::SPaletteViewModel::Private
 	{
 	public:
 		FUnloadedBlueprintData(const FAssetData& InAssetData)
-			:ClassPath(NAME_None)
-			,ParentClassPath(NAME_None)
+			:ClassPath()
 			,ClassFlags(CLASS_None)
 			,bIsNormalBlueprintType(false)
 		{
@@ -49,17 +48,11 @@ namespace UE::Editor::SPaletteViewModel::Private
 			const UClass* AssetClass = InAssetData.GetClass();
 			if (AssetClass && AssetClass->IsChildOf(UBlueprintGeneratedClass::StaticClass()))
 			{
-				ClassPath = InAssetData.ObjectPath;
+				ClassPath = InAssetData.ToSoftObjectPath().GetAssetPathString();
 			}
 			else if (InAssetData.GetTagValue(FBlueprintTags::GeneratedClassPath, GeneratedClassPath))
 			{
-				ClassPath = FName(*FPackageName::ExportTextPathToObjectPath(GeneratedClassPath));
-			}
-
-			FString ParentClassPathString;
-			if (InAssetData.GetTagValue(FBlueprintTags::ParentClassPath, ParentClassPathString))
-			{
-				ParentClassPath = FName(*FPackageName::ExportTextPathToObjectPath(ParentClassPathString));
+				ClassPath = FTopLevelAssetPath(*FPackageName::ExportTextPathToObjectPath(GeneratedClassPath));
 			}
 
 			FEditorClassUtils::GetImplementedInterfaceClassPathsFromAsset(InAssetData, ImplementedInterfaces);
@@ -137,19 +130,20 @@ namespace UE::Editor::SPaletteViewModel::Private
 
 		virtual FName GetClassPath() const
 		{
-			return ClassPath;
+PRAGMA_DISABLE_DEPRECATION_WARNINGS
+			return ClassPath.ToFName();
+PRAGMA_ENABLE_DEPRECATION_WARNINGS
 		}
 
 		virtual FTopLevelAssetPath GetClassPathName() const
 		{
-			return FTopLevelAssetPath(ClassPath.ToString());
+			return ClassPath;
 		}
 		// End IUnloadedBlueprintData interface
 
 	private:
 		TSharedPtr<FString> ClassName;
-		FName ClassPath;
-		FName ParentClassPath;
+		FTopLevelAssetPath ClassPath;
 		uint32 ClassFlags;
 		TArray<FString> ImplementedInterfaces;
 		bool bIsNormalBlueprintType;
@@ -360,7 +354,7 @@ void FWidgetCatalogViewModel::BuildClassWidgetList()
 
 	auto WidgetPassesConfigFiltering = [&](const FAssetData& InWidgetAssetData, const FString& InCategoryName, TWeakObjectPtr<UClass> InWidgetClass)
 	{
-		if (AllowedPaletteWidgets.PassesFilter(InWidgetAssetData.ObjectPath))
+		if (AllowedPaletteWidgets.PassesFilter(InWidgetAssetData.GetObjectPathString()))
 		{
 			return true;
 		}
@@ -431,7 +425,7 @@ void FWidgetCatalogViewModel::BuildClassWidgetList()
 			bool bIsOnList = false;
 			for (FSoftClassPath Widget : WidgetClassesToHide)
 			{
-				if (WidgetAssetData.ObjectPath.ToString().Find(Widget.ToString()) == 0)
+				if (WidgetAssetData.GetObjectPathString().Find(Widget.ToString()) == 0)
 				{
 					bIsOnList = true;
 					break;
@@ -558,7 +552,7 @@ void FWidgetCatalogViewModel::BuildClassWidgetList()
 				bool bIsOnList = false;
 				for (FSoftClassPath Widget : WidgetClassesToHide)
 				{
-					if (Widget.ToString().Find(WidgetBPAssetData.ObjectPath.ToString()) == 0)
+					if (Widget.ToString().Find(WidgetBPAssetData.GetObjectPathString()) == 0)
 					{
 						bIsOnList = true;
 						break;
@@ -596,7 +590,7 @@ void FWidgetCatalogViewModel::BuildClassWidgetList()
 			}
 
 			if (!FilterAssetData(GeneratedBPAssetData) 
-				&& AllowedPaletteWidgets.PassesFilter(GeneratedBPAssetData.ObjectPath) 
+				&& AllowedPaletteWidgets.PassesFilter(GeneratedBPAssetData.GetObjectPathString()) 
 				&& !LoadedGeneratedBlueprintClassesByName.Contains(GeneratedBPAssetData.AssetName))
 			{
 				uint32 BPFlags = GeneratedBPAssetData.GetTagValueRef<uint32>(FBlueprintTags::ClassFlags);
diff --git a/Engine/Source/Editor/UMGEditor/Private/Templates/WidgetTemplateBlueprintClass.cpp b/Engine/Source/Editor/UMGEditor/Private/Templates/WidgetTemplateBlueprintClass.cpp
--- a/Engine/Source/Editor/UMGEditor/Private/Templates/WidgetTemplateBlueprintClass.cpp
+++ b/Engine/Source/Editor/UMGEditor/Private/Templates/WidgetTemplateBlueprintClass.cpp
@@ -70,7 +70,7 @@ UWidget* FWidgetTemplateBlueprintClass::Create(UWidgetTree* Tree)
 	// Load the blueprint asset or blueprint generated class if needed
 	if (!WidgetClass.Get())
 	{
-		FString AssetPath = WidgetAssetData.ObjectPath.ToString();
+		FString AssetPath = WidgetAssetData.GetObjectPathString();
 
 		if (bIsBlueprintGeneratedClass)
 		{
diff --git a/Engine/Source/Editor/UMGEditor/Private/Templates/WidgetTemplateClass.cpp b/Engine/Source/Editor/UMGEditor/Private/Templates/WidgetTemplateClass.cpp
--- a/Engine/Source/Editor/UMGEditor/Private/Templates/WidgetTemplateClass.cpp
+++ b/Engine/Source/Editor/UMGEditor/Private/Templates/WidgetTemplateClass.cpp
@@ -87,7 +87,7 @@ UWidget* FWidgetTemplateClass::Create(UWidgetTree* Tree)
 	// Load the blueprint asset if needed
 	if (!WidgetClass.Get())
 	{
-		FString AssetPath = WidgetAssetData.ObjectPath.ToString();
+		FString AssetPath = WidgetAssetData.GetObjectPathString();
 		UBlueprint* LoadedBP = LoadObject<UBlueprint>(nullptr, *AssetPath);
 		WidgetClass = *LoadedBP->GeneratedClass;
 	}
diff --git a/Engine/Source/Editor/UMGEditor/Private/Templates/WidgetTemplateImageClass.cpp b/Engine/Source/Editor/UMGEditor/Private/Templates/WidgetTemplateImageClass.cpp
--- a/Engine/Source/Editor/UMGEditor/Private/Templates/WidgetTemplateImageClass.cpp
+++ b/Engine/Source/Editor/UMGEditor/Private/Templates/WidgetTemplateImageClass.cpp
@@ -22,7 +22,7 @@ UWidget* FWidgetTemplateImageClass::Create(UWidgetTree* WidgetTree)
 	{
 		UImage* ImageWidget = Cast<UImage>(Widget);
 
-		UObject* ImageResource = FindObject<UObject>(nullptr, *WidgetAssetData.ObjectPath.ToString());
+		UObject* ImageResource = FindObject<UObject>(nullptr, *WidgetAssetData.GetObjectPathString());
 		ImageWidget->Brush.SetResourceObject(ImageResource);
 	}
 
diff --git a/Engine/Source/Editor/UMGEditor/Private/WidgetBlueprintEditorUtils.cpp b/Engine/Source/Editor/UMGEditor/Private/WidgetBlueprintEditorUtils.cpp
--- a/Engine/Source/Editor/UMGEditor/Private/WidgetBlueprintEditorUtils.cpp
+++ b/Engine/Source/Editor/UMGEditor/Private/WidgetBlueprintEditorUtils.cpp
@@ -889,7 +889,7 @@ void FWidgetBlueprintEditorUtils::BuildReplaceWithMenu(FMenuBuilder& Menu, TShar
 			UClass* WidgetClass = Widget.GetTemplate()->GetClass();
 			TWeakObjectPtr<UClass> TemplateWidget = BlueprintEditor->GetSelectedTemplate();
 			FAssetData SelectedUserWidget = BlueprintEditor->GetSelectedUserWidget();
-			if (TemplateWidget.IsValid() || SelectedUserWidget.ObjectPath != NAME_None)
+			if (TemplateWidget.IsValid() || SelectedUserWidget.GetSoftObjectPath().IsValid() )
 			{
 				Menu.AddMenuEntry(
 					FText::Format(LOCTEXT("WidgetTree_ReplaceWithSelection", "Replace With {0}"), FText::FromString(TemplateWidget.IsValid() ? TemplateWidget->GetName() : SelectedUserWidget.AssetName.ToString())),
@@ -986,7 +986,7 @@ void FWidgetBlueprintEditorUtils::ReplaceWidgetWithSelectedTemplate(TSharedRef<F
 		TSharedPtr<FWidgetTemplateClass> Template = MakeShareable(new FWidgetTemplateClass(WidgetClass));
 		NewReplacementWidget = Template->Create(BP->WidgetTree);
 	}
-	else if (BlueprintEditor->GetSelectedUserWidget().ObjectPath != NAME_None)
+	else if (BlueprintEditor->GetSelectedUserWidget().GetSoftObjectPath().IsValid())
 	{
 		bIsUserWidget = true;
 		FAssetData WidgetAssetData = BlueprintEditor->GetSelectedUserWidget();
@@ -1065,7 +1065,7 @@ bool FWidgetBlueprintEditorUtils::CanBeReplacedWithTemplate(TSharedRef<FWidgetBl
 	UWidget* ThisWidget = Widget.GetTemplate();
 	UPanelWidget* ExistingPanel = Cast<UPanelWidget>(ThisWidget);
 	// If selecting another widget blueprint
-	if (SelectedUserWidget.ObjectPath != NAME_None)
+	if (SelectedUserWidget.GetSoftObjectPath().IsValid())
 	{
 		if (ExistingPanel)
 		{
@@ -1422,7 +1422,7 @@ UWidget* FWidgetBlueprintEditorUtils::GetWidgetTemplateFromDragDrop(UWidgetBluep
 			const FAssetData& AssetData = AssetDragDropOp->GetAssets()[0];
 
 			bool CodeClass = AssetData.AssetClassPath == FTopLevelAssetPath(TEXT("/Script/CoreUObject"), TEXT("Class"));
-			FString ClassName = CodeClass ? AssetData.ObjectPath.ToString() : AssetData.AssetClassPath.ToString();
+			FString ClassName = CodeClass ? AssetData.GetObjectPathString() : AssetData.AssetClassPath.ToString();
 			UClass* AssetClass = FindObjectChecked<UClass>(nullptr, *ClassName);
 
 			if (FWidgetTemplateBlueprintClass::Supports(AssetClass))
@@ -1430,7 +1430,7 @@ UWidget* FWidgetBlueprintEditorUtils::GetWidgetTemplateFromDragDrop(UWidgetBluep
 				// Allows a UMG Widget Blueprint to be dragged from the Content Browser to another Widget Blueprint...as long as we're not trying to place a
 				// blueprint inside itself.
 				FString BlueprintPath = Blueprint->GetPathName();
-				if (BlueprintPath != AssetData.ObjectPath.ToString())
+				if (BlueprintPath != AssetData.GetSoftObjectPath().ToString())
 				{
 					Widget = FWidgetTemplateBlueprintClass(AssetData).Create(RootWidgetTree);
 				}
