UMG: Add stackbox, a panel that behave like a vertical or horizontal box base. The direction can be changed at runtime.

diff --git a/Engine/Source/Editor/UMGEditor/Private/Designer/SDesignerView.cpp b/Engine/Source/Editor/UMGEditor/Private/Designer/SDesignerView.cpp
--- a/Engine/Source/Editor/UMGEditor/Private/Designer/SDesignerView.cpp
+++ b/Engine/Source/Editor/UMGEditor/Private/Designer/SDesignerView.cpp
@@ -30,6 +30,7 @@
 #include "Extensions/CanvasSlotExtension.h"
 #include "Extensions/GridSlotExtension.h"
 #include "Extensions/HorizontalSlotExtension.h"
+#include "Extensions/StackBoxSlotExtension.h"
 #include "Extensions/UniformGridSlotExtension.h"
 #include "Extensions/VerticalSlotExtension.h"
 #include "Designer/SPaintSurface.h"
@@ -1513,6 +1514,7 @@ void SDesignerView::RegisterExtensions()
 {
 	Register(MakeShareable(new FVerticalSlotExtension()));
 	Register(MakeShareable(new FHorizontalSlotExtension()));
+	Register(MakeShareable(new FStackBoxSlotExtension()));
 	Register(MakeShareable(new FCanvasSlotExtension()));
 	Register(MakeShareable(new FUniformGridSlotExtension()));
 	Register(MakeShareable(new FGridSlotExtension()));
diff --git a/Engine/Source/Editor/UMGEditor/Private/Extensions/StackBoxSlotExtension.cpp b/Engine/Source/Editor/UMGEditor/Private/Extensions/StackBoxSlotExtension.cpp
new file mode 100644
--- /dev/null
+++ b/Engine/Source/Editor/UMGEditor/Private/Extensions/StackBoxSlotExtension.cpp
@@ -0,0 +1,125 @@
+// Copyright Epic Games, Inc. All Rights Reserved.
+
+#include "Extensions/StackBoxSlotExtension.h"
+#include "Engine/GameViewportClient.h"
+#include "Components/StackBox.h"
+#include "Components/StackBoxSlot.h"
+#include "Kismet2/BlueprintEditorUtils.h"
+#include "Misc/Optional.h"
+#include "WidgetBlueprint.h"
+#include "Widgets/Input/SButton.h"
+
+#define LOCTEXT_NAMESPACE "UMG"
+
+FStackBoxSlotExtension::FStackBoxSlotExtension()
+{
+	ExtensionId = FName(TEXT("StackBoxSlot"));
+}
+
+bool FStackBoxSlotExtension::CanExtendSelection(const TArray< FWidgetReference >& Selection) const
+{
+	for ( const FWidgetReference& Widget : Selection )
+	{
+		if ( !Widget.GetTemplate()->Slot || !Widget.GetTemplate()->Slot->IsA(UStackBoxSlot::StaticClass()) )
+		{
+			return false;
+		}
+	}
+
+	return Selection.Num() == 1;
+}
+
+void FStackBoxSlotExtension::ExtendSelection(const TArray< FWidgetReference >& Selection, TArray< TSharedRef<FDesignerSurfaceElement> >& SurfaceElements)
+{
+	bool bShowWidgets = false;
+	TOptional<EOrientation> StackBoxType;
+	for (const FWidgetReference& Widget : Selection)
+	{
+		check(Widget.GetTemplate());
+		check(Widget.GetTemplate()->Slot);
+		UStackBox* StackBox = CastChecked<UStackBox>(Widget.GetTemplate()->Slot->Parent);
+		if (!StackBoxType.IsSet())
+		{
+			StackBoxType = StackBox->GetOrientation();
+			bShowWidgets = true;
+		}
+		else if (StackBoxType.GetValue() != StackBox->GetOrientation())
+		{
+			bShowWidgets = false;
+			break;
+		}
+	}
+
+	SelectionCache = Selection;
+
+	if (StackBoxType.IsSet())
+	{
+		TSharedPtr<SButton> ArrowA;
+		TSharedPtr<SButton> ArrowB;
+		if (StackBoxType == EOrientation::Orient_Vertical)
+		{
+			ArrowA = SNew(SButton)
+				.Text(LOCTEXT("UpArrow", "\u2191"))
+				.ContentPadding(FMargin(6.f, 2.f))
+				.IsEnabled(bShowWidgets)
+				.OnClicked(this, &FStackBoxSlotExtension::HandleShiftDirection, -1);
+
+			ArrowB = SNew(SButton)
+				.Text(LOCTEXT("DownArrow", "\u2193"))
+				.ContentPadding(FMargin(6.f, 2.f))
+				.IsEnabled(bShowWidgets)
+				.OnClicked(this, &FStackBoxSlotExtension::HandleShiftDirection, 1);
+		}
+		else
+		{
+			ArrowA = SNew(SButton)
+				.Text(LOCTEXT("LeftArrow", "\u2190"))
+				.ContentPadding(FMargin(2.f, 6.f))
+				.IsEnabled(bShowWidgets)
+				.OnClicked(this, &FStackBoxSlotExtension::HandleShiftDirection, -1);
+
+			ArrowB = SNew(SButton)
+				.Text(LOCTEXT("RightArrow", "\u2192"))
+				.ContentPadding(FMargin(2.f, 6.f))
+				.IsEnabled(bShowWidgets)
+				.OnClicked(this, &FStackBoxSlotExtension::HandleShiftDirection, 1);
+		}
+
+		ArrowA->SlatePrepass();
+		ArrowB->SlatePrepass();
+
+		SurfaceElements.Add(MakeShareable(new FDesignerSurfaceElement(ArrowA.ToSharedRef(), EExtensionLayoutLocation::TopCenter, FVector2D(ArrowA->GetDesiredSize().X * -0.5f, -ArrowA->GetDesiredSize().Y))));
+		SurfaceElements.Add(MakeShareable(new FDesignerSurfaceElement(ArrowB.ToSharedRef(), EExtensionLayoutLocation::BottomCenter, FVector2D(ArrowB->GetDesiredSize().X * -0.5f, 0))));
+	}
+}
+
+FReply FStackBoxSlotExtension::HandleShiftDirection(int32 ShiftAmount)
+{
+	BeginTransaction(LOCTEXT("MoveWidget", "Move Widget"));
+
+	for ( FWidgetReference& Selection : SelectionCache )
+	{
+		ShiftDirection(Selection.GetPreview(), ShiftAmount);
+		ShiftDirection(Selection.GetTemplate(), ShiftAmount);
+	}
+
+	EndTransaction();
+
+	if (UWidgetBlueprint* BlueprintPtr = Blueprint.Get())
+	{
+		FBlueprintEditorUtils::MarkBlueprintAsStructurallyModified(BlueprintPtr);
+	}
+
+	return FReply::Handled();
+}
+
+void FStackBoxSlotExtension::ShiftDirection(UWidget* Widget, int32 ShiftAmount)
+{
+	UStackBox* Parent = CastChecked<UStackBox>(Widget->GetParent());
+
+	Parent->Modify();
+	int32 CurrentIndex = Parent->GetChildIndex(Widget);
+	Parent->ShiftChild(CurrentIndex + ShiftAmount, Widget);
+}
+
+#undef LOCTEXT_NAMESPACE
diff --git a/Engine/Source/Editor/UMGEditor/Private/Extensions/StackBoxSlotExtension.h b/Engine/Source/Editor/UMGEditor/Private/Extensions/StackBoxSlotExtension.h
new file mode 100644
--- /dev/null
+++ b/Engine/Source/Editor/UMGEditor/Private/Extensions/StackBoxSlotExtension.h
@@ -0,0 +1,23 @@
+// Copyright Epic Games, Inc. All Rights Reserved.
+
+#pragma once
+
+#include "CoreMinimal.h"
+#include "Input/Reply.h"
+#include "WidgetReference.h"
+#include "DesignerExtension.h"
+
+class FStackBoxSlotExtension : public FDesignerExtension
+{
+public:
+	FStackBoxSlotExtension();
+
+	virtual ~FStackBoxSlotExtension() {}
+
+	virtual bool CanExtendSelection(const TArray< FWidgetReference >& Selection) const override;
+	virtual void ExtendSelection(const TArray< FWidgetReference >& Selection, TArray< TSharedRef<FDesignerSurfaceElement> >& SurfaceElements) override;
+
+private:
+	FReply HandleShiftDirection(int32 ShiftAmount);
+	void ShiftDirection(UWidget* Widget, int32 ShiftAmount);
+};
diff --git a/Engine/Source/Runtime/SlateCore/Private/Widgets/SBoxPanel.cpp b/Engine/Source/Runtime/SlateCore/Private/Widgets/SBoxPanel.cpp
--- a/Engine/Source/Runtime/SlateCore/Private/Widgets/SBoxPanel.cpp
+++ b/Engine/Source/Runtime/SlateCore/Private/Widgets/SBoxPanel.cpp
@@ -20,6 +20,11 @@ void SVerticalBox::PrivateRegisterAttributes(FSlateAttributeInitializer& Attribu
 {
 }
 
+SLATE_IMPLEMENT_WIDGET(SStackBox)
+void SStackBox::PrivateRegisterAttributes(FSlateAttributeInitializer& AttributeInitializer)
+{
+}
+
 void SHorizontalBox::Construct( const SHorizontalBox::FArguments& InArgs )
 {
 	Children.Reserve(InArgs._Slots.Num());
@@ -47,6 +52,17 @@ void SVerticalBox::Construct( const SVerticalBox::FArguments& InArgs )
 	}
 }
 
+void SStackBox::Construct(const SStackBox::FArguments& InArgs)
+{
+	Children.Reserve(InArgs._Slots.Num());
+	for (const FSlot::FSlotArguments& Arg : InArgs._Slots)
+	{
+		const FSlotBase::FSlotArguments& ChilSlotArgument = static_cast<const FSlotBase::FSlotArguments&>(Arg);
+		const SBoxPanel::FSlot::FSlotArguments& BoxSlotArgument = static_cast<const SBoxPanel::FSlot::FSlotArguments&>(ChilSlotArgument);
+		Children.AddSlot(MoveTemp(const_cast<SBoxPanel::FSlot::FSlotArguments&>(BoxSlotArgument)));
+	}
+}
+
 SHorizontalBox::FSlot& SHorizontalBox::GetSlot(int32 SlotIndex)
 {
 	check(this->IsValidSlotIndex(SlotIndex));
@@ -75,6 +91,20 @@ const SVerticalBox::FSlot& SVerticalBox::GetSlot(int32 SlotIndex) const
 	return static_cast<const SVerticalBox::FSlot&>(BaseSlot);
 }
 
+SStackBox::FSlot& SStackBox::GetSlot(int32 SlotIndex)
+{
+	check(this->IsValidSlotIndex(SlotIndex));
+	FSlotBase& BaseSlot = static_cast<FSlotBase&>(Children[SlotIndex]);
+	return static_cast<SStackBox::FSlot&>(BaseSlot);
+}
+
+const SStackBox::FSlot& SStackBox::GetSlot(int32 SlotIndex) const
+{
+	check(this->IsValidSlotIndex(SlotIndex));
+	const FSlotBase& BaseSlot = static_cast<const FSlotBase&>(Children[SlotIndex]);
+	return static_cast<const SStackBox::FSlot&>(BaseSlot);
+}
+
 template<EOrientation Orientation, typename SlotType>
 static void ArrangeChildrenAlong(EFlowDirection InLayoutFlow, const TPanelChildren<SlotType>& Children, const FGeometry& AllottedGeometry, FArrangedChildren& ArrangedChildren )
 {
@@ -310,6 +340,22 @@ void SBoxPanel::ClearChildren()
 	Children.Empty();
 }
 
+void SBoxPanel::SetOrientation(EOrientation InOrientation)
+{
+	if (Orientation != InOrientation)
+	{
+		Orientation = InOrientation;
+		Invalidate(EInvalidateWidgetReason::Layout);
+	}
+}
+
+SBoxPanel::SBoxPanel()
+	: Children(this, GET_MEMBER_NAME_CHECKED(SBoxPanel, Children))
+	, Orientation(EOrientation::Orient_Horizontal)
+{
+
+}
+
 SBoxPanel::SBoxPanel( EOrientation InOrientation )
 	: Children(this, GET_MEMBER_NAME_CHECKED(SBoxPanel, Children))
 	, Orientation(InOrientation)
@@ -317,6 +363,7 @@ SBoxPanel::SBoxPanel( EOrientation InOrientation )
 
 }
 
+
 void SDragAndDropVerticalBox::Construct(const FArguments& InArgs)
 {
 	SVerticalBox::Construct(SVerticalBox::FArguments());
diff --git a/Engine/Source/Runtime/SlateCore/Public/Widgets/SBoxPanel.h b/Engine/Source/Runtime/SlateCore/Public/Widgets/SBoxPanel.h
--- a/Engine/Source/Runtime/SlateCore/Public/Widgets/SBoxPanel.h
+++ b/Engine/Source/Runtime/SlateCore/Public/Widgets/SBoxPanel.h
@@ -195,26 +195,33 @@ public:
 	/** @return if it's a valid index slot index. */
 	bool IsValidSlotIndex(int32 Index) const { return Children.IsValidIndex(Index); }
 
-public:
 	//~ Begin SWidget overrides.
+public:
 	virtual void OnArrangeChildren(const FGeometry& AllottedGeometry, FArrangedChildren& ArrangedChildren) const override;
 	virtual FChildren* GetChildren() override;
+
 protected:
 	virtual FVector2D ComputeDesiredSize(float) const override;
 	//~ End SWidget overrides.
 
-	/**
-	 * A Box Panel's orientation cannot be changed once it is constructed..
-	 *
-	 * @param InOrientation   The orientation of the Box Panel
-	 */
-	SBoxPanel( EOrientation InOrientation );
+protected:
+	SBoxPanel();
+	/** A Box Panel's orientation. */
+	SBoxPanel(EOrientation InOrientation);
+
+	/** Set the orientation of the Box. It will do a full invalidation of the widget. */
+	void SetOrientation(EOrientation InOrientation);
+	EOrientation GetOrientation() const
+	{
+		return Orientation;
+	}
 
+protected:
 	/** The Box Panel's children. */
 	TPanelChildren<FSlot> Children;
 
 	/** The Box Panel's orientation; determined at construct time. */
-	const EOrientation Orientation;
+	EOrientation Orientation;
 };
 
 
@@ -425,6 +432,86 @@ public:
 	void Construct( const FArguments& InArgs );
 };
 
+/** A Stack Box Panel that stack vertically or horizontally. See SBoxPanel for more info. */
+class SLATECORE_API SStackBox : public SBoxPanel
+{
+	SLATE_DECLARE_WIDGET(SStackBox, SBoxPanel)
+public:
+	class FSlot : public SBoxPanel::TSlot<FSlot>
+	{
+	public:
+		SLATE_SLOT_BEGIN_ARGS(FSlot, SBoxPanel::TSlot<FSlot>)
+			/** The widget's DesiredSize will be used as the space required. */
+			FSlotArguments& AutoSize()
+			{
+				_SizeParam = FAuto();
+				return Me();
+			}
+			/** The available space will be distributed proportionately. */
+			FSlotArguments& FillSize(TAttribute<float> InStretchCoefficient)
+			{
+				_SizeParam = FStretch(MoveTemp(InStretchCoefficient));
+				return Me();
+			}
+			/** Set the max size in SlateUnit this slot can be. */
+			FSlotArguments& MaxSize(TAttribute<float> InMaxHeight)
+			{
+				_MaxSize = MoveTemp(InMaxHeight);
+				return Me();
+			}
+		SLATE_SLOT_END_ARGS()
+
+		void Construct(const FChildren& SlotOwner, FSlotArguments&& InArgs)
+		{
+			SBoxPanel::TSlot<FSlot>::Construct(SlotOwner, MoveTemp(InArgs));
+		}
+	};
+
+	static FSlot::FSlotArguments Slot()
+	{
+		return FSlot::FSlotArguments(MakeUnique<FSlot>());
+	}
+
+
+	SLATE_BEGIN_ARGS(SStackBox)
+	{
+		_Visibility = EVisibility::SelfHitTestInvisible;
+	}
+		SLATE_ARGUMENT(EOrientation, Orientation)
+		SLATE_SLOT_ARGUMENT(SStackBox::FSlot, Slots)
+	SLATE_END_ARGS()
+
+	using FScopedWidgetSlotArguments = SBoxPanel::FScopedWidgetSlotArguments<SStackBox::FSlot>;
+	FScopedWidgetSlotArguments AddSlot()
+	{
+		return InsertSlot(INDEX_NONE);
+	}
+
+	FScopedWidgetSlotArguments InsertSlot(int32 Index = INDEX_NONE)
+	{
+		return FScopedWidgetSlotArguments(MakeUnique<FSlot>(), this->Children, Index);
+	}
+
+	FSlot& GetSlot(int32 SlotIndex);
+	const FSlot& GetSlot(int32 SlotIndex) const;
+
+	FORCENOINLINE SStackBox()
+	{
+		SetCanTick(false);
+		bCanSupportFocus = false;
+	}
+
+	/**
+	 * Construct this widget
+	 *
+	 * @param	InArgs	The declaration data for this widget
+	 */
+	void Construct(const FArguments& InArgs);
+
+	using SBoxPanel::SetOrientation;
+	using SBoxPanel::GetOrientation;
+};
+
 class FDragAndDropVerticalBoxOp : public FDragDropOperation
 {
 public:
diff --git a/Engine/Source/Runtime/UMG/Private/Components/StackBox.cpp b/Engine/Source/Runtime/UMG/Private/Components/StackBox.cpp
new file mode 100644
--- /dev/null
+++ b/Engine/Source/Runtime/UMG/Private/Components/StackBox.cpp
@@ -0,0 +1,110 @@
+// Copyright Epic Games, Inc. All Rights Reserved.
+
+#include "Components/StackBox.h"
+#include "Components/StackBoxSlot.h"
+
+#include "Widgets/SBoxPanel.h"
+#include "Widgets/SOverlay.h"
+
+#define LOCTEXT_NAMESPACE "UMG"
+
+/////////////////////////////////////////////////////
+// UStackBox
+
+UStackBox::UStackBox()
+{
+	bIsVariable = false;
+	SetVisibilityInternal(ESlateVisibility::SelfHitTestInvisible);
+}
+
+void UStackBox::ReleaseSlateResources(bool bReleaseChildren)
+{
+	Super::ReleaseSlateResources(bReleaseChildren);
+
+	MyBox.Reset();
+}
+
+UClass* UStackBox::GetSlotClass() const
+{
+	return UStackBoxSlot::StaticClass();
+}
+
+void UStackBox::OnSlotAdded(UPanelSlot* InSlot)
+{
+	// Add the child to the live canvas if it already exists
+	if (MyBox.IsValid() )
+	{
+		CastChecked<UStackBoxSlot>(InSlot)->BuildSlot(MyBox.ToSharedRef());
+	}
+}
+
+void UStackBox::OnSlotRemoved(UPanelSlot* InSlot)
+{
+	// Remove the widget from the live slot if it exists.
+	if (MyBox.IsValid() && InSlot->Content)
+	{
+		TSharedPtr<SWidget> Widget = InSlot->Content->GetCachedWidget();
+		if ( Widget.IsValid() )
+		{
+			MyBox->RemoveSlot(Widget.ToSharedRef());
+		}
+	}
+}
+
+EOrientation UStackBox::GetOrientation() const
+{
+	return Orientation;
+}
+
+void UStackBox::SetOrientation(EOrientation InOrientation)
+{
+	if (Orientation != InOrientation)
+	{
+		Orientation = InOrientation;
+		if (MyBox)
+		{
+			MyBox->SetOrientation(InOrientation);
+		}
+	}
+}
+
+UStackBoxSlot* UStackBox::AddChildToStackBox(UWidget* Content)
+{
+	return Cast<UStackBoxSlot>(Super::AddChild(Content));
+}
+
+TSharedRef<SWidget> UStackBox::RebuildWidget()
+{
+	MyBox = SNew(SStackBox)
+		.Orientation(Orientation);
+
+	for ( UPanelSlot* PanelSlot : Slots )
+	{
+		if ( UStackBoxSlot* TypedSlot = Cast<UStackBoxSlot>(PanelSlot) )
+		{
+			TypedSlot->Parent = this;
+			TypedSlot->BuildSlot(MyBox.ToSharedRef());
+		}
+	}
+
+	return MyBox.ToSharedRef();
+}
+
+void UStackBox::SynchronizeProperties()
+{
+	Super::SynchronizeProperties();
+	MyBox->SetOrientation(Orientation);
+}
+
+#if WITH_EDITOR
+
+const FText UStackBox::GetPaletteCategory()
+{
+	return LOCTEXT("Panel", "Panel");
+}
+
+#endif
+
+/////////////////////////////////////////////////////
+
+#undef LOCTEXT_NAMESPACE
diff --git a/Engine/Source/Runtime/UMG/Private/Components/StackBoxSlot.cpp b/Engine/Source/Runtime/UMG/Private/Components/StackBoxSlot.cpp
new file mode 100644
--- /dev/null
+++ b/Engine/Source/Runtime/UMG/Private/Components/StackBoxSlot.cpp
@@ -0,0 +1,131 @@
+// Copyright Epic Games, Inc. All Rights Reserved.
+
+#include "Components/StackBoxSlot.h"
+#include "Components/Widget.h"
+#include "Components/StackBox.h"
+
+#include "Widgets/SBoxPanel.h"
+#include "Widgets/SNullWidget.h"
+
+/////////////////////////////////////////////////////
+// UStackBoxSlot
+
+void UStackBoxSlot::ReleaseSlateResources(bool bReleaseChildren)
+{
+	Super::ReleaseSlateResources(bReleaseChildren);
+
+	Slot = nullptr;
+}
+
+void UStackBoxSlot::BuildSlot(TSharedRef<SStackBox> StackBox)
+{
+	StackBox->AddSlot()
+		.Expose(Slot)
+		.Padding(Padding)
+		.HAlign(HorizontalAlignment)
+		.VAlign(VerticalAlignment)
+		.SizeParam(UWidget::ConvertSerializedSizeParamToRuntime(Size))
+		[
+			Content == nullptr ? SNullWidget::NullWidget : Content->TakeWidget()
+		];
+}
+
+FMargin UStackBoxSlot::GetPadding() const
+{
+	return Padding;
+}
+
+void UStackBoxSlot::SetPadding(FMargin InPadding)
+{
+	Padding = InPadding;
+	if (Slot)
+	{
+		Slot->SetPadding(InPadding);
+	}
+}
+
+FSlateChildSize UStackBoxSlot::GetSize() const
+{
+	return Size;
+}
+
+void UStackBoxSlot::SetSize(FSlateChildSize InSize)
+{
+	Size = InSize;
+	if (Slot)
+	{
+		Slot->SetSizeParam(UWidget::ConvertSerializedSizeParamToRuntime(InSize));
+	}
+}
+
+EHorizontalAlignment UStackBoxSlot::GetHorizontalAlignment() const
+{
+	return HorizontalAlignment;
+}
+
+void UStackBoxSlot::SetHorizontalAlignment(EHorizontalAlignment InHorizontalAlignment)
+{
+	HorizontalAlignment = InHorizontalAlignment;
+	if (Slot)
+	{
+		Slot->SetHorizontalAlignment(InHorizontalAlignment);
+	}
+}
+
+EVerticalAlignment UStackBoxSlot::GetVerticalAlignment() const
+{
+	return VerticalAlignment;
+}
+
+void UStackBoxSlot::SetVerticalAlignment(EVerticalAlignment InVerticalAlignment)
+{
+	VerticalAlignment = InVerticalAlignment;
+	if (Slot)
+	{
+		Slot->SetVerticalAlignment(InVerticalAlignment);
+	}
+}
+
+void UStackBoxSlot::SynchronizeProperties()
+{
+	SetPadding(Padding);
+	SetSize(Size);
+	SetHorizontalAlignment(HorizontalAlignment);
+	SetVerticalAlignment(VerticalAlignment);
+}
+
+#if WITH_EDITOR
+bool UStackBoxSlot::NudgeByDesigner(const FVector2D& NudgeDirection, const TOptional<int32>& GridSnapSize)
+{
+	const FVector2D ClampedDirection = NudgeDirection.ClampAxes(-1, 1);
+	UStackBox* ParentStackBox = CastChecked<UStackBox>(Parent);
+	const EOrientation Orientation = ParentStackBox->GetOrientation();
+	const int32 UnitDirection = Orientation == EOrientation::Orient_Horizontal ? (int32)ClampedDirection.X : (int32)ClampedDirection.Y;
+	if (UnitDirection == 0)
+	{
+		return false;
+	}
+
+	const int32 CurrentIndex = ParentStackBox->GetChildIndex(Content);
+
+	if ((CurrentIndex == 0 && UnitDirection < 0) ||
+		(CurrentIndex + 1 >= ParentStackBox->GetChildrenCount() && UnitDirection > 0))
+	{
+		return false;
+	}
+
+	ParentStackBox->Modify();
+	ParentStackBox->ShiftChild(CurrentIndex + UnitDirection, Content);
+
+	return true;
+}
+
+void UStackBoxSlot::SynchronizeFromTemplate(const UPanelSlot* const TemplateSlot)
+{
+	const ThisClass* const TemplateStackBoxSlot = CastChecked<ThisClass>(TemplateSlot);
+	const int32 CurrentIndex = TemplateStackBoxSlot->Parent->GetChildIndex(TemplateStackBoxSlot->Content);
+
+	UStackBox* ParentStackBox = CastChecked<UStackBox>(Parent);
+	ParentStackBox->ShiftChild(CurrentIndex, Content);
+}
+#endif
diff --git a/Engine/Source/Runtime/UMG/Public/Components/StackBox.h b/Engine/Source/Runtime/UMG/Public/Components/StackBox.h
new file mode 100644
--- /dev/null
+++ b/Engine/Source/Runtime/UMG/Public/Components/StackBox.h
@@ -0,0 +1,61 @@
+// Copyright Epic Games, Inc. All Rights Reserved.
+
+#pragma once
+
+#include "CoreMinimal.h"
+#include "Components/PanelWidget.h"
+#include "Types/SlateEnums.h"
+#include "StackBox.generated.h"
+
+class UStackBoxSlot;
+class SStackBox;
+
+/**
+ * A stack box widget is a layout panel allowing child widgets to be automatically laid out
+ * vertically or horizontally.
+ *
+ * * Many Children
+ * * Flows Vertical or Horizontal
+ */
+UCLASS(meta = (ShortTooltip = "A layout panel for automatically laying child widgets out vertically or horizontally"))
+class UMG_API UStackBox : public UPanelWidget
+{
+	GENERATED_BODY()
+
+	UStackBox();
+
+private:
+	/** The orientation of the stack box. */
+	UPROPERTY(EditAnywhere, BlueprintReadWrite, Setter, Getter, Category = "Layout", meta = (AllowPrivateAccess = "true"))
+	TEnumAsByte<EOrientation> Orientation;
+
+public:
+	/** Get the orientation of the stack box. */
+	EOrientation GetOrientation() const;
+	/** Set the orientation of the stack box. The existing elements will be rearranged. */
+	void SetOrientation(EOrientation InType);
+
+	/** Adds a new child widget to the container. */
+	UFUNCTION(BlueprintCallable, Category="Panel")
+	UStackBoxSlot* AddChildToStackBox(UWidget* Content);
+
+#if WITH_EDITOR
+	//~ UWidget interface
+	virtual const FText GetPaletteCategory() override;
+	//~ End UWidget interface
+#endif
+
+protected:
+	//~ UPanelWidget
+	virtual UClass* GetSlotClass() const override;
+	virtual void OnSlotAdded(UPanelSlot* Slot) override;
+	virtual void OnSlotRemoved(UPanelSlot* Slot) override;
+	virtual TSharedRef<SWidget> RebuildWidget() override;
+	virtual void SynchronizeProperties() override;
+	//~ End UPanelWidget
+
+	virtual void ReleaseSlateResources(bool bReleaseChildren) override;
+
+private:
+	TSharedPtr<SStackBox> MyBox;
+};
diff --git a/Engine/Source/Runtime/UMG/Public/Components/StackBoxSlot.h b/Engine/Source/Runtime/UMG/Public/Components/StackBoxSlot.h
new file mode 100644
--- /dev/null
+++ b/Engine/Source/Runtime/UMG/Public/Components/StackBoxSlot.h
@@ -0,0 +1,69 @@
+// Copyright Epic Games, Inc. All Rights Reserved.
+
+#pragma once
+
+#include "CoreMinimal.h"
+#include "Layout/Margin.h"
+#include "Components/SlateWrapperTypes.h"
+#include "Components/PanelSlot.h"
+#include "Widgets/SBoxPanel.h"
+
+#include "StackBoxSlot.generated.h"
+
+class SPanel;
+
+/** The Slot for the UStackBox, contains the widget that is flowed vertically or horizontally. */
+UCLASS()
+class UMG_API UStackBoxSlot : public UPanelSlot
+{
+	GENERATED_BODY()
+
+private:
+	/** The padding area between the slot and the content it contains. */
+	UPROPERTY(EditAnywhere, BlueprintReadWrite, Setter, Getter, Category="Layout|Stack Box Slot", meta = (AllowPrivateAccess = "true"))
+	FMargin Padding;
+
+	/** How much space this slot should occupy in the direction of the panel. */
+	UPROPERTY(EditAnywhere, BlueprintReadWrite, Setter, Getter, Category = "Layout|Stack Box Slot", meta = (AllowPrivateAccess = "true"))
+	FSlateChildSize Size = FSlateChildSize(ESlateSizeRule::Automatic);
+
+	/** The alignment of the object horizontally. */
+	UPROPERTY(EditAnywhere, BlueprintReadWrite, Setter, Getter, Category="Layout|Stack Box Slot", meta = (AllowPrivateAccess = "true"))
+	TEnumAsByte<EHorizontalAlignment> HorizontalAlignment = EHorizontalAlignment::HAlign_Fill;
+
+	/** The alignment of the object vertically. */
+	UPROPERTY(EditAnywhere, BlueprintReadWrite, Setter, Getter, Category="Layout|Stack Box Slot", meta = (AllowPrivateAccess = "true"))
+	TEnumAsByte<EVerticalAlignment> VerticalAlignment = EVerticalAlignment::VAlign_Fill;
+
+private:
+	/** A raw pointer to the slot to allow us to adjust the size, padding...etc at runtime. */
+	SStackBox::FSlot* Slot = nullptr;
+
+public:
+	FMargin GetPadding() const;
+	void SetPadding(FMargin InPadding);
+
+	FSlateChildSize GetSize() const;
+	void SetSize(FSlateChildSize InSize);
+
+	EHorizontalAlignment GetHorizontalAlignment() const;
+	void SetHorizontalAlignment(EHorizontalAlignment InHorizontalAlignment);
+
+	EVerticalAlignment GetVerticalAlignment() const;
+	void SetVerticalAlignment(EVerticalAlignment InVerticalAlignment);
+
+public:
+	//~ UPanelSlot interface
+	virtual void SynchronizeProperties() override;
+	//~ End of UPanelSlot interface
+
+	virtual void ReleaseSlateResources(bool bReleaseChildren) override;
+
+	/** Builds the underlying FSlot for the Slate layout panel. */
+	void BuildSlot(TSharedRef<SStackBox> InBox);
+
+#if WITH_EDITOR
+	virtual bool NudgeByDesigner(const FVector2D& NudgeDirection, const TOptional<int32>& GridSnapSize) override;
+	virtual void SynchronizeFromTemplate(const UPanelSlot* const TemplateSlot) override;
+#endif //WITH_EDITOR
+};
