Added ClearSelection() to SReadOnlyHierarchyView.

#rnx

diff --git a/Engine/Source/Editor/UMGEditor/Private/Hierarchy/SReadOnlyHierarchyView.cpp b/Engine/Source/Editor/UMGEditor/Private/Hierarchy/SReadOnlyHierarchyView.cpp
--- a/Engine/Source/Editor/UMGEditor/Private/Hierarchy/SReadOnlyHierarchyView.cpp
+++ b/Engine/Source/Editor/UMGEditor/Private/Hierarchy/SReadOnlyHierarchyView.cpp
@@ -217,6 +217,10 @@ void SReadOnlyHierarchyView::SetSelectedWidget(FName WidgetName)
 	{
 		TreeView->SetSelection(Found);
 	}
+	else
+	{
+		TreeView->ClearSelection();
+	}
 }
 
 TArray<FName> SReadOnlyHierarchyView::GetSelectedWidgets() const
@@ -237,6 +241,11 @@ TArray<FName> SReadOnlyHierarchyView::GetSelectedWidgets() const
 	return SelectedNames;
 }
 
+void SReadOnlyHierarchyView::ClearSelection()
+{
+	TreeView->ClearSelection();
+}
+
 void SReadOnlyHierarchyView::GetItemChildren(TSharedPtr<FItem> Item, TArray<TSharedPtr<FItem>>& OutChildren) const
 {
 	OutChildren.Append(Item->Children);
diff --git a/Engine/Source/Editor/UMGEditor/Public/Hierarchy/SReadOnlyHierarchyView.h b/Engine/Source/Editor/UMGEditor/Public/Hierarchy/SReadOnlyHierarchyView.h
--- a/Engine/Source/Editor/UMGEditor/Public/Hierarchy/SReadOnlyHierarchyView.h
+++ b/Engine/Source/Editor/UMGEditor/Public/Hierarchy/SReadOnlyHierarchyView.h
@@ -52,6 +52,7 @@ public:
 	void SetRawFilterText(const FText& Text);
 
 	TArray<FName> GetSelectedWidgets() const;
+	void ClearSelection();
 
 private:
 
