UMG Designer: Fix to "Replace With" when a Blueprint extending a Widget is selected in the Library.

diff --git a/Engine/Source/Editor/UMGEditor/Private/WidgetBlueprintEditorUtils.cpp b/Engine/Source/Editor/UMGEditor/Private/WidgetBlueprintEditorUtils.cpp
--- a/Engine/Source/Editor/UMGEditor/Private/WidgetBlueprintEditorUtils.cpp
+++ b/Engine/Source/Editor/UMGEditor/Private/WidgetBlueprintEditorUtils.cpp
@@ -1074,10 +1074,12 @@ bool FWidgetBlueprintEditorUtils::CanBeReplacedWithTemplate(TSharedRef<FWidgetBl
 				return false;
 			}
 		}
-		UUserWidget* NewUserWidget = CastChecked<UUserWidget>(FWidgetTemplateBlueprintClass(SelectedUserWidget).Create(BP->WidgetTree));
-		const bool bFreeFromCircularRefs = BP->IsWidgetFreeFromCircularReferences(NewUserWidget);
-		NewUserWidget->Rename(nullptr, GetTransientPackage());
-		return bFreeFromCircularRefs;
+		if (UUserWidget* NewUserWidget = CastChecked<UUserWidget>(FWidgetTemplateBlueprintClass(SelectedUserWidget).Create(BP->WidgetTree)))
+		{
+			const bool bFreeFromCircularRefs = BP->IsWidgetFreeFromCircularReferences(NewUserWidget);
+			NewUserWidget->Rename(nullptr, GetTransientPackage());
+			return bFreeFromCircularRefs;
+		}
 	}
 
 	UClass* WidgetClass = BlueprintEditor->GetSelectedTemplate().Get();
