SReadOnlyHierarchyView can now be constructed with a ShowOnly array of UWidget*, to restrict the view to only display a subset of the hierarchy.

#rnx

diff --git a/Engine/Source/Editor/UMGEditor/Private/Hierarchy/SReadOnlyHierarchyView.cpp b/Engine/Source/Editor/UMGEditor/Private/Hierarchy/SReadOnlyHierarchyView.cpp
--- a/Engine/Source/Editor/UMGEditor/Private/Hierarchy/SReadOnlyHierarchyView.cpp
+++ b/Engine/Source/Editor/UMGEditor/Private/Hierarchy/SReadOnlyHierarchyView.cpp
@@ -13,6 +13,7 @@ void SReadOnlyHierarchyView::Construct(const FArguments& InArgs, const UWidgetBl
 {
 	OnSelectionChangedDelegate = InArgs._OnSelectionChanged;
 	WidgetBlueprint = InWidgetBlueprint;
+	ShowOnly = InArgs._ShowOnly;
 
 	SearchFilter = MakeShared<FTextFilter>(FTextFilter::FItemToStringArray::CreateSP(this, &SReadOnlyHierarchyView::GetFilterStringsForItem));
 	FilterHandler = MakeShared<FTreeFilterHandler>();
@@ -251,21 +252,30 @@ void SReadOnlyHierarchyView::GetItemChildren(TSharedPtr<FItem> Item, TArray<TSha
 	OutChildren.Append(Item->Children);
 }
 
-void SReadOnlyHierarchyView::BuildWidgetChildren(const TSharedPtr<FItem>& CurrentItem)
-{
-	if (const UPanelWidget* PanelWidget = Cast<UPanelWidget>(CurrentItem->Widget.Get()))
+void SReadOnlyHierarchyView::BuildWidgetChildren(const UWidget* Widget, TSharedPtr<FItem> Parent)
+{ 
+	// if we don't have a ShowOnly filter, or this widget is in it, create the item
+	if (ShowOnly.Num() == 0 || ShowOnly.Contains(Widget->GetFName()))
 	{
-		int32 NumChildren = PanelWidget->GetChildrenCount();
-		CurrentItem->Children.Reserve(NumChildren);
+		TSharedRef<FItem> WidgetItem = MakeShared<FItem>(Widget);
+		if (Parent.IsValid())
+		{
+			Parent->Children.Add(WidgetItem);
+		}
+		else
+		{
+			RootWidgets.Add(WidgetItem);
+		}
+
+		Parent = WidgetItem;
+	}
 
-		for (int32 Idx = 0; Idx < NumChildren; ++Idx)
+	// even if we didn't create an item for this widget, we still want to iterate over its children
+	if (const UPanelWidget* PanelWidget = Cast<UPanelWidget>(Widget))
+	{
+		for (const UWidget* Child : PanelWidget->GetAllChildren())
 		{
-			if (const UWidget* Child = PanelWidget->GetChildAt(Idx))
-			{
-				TSharedRef<FItem> ChildItem = MakeShared<FItem>(Child);
-				BuildWidgetChildren(ChildItem);
-				CurrentItem->Children.Add(ChildItem);
-			}
+			BuildWidgetChildren(Child, Parent);
 		}
 	}
 }
@@ -273,15 +283,22 @@ void SReadOnlyHierarchyView::BuildWidgetChildren(const TSharedPtr<FItem>& Curren
 void SReadOnlyHierarchyView::RebuildTree()
 {
 	const UWidgetBlueprint* WidgetBP = WidgetBlueprint.Get();
-	TSharedRef<FItem> WidgetBPItem = MakeShared<FItem>(WidgetBP);
-	RootWidgets.Add(WidgetBPItem);
+	if (WidgetBP == nullptr)
+	{
+		return;
+	}
+	
+	// if we don't have a ShowOnly filter, or this widget blueprint is in it, create the root item
+	TSharedPtr<FItem> RootItem;
+	if (ShowOnly.Num() == 0 || ShowOnly.Contains(WidgetBP->GetFName()))
+	{
+		RootItem = MakeShared<FItem>(WidgetBP);
+		RootWidgets.Add(RootItem);
+	}
 
 	if (const UWidget* RootWidget = WidgetBP->WidgetTree->RootWidget.Get())
 	{
-		TSharedRef<FItem> RootWidgetItem = MakeShared<FItem>(RootWidget);
-		WidgetBPItem->Children.Add(RootWidgetItem);
-
-		BuildWidgetChildren(RootWidgetItem);
+		BuildWidgetChildren(RootWidget, RootItem);
 	}
 
 	ExpandAll();
diff --git a/Engine/Source/Editor/UMGEditor/Public/Hierarchy/SReadOnlyHierarchyView.h b/Engine/Source/Editor/UMGEditor/Public/Hierarchy/SReadOnlyHierarchyView.h
--- a/Engine/Source/Editor/UMGEditor/Public/Hierarchy/SReadOnlyHierarchyView.h
+++ b/Engine/Source/Editor/UMGEditor/Public/Hierarchy/SReadOnlyHierarchyView.h
@@ -41,6 +41,7 @@ public:
 		SLATE_ARGUMENT_DEFAULT(bool, ShowSearch) = true;
 		SLATE_ARGUMENT_DEFAULT(ESelectionMode::Type, SelectionMode) = ESelectionMode::Single;
 		SLATE_EVENT(FOnSelectionChanged, OnSelectionChanged)
+		SLATE_ARGUMENT(TArray<FName>, ShowOnly)
 	SLATE_END_ARGS()
 
 	virtual ~SReadOnlyHierarchyView();
@@ -68,7 +69,7 @@ private:
 	void OnSelectionChanged(TSharedPtr<FItem> Selected, ESelectInfo::Type SelectionType);
 
 	void RebuildTree();
-	void BuildWidgetChildren(const TSharedPtr<FItem>& CurrentItem);
+	void BuildWidgetChildren(const UWidget* Widget, TSharedPtr<FItem> Parent);
 
 	void SetItemExpansionRecursive(TSharedPtr<FItem> Item, bool bShouldBeExpanded);
 	void ExpandAll();
@@ -78,6 +79,7 @@ private:
 private:
 
 	TWeakObjectPtr<const UWidgetBlueprint> WidgetBlueprint;
+	TArray<FName> ShowOnly;
 
 	TArray<TSharedPtr<FItem>> RootWidgets;
 	TArray<TSharedPtr<FItem>> FilteredRootWidgets;
