- Add an indication of selected animation in UMG
- Deselect all animation upon tab closure
- Update both sequencers on animation change

diff --git a/Engine/Source/Editor/UMGEditor/Private/BlueprintModes/WidgetDesignerApplicationMode.cpp b/Engine/Source/Editor/UMGEditor/Private/BlueprintModes/WidgetDesignerApplicationMode.cpp
--- a/Engine/Source/Editor/UMGEditor/Private/BlueprintModes/WidgetDesignerApplicationMode.cpp
+++ b/Engine/Source/Editor/UMGEditor/Private/BlueprintModes/WidgetDesignerApplicationMode.cpp
@@ -200,8 +200,8 @@ void FWidgetDesignerApplicationMode::PostActivateMode()
 	{
 		FWidgetDrawerConfig WidgetAnimSequencerDrawer(FAnimationTabSummoner::WidgetAnimSequencerDrawerID);
 		WidgetAnimSequencerDrawer.GetDrawerContentDelegate.BindSP(BP.Get(), &FWidgetBlueprintEditor::OnGetWidgetAnimSequencer);
-		WidgetAnimSequencerDrawer.OnDrawerOpenedDelegate.BindSP(BP.Get(), &FWidgetBlueprintEditor::OnWidgetAnimSequencerOpened);
-		WidgetAnimSequencerDrawer.OnDrawerDismissedDelegate.BindSP(BP.Get(), &FWidgetBlueprintEditor::OnWidgetAnimSequencerDismissed);
+		WidgetAnimSequencerDrawer.OnDrawerOpenedDelegate.BindSP(BP.Get(), &FWidgetBlueprintEditor::OnWidgetAnimDrawerSequencerOpened);
+		WidgetAnimSequencerDrawer.OnDrawerDismissedDelegate.BindSP(BP.Get(), &FWidgetBlueprintEditor::OnWidgetAnimDrawerSequencerDismissed);
 		WidgetAnimSequencerDrawer.ButtonText = LOCTEXT("StatusBar_WidgetAnimSequencer", "Animations");
 		WidgetAnimSequencerDrawer.ToolTipText = LOCTEXT("StatusBar_WidgetAnimSequencerToolTip", "Opens animation sequencer (Ctrl+Shift+Space Bar).");
 		WidgetAnimSequencerDrawer.Icon = FAppStyle::Get().GetBrush("UMGEditor.AnimTabIcon");
diff --git a/Engine/Source/Editor/UMGEditor/Private/Designer/SDesignerView.cpp b/Engine/Source/Editor/UMGEditor/Private/Designer/SDesignerView.cpp
--- a/Engine/Source/Editor/UMGEditor/Private/Designer/SDesignerView.cpp
+++ b/Engine/Source/Editor/UMGEditor/Private/Designer/SDesignerView.cpp
@@ -3121,6 +3121,15 @@ EVisibility SDesignerView::GetDesignerOutlineVisibility() const
 		return EVisibility::HitTestInvisible;
 	}
 
+	if ( Sequencer.IsValid() )
+	{
+		UWidgetAnimation* WidgetAnimation = Cast<UWidgetAnimation>(Sequencer->GetFocusedMovieSceneSequence());
+		if ( WidgetAnimation != UWidgetAnimation::GetNullAnimation() )
+		{
+			return EVisibility::HitTestInvisible;
+		}
+	}
+
 	return EVisibility::Hidden;
 }
 
@@ -3128,13 +3137,25 @@ FSlateColor SDesignerView::GetDesignerOutlineColor() const
 {
 	if ( GEditor->bIsSimulatingInEditor || GEditor->PlayWorld != nullptr )
 	{
-		return FLinearColor(0.863f, 0.407, 0.0f);
+		FLinearColor SimulatingIndicatorColor = FLinearColor(0.863f, 0.407, 0.0f);
+		return SimulatingIndicatorColor;
 	}
 
 	TSharedPtr<ISequencer> Sequencer = BlueprintEditor.Pin()->GetSequencer();
 	if ( Sequencer.IsValid() && Sequencer->GetAutoChangeMode() != EAutoChangeMode::None )
 	{
-		return FLinearColor::FromSRGBColor(FColor(251, 37, 0));
+		FLinearColor AnimRecordingIndicatorColor = FLinearColor::FromSRGBColor(FColor(251, 37, 0));
+		return AnimRecordingIndicatorColor;
+	}
+
+	if ( Sequencer.IsValid() )
+	{
+		UWidgetAnimation* WidgetAnimation = Cast<UWidgetAnimation>(Sequencer->GetFocusedMovieSceneSequence());
+		if ( WidgetAnimation != UWidgetAnimation::GetNullAnimation() )
+		{
+			FLinearColor AnimSelectedIndicatorColor = FLinearColor::FromSRGBColor(FColor(0, 67, 240));
+			return AnimSelectedIndicatorColor;
+		}
 	}
 
 	return FLinearColor::Transparent;
@@ -3146,13 +3167,24 @@ FText SDesignerView::GetDesignerOutlineText() const
 	{
 		return LOCTEXT("SIMULATING", "SIMULATING");
 	}
-	
+
 	TSharedPtr<ISequencer> Sequencer = BlueprintEditor.Pin()->GetSequencer();
 	if ( Sequencer.IsValid() && Sequencer->GetAutoChangeMode() != EAutoChangeMode::None )
 	{
 		return LOCTEXT("RECORDING", "RECORDING");
 	}
 
+	if ( Sequencer.IsValid() )
+	{
+		UWidgetAnimation* WidgetAnimation = Cast<UWidgetAnimation>(Sequencer->GetFocusedMovieSceneSequence());
+		if ( WidgetAnimation != UWidgetAnimation::GetNullAnimation() )
+		{
+			FFormatNamedArguments Args;
+			Args.Add(TEXT("Name"), FText::FromString(WidgetAnimation->GetDisplayLabel()));
+			return FText::Format(LOCTEXT("SELECTED", "SELECTED: {Name}"), Args);
+		}
+	}
+
 	return FText::GetEmpty();
 }
 
diff --git a/Engine/Source/Editor/UMGEditor/Private/TabFactory/AnimationTabSummoner.cpp b/Engine/Source/Editor/UMGEditor/Private/TabFactory/AnimationTabSummoner.cpp
--- a/Engine/Source/Editor/UMGEditor/Private/TabFactory/AnimationTabSummoner.cpp
+++ b/Engine/Source/Editor/UMGEditor/Private/TabFactory/AnimationTabSummoner.cpp
@@ -8,6 +8,7 @@
 #include "Animation/WidgetAnimation.h"
 #include "WidgetBlueprint.h"
 #include "WidgetBlueprintEditorUtils.h"
+#include "Widgets/Docking/SDockTab.h"
 #include "Widgets/Images/SImage.h"
 #include "Widgets/Input/SButton.h"
 #include "Widgets/Layout/SScrollBorder.h"
@@ -835,7 +836,17 @@ TSharedRef<SWidget> FAnimationTabSummoner::CreateTabBody(const FWorkflowTabSpawn
 	TSharedPtr<FWidgetBlueprintEditor> BlueprintEditorPinned = BlueprintEditor.Pin();
 
 	return SNew( SUMGAnimationList, BlueprintEditorPinned, bIsDrawerTab);
-	
+}
+
+TSharedRef<SDockTab> FAnimationTabSummoner::SpawnTab(const FWorkflowTabSpawnInfo& Info) const
+{
+	TSharedRef<SDockTab> NewTab = FWorkflowTabFactory::SpawnTab(Info);
+	if (TSharedPtr<FWidgetBlueprintEditor> BlueprintEditorPinned = BlueprintEditor.Pin())
+	{
+		NewTab->SetOnTabClosed(SDockTab::FOnTabClosedCallback::CreateSP(BlueprintEditorPinned.ToSharedRef(), &FWidgetBlueprintEditor::OnWidgetAnimTabSequencerClosed));
+		BlueprintEditorPinned->OnWidgetAnimTabSequencerOpened();
+	}
+	return NewTab;
 }
 
 #undef LOCTEXT_NAMESPACE 
diff --git a/Engine/Source/Editor/UMGEditor/Private/TabFactory/AnimationTabSummoner.h b/Engine/Source/Editor/UMGEditor/Private/TabFactory/AnimationTabSummoner.h
--- a/Engine/Source/Editor/UMGEditor/Private/TabFactory/AnimationTabSummoner.h
+++ b/Engine/Source/Editor/UMGEditor/Private/TabFactory/AnimationTabSummoner.h
@@ -7,6 +7,8 @@
 #include "WorkflowOrientedApp/WorkflowTabFactory.h"
 #include "WidgetBlueprintEditor.h"
 
+class SDockTab;
+
 struct FAnimationTabSummoner : public FWorkflowTabFactory
 {
 public:
@@ -17,7 +19,8 @@ public:
 	FAnimationTabSummoner(TSharedPtr<class FWidgetBlueprintEditor> InBlueprintEditor, bool bInIsDrawerTab = false);
 	
 	virtual TSharedRef<SWidget> CreateTabBody(const FWorkflowTabSpawnInfo& Info) const override;
-	
+	virtual TSharedRef<SDockTab> SpawnTab(const FWorkflowTabSpawnInfo& Info) const override;
+
 protected:
 	TWeakPtr<class FWidgetBlueprintEditor> BlueprintEditor;
 	bool bIsDrawerTab;
diff --git a/Engine/Source/Editor/UMGEditor/Private/WidgetBlueprintEditor.cpp b/Engine/Source/Editor/UMGEditor/Private/WidgetBlueprintEditor.cpp
--- a/Engine/Source/Editor/UMGEditor/Private/WidgetBlueprintEditor.cpp
+++ b/Engine/Source/Editor/UMGEditor/Private/WidgetBlueprintEditor.cpp
@@ -1187,7 +1187,18 @@ void FWidgetBlueprintEditor::NotifyWidgetAnimListChanged()
 	}
 }
 
+
 void FWidgetBlueprintEditor::OnWidgetAnimSequencerOpened(FName StatusBarWithDrawerName)
+{
+	OnWidgetAnimDrawerSequencerOpened(StatusBarWithDrawerName);
+}
+
+void FWidgetBlueprintEditor::OnWidgetAnimSequencerDismissed(const TSharedPtr<SWidget>& NewlyFocusedWidget)
+{
+	OnWidgetAnimDrawerSequencerDismissed(NewlyFocusedWidget);
+}
+
+void FWidgetBlueprintEditor::OnWidgetAnimDrawerSequencerOpened(FName StatusBarWithDrawerName)
 {
 	bIsSequencerDrawerOpen = true;
 
@@ -1214,7 +1225,7 @@ void FWidgetBlueprintEditor::OnWidgetAnimSequencerOpened(FName StatusBarWithDraw
 	}
 }
 
-void FWidgetBlueprintEditor::OnWidgetAnimSequencerDismissed(const TSharedPtr<SWidget>& NewlyFocusedWidget)
+void FWidgetBlueprintEditor::OnWidgetAnimDrawerSequencerDismissed(const TSharedPtr<SWidget>& NewlyFocusedWidget)
 {
 	bIsSequencerDrawerOpen = false;
 
@@ -1225,6 +1236,8 @@ void FWidgetBlueprintEditor::OnWidgetAnimSequencerDismissed(const TSharedPtr<SWi
 		{
 			ChangeViewedAnimation(*WidgetAnimation);
 		}
+		ActiveSequencer->GetSequencerWidget()->SetEnabled(false);
+		ActiveSequencer->SetAutoChangeMode(EAutoChangeMode::None);
 	}
 
 	for (TWeakPtr<ISequencer> SequencerPtr : Sequencers)
@@ -1238,6 +1251,28 @@ void FWidgetBlueprintEditor::OnWidgetAnimSequencerDismissed(const TSharedPtr<SWi
 	SetKeyboardFocus();
 }
 
+void FWidgetBlueprintEditor::OnWidgetAnimTabSequencerClosed(TSharedRef<SDockTab> ClosedTab)
+{
+	// Deselected any animation when closing the tab 
+	ChangeViewedAnimation(*UWidgetAnimation::GetNullAnimation());
+	if (TSharedPtr<ISequencer>& ActiveSequencer = GetSequencer())
+	{
+		ActiveSequencer->GetSequencerWidget()->SetEnabled(false);
+		ActiveSequencer->SetAutoChangeMode(EAutoChangeMode::None);
+	}
+}
+
+void FWidgetBlueprintEditor::OnWidgetAnimTabSequencerOpened()
+{
+	if (TSharedPtr<ISequencer>& ActiveSequencer = GetSequencer())
+	{
+		if (UWidgetAnimation* WidgetAnimation = Cast<UWidgetAnimation>(ActiveSequencer->GetFocusedMovieSceneSequence()))
+		{
+			ChangeViewedAnimation(*WidgetAnimation);
+		}
+	}
+}
+
 UWidgetBlueprint* FWidgetBlueprintEditor::GetWidgetBlueprintObj() const
 {
 	return Cast<UWidgetBlueprint>(GetBlueprintObj());
@@ -1386,48 +1421,54 @@ void FWidgetBlueprintEditor::DockInLayoutClicked()
 void FWidgetBlueprintEditor::ChangeViewedAnimation( UWidgetAnimation& InAnimationToView )
 {
 	CurrentAnimation = &InAnimationToView;
-
-	if (TSharedPtr<ISequencer>& Sequencer = GetSequencer())
+	for (TWeakPtr<ISequencer> SequencerPtr : Sequencers)
 	{
-		Sequencer->ResetToNewRootSequence(InAnimationToView);
-		if (&InAnimationToView == UWidgetAnimation::GetNullAnimation())
-		{
-			Sequencer->GetSequencerWidget()->SetEnabled(false);
-			Sequencer->SetAutoChangeMode(EAutoChangeMode::None);
-		}
-		else
+		if (SequencerPtr.IsValid())
 		{
-			Sequencer->GetSequencerWidget()->SetEnabled(true);
+			TSharedPtr<ISequencer>  Sequencer = SequencerPtr.Pin();
+			Sequencer->ResetToNewRootSequence(InAnimationToView);
+			if (&InAnimationToView == UWidgetAnimation::GetNullAnimation())
+			{
+				Sequencer->GetSequencerWidget()->SetEnabled(false);
+				Sequencer->SetAutoChangeMode(EAutoChangeMode::None);
+			}
+			else
+			{
+				Sequencer->GetSequencerWidget()->SetEnabled(true);
+			}
 		}
 	}
 
-	TWeakPtr<SOverlay> SequencerOverlay = bIsSequencerDrawerOpen ? DrawerSequencerOverlay : TabSequencerOverlay;
-	TWeakPtr<STextBlock> NoAnimationTextBlock = bIsSequencerDrawerOpen ? NoAnimationTextBlockDrawer : NoAnimationTextBlockTab;
-	if (SequencerOverlay.IsValid() && NoAnimationTextBlock.IsValid())
+	auto ToggleSequencerInteraction = [this](TWeakPtr<SOverlay> SequencerOverlay, TWeakPtr<STextBlock> NoAnimationTextBlock, UWidgetAnimation& InAnimationToView)
 	{
-		TSharedPtr<SOverlay> SequencerOverlayPin = SequencerOverlay.Pin();
-		TSharedPtr<STextBlock> NoAnimationTextBlockPin = NoAnimationTextBlock.Pin();
-
-		if (&InAnimationToView == UWidgetAnimation::GetNullAnimation())
+		if (SequencerOverlay.IsValid() && NoAnimationTextBlock.IsValid())
 		{
-			const FName CurveEditorTabName = FName(TEXT("SequencerGraphEditor"));
-			TSharedPtr<SDockTab> ExistingTab = GetToolkitHost()->GetTabManager()->FindExistingLiveTab(CurveEditorTabName);
-			if (ExistingTab)
+			TSharedPtr<SOverlay> SequencerOverlayPin = SequencerOverlay.Pin();
+			TSharedPtr<STextBlock> NoAnimationTextBlockPin = NoAnimationTextBlock.Pin();
+
+			if (&InAnimationToView == UWidgetAnimation::GetNullAnimation())
 			{
-				ExistingTab->RequestCloseTab();
-			}
+				const FName CurveEditorTabName = FName(TEXT("SequencerGraphEditor"));
+				TSharedPtr<SDockTab> ExistingTab = GetToolkitHost()->GetTabManager()->FindExistingLiveTab(CurveEditorTabName);
+				if (ExistingTab)
+				{
+					ExistingTab->RequestCloseTab();
+				}
 
-			// Disable sequencer from interaction
-			NoAnimationTextBlockPin->SetVisibility(EVisibility::Visible);
-			SequencerOverlayPin->SetVisibility(EVisibility::HitTestInvisible);
-		}
-		else
-		{
-			// Allow sequencer to be interacted with
-			NoAnimationTextBlockPin->SetVisibility(EVisibility::Collapsed);
-			SequencerOverlayPin->SetVisibility(EVisibility::SelfHitTestInvisible);
+				// Disable sequencer from interaction
+				NoAnimationTextBlockPin->SetVisibility(EVisibility::Visible);
+				SequencerOverlayPin->SetVisibility(EVisibility::HitTestInvisible);
+			}
+			else
+			{
+				// Allow sequencer to be interacted with
+				NoAnimationTextBlockPin->SetVisibility(EVisibility::Collapsed);
+				SequencerOverlayPin->SetVisibility(EVisibility::SelfHitTestInvisible);
+			}
 		}
-	}
+	};
+	ToggleSequencerInteraction(TabSequencerOverlay, NoAnimationTextBlockTab, InAnimationToView);
+	ToggleSequencerInteraction(DrawerSequencerOverlay, NoAnimationTextBlockDrawer, InAnimationToView);
 
 	InvalidatePreview();
 }
diff --git a/Engine/Source/Editor/UMGEditor/Public/WidgetBlueprintEditor.h b/Engine/Source/Editor/UMGEditor/Public/WidgetBlueprintEditor.h
--- a/Engine/Source/Editor/UMGEditor/Public/WidgetBlueprintEditor.h
+++ b/Engine/Source/Editor/UMGEditor/Public/WidgetBlueprintEditor.h
@@ -205,12 +205,24 @@ public:
 	/** Callback for when widget animation list has been updated */
 	void NotifyWidgetAnimListChanged();
 
-	/** Callback for anim drawer opening */
+	UE_DEPRECATED(5.2, "OnWidgetAnimSequencerOpened is deprecated, please use OnWidgetAnimDrawerSequencerOpened instead.")
 	void OnWidgetAnimSequencerOpened(FName StatusBarWithDrawerName);
 
-	/** Callback for anim drawer closing */
+	UE_DEPRECATED(5.2, "OnWidgetAnimSequencerDismissed is deprecated, please use OnWidgetAnimDrawerSequencerDismissed instead.")
 	void OnWidgetAnimSequencerDismissed(const TSharedPtr<SWidget>& NewlyFocusedWidget);
 
+	/** Callback for anim drawer opening */
+	void OnWidgetAnimDrawerSequencerOpened(FName StatusBarWithDrawerName);
+
+	/** Callback for anim drawer closing */
+	void OnWidgetAnimDrawerSequencerDismissed(const TSharedPtr<SWidget>& NewlyFocusedWidget);
+
+	/** Callback for anim tab closing */
+	void OnWidgetAnimTabSequencerClosed(TSharedRef<SDockTab> ClosedTab);
+
+	/** Callback for anim tab opening */
+	void OnWidgetAnimTabSequencerOpened();
+
 	/**
 	 * The widget we're now hovering over in any particular context, allows multiple views to 
 	 * synchronize feedback on where that widget is in their representation.
