Fix NamedSlot renaming issue.

diff --git a/Engine/Source/Editor/UMGEditor/Private/Hierarchy/SHierarchyViewItem.cpp b/Engine/Source/Editor/UMGEditor/Private/Hierarchy/SHierarchyViewItem.cpp
--- a/Engine/Source/Editor/UMGEditor/Private/Hierarchy/SHierarchyViewItem.cpp
+++ b/Engine/Source/Editor/UMGEditor/Private/Hierarchy/SHierarchyViewItem.cpp
@@ -1047,6 +1047,37 @@ FNamedSlotModel::FNamedSlotModel(FWidgetReference InItem, FName InSlotName, TSha
 	: FNamedSlotModelBase(InSlotName, InBlueprintEditor)
 	, Item(InItem)
 {
+#if WITH_EDITOR
+
+	// Revive trashed child of this NamedSlot, if any.
+	if (INamedSlotInterface* TemplateWidget = Cast<INamedSlotInterface>(Item.GetTemplate()))
+	{
+		if (UWidget* SlotContent = TemplateWidget->GetContentForSlot(SlotName))
+		{
+			if (SlotContent->HasAllFlags(RF_Transient) && SlotContent->GetOuter() == GetTransientPackage() && SlotContent->GetFName().ToString().StartsWith("TRASH_") && BlueprintEditor.IsValid())
+			{
+				if (UWidgetBlueprint* Blueprint = BlueprintEditor.Pin()->GetWidgetBlueprintObj())
+				{
+					FString NewName = SlotContent->GetFName().ToString();
+					NewName.RemoveFromStart("TRASH_");
+					SlotContent->ClearFlags(RF_Transient);
+					SlotContent->Rename(*NewName, Blueprint->WidgetTree);
+					TemplateWidget->SetContentForSlot(SlotName, SlotContent);
+				}
+			}
+		}
+	}
+
+	// Update the list of bindings if any renaming has occurred.
+	if (UUserWidget* TemplateWidget = Cast<UUserWidget>(Item.GetTemplate()))
+	{
+		TemplateWidget->AssignGUIDToBindings();
+		if (!TemplateWidget->GetContentForSlot(SlotName))
+		{
+			TemplateWidget->UpdateBindingForSlot(SlotName);
+		}
+	}
+#endif
 }
 
 FName FNamedSlotModel::GetUniqueName() const
diff --git a/Engine/Source/Editor/UMGEditor/Private/WidgetBlueprintCompiler.cpp b/Engine/Source/Editor/UMGEditor/Private/WidgetBlueprintCompiler.cpp
--- a/Engine/Source/Editor/UMGEditor/Private/WidgetBlueprintCompiler.cpp
+++ b/Engine/Source/Editor/UMGEditor/Private/WidgetBlueprintCompiler.cpp
@@ -906,6 +906,9 @@ void FWidgetBlueprintCompilerContext::FinishCompilingClass(UClass* Class)
 
 		// Add all the names of the named slot widgets to the slot names structure.
 		{
+		#if WITH_EDITOR
+			BPGClass->NamedSlotsWithID.Reset();
+		#endif
 			BPGClass->NamedSlots.Reset();
 			BPGClass->InstanceNamedSlots.Reset();
 			UWidgetBlueprint* WidgetBPIt = WidgetBP;
@@ -916,6 +919,10 @@ void FWidgetBlueprintCompilerContext::FinishCompilingClass(UClass* Class)
 					{
 						BPGClass->NamedSlots.Add(Widget->GetFName());
 
+					#if WITH_EDITOR
+						BPGClass->NamedSlotsWithID.Add(TPair<FName, FGuid>(Widget->GetFName(), NamedSlot->GetSlotGUID()));
+					#endif
+
 						if (NamedSlot->bExposeOnInstanceOnly)
 						{
 							BPGClass->InstanceNamedSlots.Add(Widget->GetFName());
diff --git a/Engine/Source/Runtime/UMG/Private/Components/NamedSlot.cpp b/Engine/Source/Runtime/UMG/Private/Components/NamedSlot.cpp
--- a/Engine/Source/Runtime/UMG/Private/Components/NamedSlot.cpp
+++ b/Engine/Source/Runtime/UMG/Private/Components/NamedSlot.cpp
@@ -96,6 +96,10 @@ const FText UNamedSlot::GetPaletteCategory()
 	return LOCTEXT("Common", "Common");
 }
 
+FGuid UNamedSlot::GetSlotGUID() const
+{
+	return SlotGuid;
+}
 #endif
 
 void UNamedSlot::Serialize(FArchive& Ar)
diff --git a/Engine/Source/Runtime/UMG/Private/UserWidget.cpp b/Engine/Source/Runtime/UMG/Private/UserWidget.cpp
--- a/Engine/Source/Runtime/UMG/Private/UserWidget.cpp
+++ b/Engine/Source/Runtime/UMG/Private/UserWidget.cpp
@@ -1240,6 +1240,37 @@ void UUserWidget::PostEditChangeProperty(FPropertyChangedEvent& PropertyChangedE
 	}
 }
 
+void UUserWidget::AssignGUIDToBindings()
+{
+	if (UWidgetBlueprintGeneratedClass* BGClass = GetWidgetTreeOwningClass())
+	{
+		for (int32 BindingIndex = 0; BindingIndex < NamedSlotBindings.Num(); BindingIndex++)
+		{
+			FNamedSlotBinding& Binding = NamedSlotBindings[BindingIndex];
+			if (BGClass->NamedSlotsWithID.Contains(Binding.Name))
+			{
+				Binding.Guid = BGClass->NamedSlotsWithID[Binding.Name];
+			}
+		}
+	}
+}
+
+void UUserWidget::UpdateBindingForSlot(FName SlotName)
+{
+	if (UWidgetBlueprintGeneratedClass* BGClass = GetWidgetTreeOwningClass())
+	{
+		if (BGClass->NamedSlotsWithID.Contains(SlotName))
+		{
+			for (FNamedSlotBinding& Binding : NamedSlotBindings)
+			{
+				if (BGClass->NamedSlotsWithID[SlotName] == Binding.Guid && !BGClass->NamedSlotsWithID.Contains(Binding.Name))
+				{
+					Binding.Name = SlotName;
+				}
+			}
+		}
+	}
+}
 #endif
 
 void UUserWidget::OnAnimationStarted_Implementation(const UWidgetAnimation* Animation)
diff --git a/Engine/Source/Runtime/UMG/Public/Blueprint/UserWidget.h b/Engine/Source/Runtime/UMG/Public/Blueprint/UserWidget.h
--- a/Engine/Source/Runtime/UMG/Public/Blueprint/UserWidget.h
+++ b/Engine/Source/Runtime/UMG/Public/Blueprint/UserWidget.h
@@ -157,6 +157,12 @@ public:
 	UPROPERTY()
 	FName Name;
 
+#if WITH_EDITORONLY_DATA
+	// GUID of the NamedSlot is used as a secondary identifier to find a binding in case the name of NamedSlot has changed.
+	UPROPERTY()
+	FGuid Guid;
+#endif
+
 	UPROPERTY(Instanced)
 	TObjectPtr<UWidget> Content;
 };
@@ -1135,6 +1141,12 @@ public:
 	virtual void OnDesignerChanged(const FDesignerChangedEventArgs& EventArgs) override;
 	virtual void PostEditChangeProperty(FPropertyChangedEvent& PropertyChangedEvent) override;
 
+	/** Update the binding for this namedslot if the name is not found but GUID is matched. */
+	void UpdateBindingForSlot(FName SlotName);
+
+	/** Add the GUID of each Namedslot widget to its corresponding binding, if any. */
+	void AssignGUIDToBindings();
+
 	/**
 	 * Final step of Widget Blueprint compilation. Allows widgets to perform custom validation and trigger compiler outputs as needed.
 	 * @see ValidateCompiledDefaults
diff --git a/Engine/Source/Runtime/UMG/Public/Blueprint/WidgetBlueprintGeneratedClass.h b/Engine/Source/Runtime/UMG/Public/Blueprint/WidgetBlueprintGeneratedClass.h
--- a/Engine/Source/Runtime/UMG/Public/Blueprint/WidgetBlueprintGeneratedClass.h
+++ b/Engine/Source/Runtime/UMG/Public/Blueprint/WidgetBlueprintGeneratedClass.h
@@ -105,6 +105,12 @@ public:
 	UPROPERTY()
 	TArray<FName> NamedSlots;
 
+#if WITH_EDITORONLY_DATA
+	/** All named slots mapped the assigned GUID of their UNamedSlot widget. **/
+	UPROPERTY()
+	TMap<FName, FGuid> NamedSlotsWithID;
+#endif
+
 	/**
 	 * Available Named Slots for content in a subclass.  These are slots that are accumulated from all super
 	 * classes on compile.  They will exclude any named slots that are filled by a parent class.
diff --git a/Engine/Source/Runtime/UMG/Public/Components/NamedSlot.h b/Engine/Source/Runtime/UMG/Public/Components/NamedSlot.h
--- a/Engine/Source/Runtime/UMG/Public/Components/NamedSlot.h
+++ b/Engine/Source/Runtime/UMG/Public/Components/NamedSlot.h
@@ -32,6 +32,7 @@ public:
 
 #if WITH_EDITOR
 	virtual const FText GetPaletteCategory() override;
+	FGuid GetSlotGUID() const;
 #endif
 
 public:
@@ -56,4 +57,12 @@ protected:
 
 protected:
 	TSharedPtr<SBox> MyBox;
+
+private:	
+#if WITH_EDITORONLY_DATA
+
+	UPROPERTY()
+	FGuid SlotGuid = FGuid::NewGuid();
+#endif
+
 };
