Sync the animation list content and selection between animation tab and animation drawer.

diff --git a/Engine/Source/Editor/UMGEditor/Private/TabFactory/AnimationTabSummoner.cpp b/Engine/Source/Editor/UMGEditor/Private/TabFactory/AnimationTabSummoner.cpp
--- a/Engine/Source/Editor/UMGEditor/Private/TabFactory/AnimationTabSummoner.cpp
+++ b/Engine/Source/Editor/UMGEditor/Private/TabFactory/AnimationTabSummoner.cpp
@@ -259,6 +259,11 @@ private:
 					Blueprint->Modify();
 					Blueprint->Animations.Add(WidgetAnimation);
 					ListItem.Pin()->bNewAnimation = false;
+
+					if (TSharedPtr<FWidgetBlueprintEditor> WidgetBlueprintEditorPin = BlueprintEditor.Pin())
+					{
+						WidgetBlueprintEditorPin->NotifyWidgetAnimListChanged();
+					}
 				}
 			}
 
@@ -359,6 +364,7 @@ public:
 		if (TSharedPtr<FWidgetBlueprintEditor> WidgetBlueprintEditorPin = BlueprintEditor.Pin())
 		{
 			WidgetBlueprintEditorPin->OnWidgetAnimationsUpdated.RemoveAll(this);
+			WidgetBlueprintEditorPin->OnSelectedAnimationChanged.RemoveAll(this);
 		}
 	}
 
@@ -370,6 +376,7 @@ public:
 		InBlueprintEditor->GetOnWidgetBlueprintTransaction().AddSP( this, &SUMGAnimationList::OnWidgetBlueprintTransaction );
 		InBlueprintEditor->OnEnterWidgetDesigner.AddSP(this, &SUMGAnimationList::OnEnteringDesignerMode);
 		InBlueprintEditor->OnWidgetAnimationsUpdated.AddSP(this, &SUMGAnimationList::OnUpdatedAnimationList);
+		InBlueprintEditor->OnSelectedAnimationChanged.AddSP(this, &SUMGAnimationList::AnimationListSelelctionSync);
 
 		SAssignNew(AnimationListView, SWidgetAnimationListView)
 			.ItemHeight(20.0f)
@@ -476,6 +483,34 @@ private:
 		AnimationListView->RequestListRefresh();
 	}
 
+	void AnimationListSelelctionSync()
+	{
+		if (TSharedPtr<FWidgetBlueprintEditor> WidgetBlueprintEditorPin = BlueprintEditor.Pin())
+		{
+			UWidgetAnimation* CurrentSelectedAnimation = WidgetBlueprintEditorPin->GetCurrentAnimation();
+
+			// This is to avoid looping calls to this function due to broadcast.
+			for (const TSharedPtr<FWidgetAnimationListItem>& SelectedAnimItem : AnimationListView->GetSelectedItems())
+			{
+				if (SelectedAnimItem->Animation == CurrentSelectedAnimation)
+				{
+					return;
+				}
+			}
+
+			// Find the list item containing the selected animation.
+			for (const TSharedPtr<FWidgetAnimationListItem>& AnimItem : Animations)
+			{
+				if (AnimItem->Animation == CurrentSelectedAnimation)
+				{
+					AnimationListView->SetSelection(AnimItem);
+					return;
+				}
+			}
+			AnimationListView->ClearSelection();
+		}
+	}
+
 	void UpdateAnimationList()
 	{
 		// There may be multiple sequencers acting as a view for our widget
diff --git a/Engine/Source/Editor/UMGEditor/Private/WidgetBlueprintEditor.cpp b/Engine/Source/Editor/UMGEditor/Private/WidgetBlueprintEditor.cpp
--- a/Engine/Source/Editor/UMGEditor/Private/WidgetBlueprintEditor.cpp
+++ b/Engine/Source/Editor/UMGEditor/Private/WidgetBlueprintEditor.cpp
@@ -1227,8 +1227,6 @@ void FWidgetBlueprintEditor::OnWidgetAnimDrawerSequencerOpened(FName StatusBarWi
 
 void FWidgetBlueprintEditor::OnWidgetAnimDrawerSequencerDismissed(const TSharedPtr<SWidget>& NewlyFocusedWidget)
 {
-	bIsSequencerDrawerOpen = false;
-
 	if (TSharedPtr<ISequencer>& ActiveSequencer = GetSequencer())
 	{
 		UWidgetAnimation* WidgetAnimation = Cast<UWidgetAnimation>(ActiveSequencer->GetFocusedMovieSceneSequence());
@@ -1239,6 +1237,7 @@ void FWidgetBlueprintEditor::OnWidgetAnimDrawerSequencerDismissed(const TSharedP
 		ActiveSequencer->GetSequencerWidget()->SetEnabled(false);
 		ActiveSequencer->SetAutoChangeMode(EAutoChangeMode::None);
 	}
+	bIsSequencerDrawerOpen = false;
 
 	for (TWeakPtr<ISequencer> SequencerPtr : Sequencers)
 	{
@@ -1471,6 +1470,7 @@ void FWidgetBlueprintEditor::ChangeViewedAnimation( UWidgetAnimation& InAnimatio
 	ToggleSequencerInteraction(DrawerSequencerOverlay, NoAnimationTextBlockDrawer, InAnimationToView);
 
 	InvalidatePreview();
+	OnSelectedAnimationChanged.Broadcast();
 }
 
 void FWidgetBlueprintEditor::RefreshPreview()
diff --git a/Engine/Source/Editor/UMGEditor/Public/WidgetBlueprintEditor.h b/Engine/Source/Editor/UMGEditor/Public/WidgetBlueprintEditor.h
--- a/Engine/Source/Editor/UMGEditor/Public/WidgetBlueprintEditor.h
+++ b/Engine/Source/Editor/UMGEditor/Public/WidgetBlueprintEditor.h
@@ -53,6 +53,9 @@ public:
 	/** Called when animation list changes */
 	DECLARE_MULTICAST_DELEGATE(FOnWidgetAnimationsUpdated)
 
+	/** Called when animation selection changes */
+	DECLARE_MULTICAST_DELEGATE(FOnSelectedAnimationChanged)
+
 	DECLARE_EVENT(FWidgetBlueprintEditor, FOnEnterWidgetDesigner)
 
 public:
@@ -267,6 +270,9 @@ public:
 	/** Fires whenever the selected set of widgets changes */
 	FOnSelectedWidgetsChanged OnSelectedWidgetsChanged;
 
+	/** Fires whenever the selected animation changes */
+	FOnSelectedAnimationChanged OnSelectedAnimationChanged;
+
 	/** Notification for when the preview widget has been updated */
 	FOnWidgetPreviewUpdated OnWidgetPreviewUpdated;
 
