Fix list representation in widget library & widget library jittering.

diff --git a/Engine/Source/Editor/UMGEditor/Private/Library/SLibraryViewModel.cpp b/Engine/Source/Editor/UMGEditor/Private/Library/SLibraryViewModel.cpp
--- a/Engine/Source/Editor/UMGEditor/Private/Library/SLibraryViewModel.cpp
+++ b/Engine/Source/Editor/UMGEditor/Private/Library/SLibraryViewModel.cpp
@@ -40,32 +40,33 @@ void FLibraryViewModel::BuildWidgetTemplateCategory(FString& Category, TArray<TS
 	Header->Children.Add(TemplateViewModel);
 	WidgetTemplateListViewModels.Add(TemplateViewModel);
 
-	for ( auto& Template : Templates)
-	{
-		TSharedPtr<FWidgetTemplateViewModel> CurrentTemplateViewModel = MakeShareable(new FWidgetTemplateViewModel());
-		CurrentTemplateViewModel->Template = Template;
-		CurrentTemplateViewModel->FavortiesViewModel = this;
-		Header->Children.Add(CurrentTemplateViewModel);
-
-		// If it's a favorite, we also add it to the Favorite section
-		int32 index = FavoritesList.Find(Template->Name.ToString());
-		if (index != INDEX_NONE)
-		{
-			CurrentTemplateViewModel->SetFavorite();
-
-			// We have to create a second copy of the ViewModel for the treeview has it doesn't support to have the same element twice.
-			TSharedPtr<FWidgetTemplateViewModel> FavoriteTemplateViewModel = MakeShareable(new FWidgetTemplateViewModel());
-			FavoriteTemplateViewModel->Template = Template;
-			FavoriteTemplateViewModel->FavortiesViewModel = this;
-			FavoriteTemplateViewModel->SetFavorite();
-
-			FavoriteHeader->Children.Add(FavoriteTemplateViewModel);
-
-			// Remove the favorite from the temporary list
-			FavoritesList.RemoveAt(index);
-		}
-
-	}
+	// @TODO: DarenC - Reference for when implementing favorites system 
+	//for ( auto& Template : Entry.Value )
+	//{
+	//	TSharedPtr<FWidgetTemplateViewModel> TemplateViewModel = MakeShareable(new FWidgetTemplateViewModel());
+	//	TemplateViewModel->Template = Template;
+	//	TemplateViewModel->FavortiesViewModel = this;
+	//	Header->Children.Add(TemplateViewModel);
+
+	//	// If it's a favorite, we also add it to the Favorite section
+	//	int32 index = FavoritesList.Find(Template->Name.ToString());
+	//	if (index != INDEX_NONE)
+	//	{
+	//		TemplateViewModel->SetFavorite();
+
+	//		// We have to create a second copy of the ViewModel for the treeview has it doesn't support to have the same element twice.
+	//		TSharedPtr<FWidgetTemplateViewModel> FavoriteTemplateViewModel = MakeShareable(new FWidgetTemplateViewModel());
+	//		FavoriteTemplateViewModel->Template = Template;
+	//		FavoriteTemplateViewModel->FavortiesViewModel = this;
+	//		FavoriteTemplateViewModel->SetFavorite();
+
+	//		FavoriteHeader->Children.Add(FavoriteTemplateViewModel);
+
+	//		// Remove the favorite from the temporary list
+	//		FavoritesList.RemoveAt(index);
+	//	}
+
+	//}
 
 	Header->Children.Sort([](const TSharedPtr<FWidgetViewModel>& L, const TSharedPtr<FWidgetViewModel>& R) { return R->GetName().CompareTo(L->GetName()) > 0; });
 
diff --git a/Engine/Source/Runtime/Slate/Private/Widgets/Views/STableViewBase.cpp b/Engine/Source/Runtime/Slate/Private/Widgets/Views/STableViewBase.cpp
--- a/Engine/Source/Runtime/Slate/Private/Widgets/Views/STableViewBase.cpp
+++ b/Engine/Source/Runtime/Slate/Private/Widgets/Views/STableViewBase.cpp
@@ -287,17 +287,8 @@ void STableViewBase::Tick( const FGeometry& AllottedGeometry, const double InCur
 			PanelGeometryLastTick = PanelGeometry;
 			
 			const int32 NumItemsPerLine = GetNumItemsPerLine();
+			const EScrollIntoViewResult ScrollIntoViewResult = ScrollIntoView(PanelGeometry);
 
-			EScrollIntoViewResult ScrollIntoViewResult = EScrollIntoViewResult::Deferred;
-
-			/* Don't try to scroll if our geometry is changing (eg: In the middle of refreshing items) because the
-			 * scrolled item could end up outside the view
-			 */
-			if(!bPanelGeometryChanged)
-			{
-				ScrollIntoViewResult = ScrollIntoView(PanelGeometry);
-			}
-			
 			double TargetScrollOffset = GetTargetScrollOffset();
 
 			if (bEnableAnimatedScrolling)
