Widget details customization extender:
- Add customization extender interface
- Add perform customization extenders

diff --git a/Engine/Source/Editor/UMGEditor/Private/Customizations/UMGDetailCustomizations.cpp b/Engine/Source/Editor/UMGEditor/Private/Customizations/UMGDetailCustomizations.cpp
--- a/Engine/Source/Editor/UMGEditor/Private/Customizations/UMGDetailCustomizations.cpp
+++ b/Engine/Source/Editor/UMGEditor/Private/Customizations/UMGDetailCustomizations.cpp
@@ -15,6 +15,7 @@
 
 #include "Algo/Transform.h"
 #include "BlueprintModes/WidgetBlueprintApplicationModes.h"
+#include "Customizations/IBlueprintWidgetCustomizationExtender.h"
 #include "DetailWidgetRow.h"
 #include "PropertyHandle.h"
 #include "IDetailPropertyRow.h"
@@ -805,12 +806,17 @@ void FBlueprintWidgetCustomization::CustomizeDetails( IDetailLayoutBuilder& Deta
 
 	TArray< TWeakObjectPtr<UObject> > OutObjects;
 	DetailLayout.GetObjectsBeingCustomized(OutObjects);
-	
+
+	FMemMark Mark(FMemStack::Get());
+	TArray<UWidget*, TMemStackAllocator<>> CustomizedWidgets;
+	CustomizedWidgets.Reserve(OutObjects.Num());
+
 	UClass* SlotBaseClasses = nullptr;
 	for (const TWeakObjectPtr<UObject>& Obj : OutObjects)
 	{
 		if (UWidget* Widget = Cast<UWidget>(Obj.Get()))
 		{
+			CustomizedWidgets.Add(Widget);
 			if (Widget->Slot)
 			{
 				UClass* SlotClass = Widget->Slot->GetClass();
@@ -830,11 +836,6 @@ void FBlueprintWidgetCustomization::CustomizeDetails( IDetailLayoutBuilder& Deta
 				break;
 			}
 		}
-		else
-		{
-			SlotBaseClasses = nullptr;
-			break;
-		}
 	}
 	
 	if (SlotBaseClasses)
@@ -848,20 +849,18 @@ void FBlueprintWidgetCustomization::CustomizeDetails( IDetailLayoutBuilder& Deta
 	}
 
 	PerformAccessibilityCustomization(DetailLayout);
-	PerformBindingCustomization(DetailLayout);
+	PerformBindingCustomization(DetailLayout, CustomizedWidgets);
+	PerformCustomizationExtenders(DetailLayout, CustomizedWidgets);
 }
 
-void FBlueprintWidgetCustomization::PerformBindingCustomization(IDetailLayoutBuilder& DetailLayout)
+void FBlueprintWidgetCustomization::PerformBindingCustomization(IDetailLayoutBuilder& DetailLayout, const TArrayView<UWidget*> Widgets)
 {
 	static const FName IsBindableEventName(TEXT("IsBindableEvent"));
 
-	TArray< TWeakObjectPtr<UObject> > OutObjects;
-	DetailLayout.GetObjectsBeingCustomized(OutObjects);
-
-	if ( OutObjects.Num() == 1 )
+	if ( Widgets.Num() == 1 )
 	{
-		UWidget* Widget = Cast<UWidget>(OutObjects[0].Get());
-		UClass* PropertyClass = OutObjects[0].Get()->GetClass();
+		UWidget* Widget = Widgets[0];
+		UClass* PropertyClass = Widget->GetClass();
 
 		for ( TFieldIterator<FProperty> PropertyIt(PropertyClass, EFieldIteratorFlags::IncludeSuper); PropertyIt; ++PropertyIt )
 		{
@@ -877,7 +876,7 @@ void FBlueprintWidgetCustomization::PerformBindingCustomization(IDetailLayoutBui
 			}
 			else if ( FMulticastDelegateProperty* MulticastDelegateProperty = CastField<FMulticastDelegateProperty>(Property) )
 			{
-				CreateMulticastEventCustomization(DetailLayout, OutObjects[0].Get()->GetFName(), PropertyClass, MulticastDelegateProperty);
+				CreateMulticastEventCustomization(DetailLayout, Widget->GetFName(), PropertyClass, MulticastDelegateProperty);
 			}
 		}
 	}
@@ -954,4 +953,18 @@ void FBlueprintWidgetCustomization::CustomizeAccessibilityProperty(IDetailLayout
 	];
 }
 
+void FBlueprintWidgetCustomization::PerformCustomizationExtenders(IDetailLayoutBuilder& DetailLayout, const TArrayView<UWidget*> Widgets)
+{
+	if (IUMGEditorModule* UMGEditorModule = FModuleManager::GetModulePtr<IUMGEditorModule>("UMGEditor"))
+	{
+		if (TSharedPtr<FWidgetBlueprintEditor> EditorPtr = Editor.Pin())
+		{
+			for (TSharedRef<IBlueprintWidgetCustomizationExtender>& CustomizationExtender : UMGEditorModule->GetAllWidgetCustomizationExtenders())
+			{
+				CustomizationExtender->CustomizeDetails(DetailLayout, Widgets, EditorPtr.ToSharedRef());
+			}
+		}
+	}
+}
+
 #undef LOCTEXT_NAMESPACE
diff --git a/Engine/Source/Editor/UMGEditor/Private/Customizations/UMGDetailCustomizations.h b/Engine/Source/Editor/UMGEditor/Private/Customizations/UMGDetailCustomizations.h
--- a/Engine/Source/Editor/UMGEditor/Private/Customizations/UMGDetailCustomizations.h
+++ b/Engine/Source/Editor/UMGEditor/Private/Customizations/UMGDetailCustomizations.h
@@ -38,7 +38,7 @@ public:
 	static bool HasPropertyBindings(TWeakPtr<FWidgetBlueprintEditor> InEditor, const TSharedRef<IPropertyHandle>& InPropertyHandle);
 
 private:
-	void PerformBindingCustomization(IDetailLayoutBuilder& DetailLayout);
+	void PerformBindingCustomization(IDetailLayoutBuilder& DetailLayout, const TArrayView<UWidget*> Widgets);
 
 	void CreateEventCustomization( IDetailLayoutBuilder& DetailLayout, FDelegateProperty* Property, UWidget* Widget );
 
@@ -52,6 +52,7 @@ private:
 
 	void PerformAccessibilityCustomization(IDetailLayoutBuilder& DetailLayout);
 	void CustomizeAccessibilityProperty(IDetailLayoutBuilder& DetailLayout, const FName& BehaviorPropertyName, const FName& TextPropertyName);
+	void PerformCustomizationExtenders(IDetailLayoutBuilder& DetailLayout, const TArrayView<UWidget*> Widgets);
 private:
 
 	TWeakPtr<FWidgetBlueprintEditor> Editor;
diff --git a/Engine/Source/Editor/UMGEditor/Private/UMGEditorModule.cpp b/Engine/Source/Editor/UMGEditor/Private/UMGEditorModule.cpp
--- a/Engine/Source/Editor/UMGEditor/Private/UMGEditorModule.cpp
+++ b/Engine/Source/Editor/UMGEditor/Private/UMGEditorModule.cpp
@@ -41,6 +41,7 @@
 #include "BlueprintEditorModule.h"
 #include "PropertyEditorModule.h"
 #include "Customizations/DynamicEntryBoxDetails.h"
+#include "Customizations/IBlueprintWidgetCustomizationExtender.h"
 #include "Customizations/ListViewBaseDetails.h"
 #include "WidgetBlueprintThumbnailRenderer.h"
 #include "Customizations/WidgetThumbnailCustomization.h"
@@ -276,6 +277,21 @@ public:
 		return WidgetEditorToolbarExtenders;
 	}
 
+	virtual void AddWidgetCustomizationExtender(const TSharedRef<IBlueprintWidgetCustomizationExtender>& WidgetCustomizationExtender) override
+	{
+		WidgetCustomizationExtenders.AddUnique(WidgetCustomizationExtender);
+	}
+
+	virtual void RemoveWidgetCustomizationExtender(const TSharedRef<IBlueprintWidgetCustomizationExtender>& WidgetCustomizationExtender) override
+	{
+		WidgetCustomizationExtenders.RemoveSingleSwap(WidgetCustomizationExtender, false);
+	}
+
+	virtual TArrayView<TSharedRef<IBlueprintWidgetCustomizationExtender>> GetAllWidgetCustomizationExtenders() override
+	{
+		return WidgetCustomizationExtenders;
+	}
+
 	virtual FOnRegisterLayoutExtensions& OnRegisterLayoutExtensions() override 
 	{
 		return RegisterLayoutExtensions; 
@@ -352,6 +368,7 @@ private:
 
 	/** All toolbar extenders. Utilized by tool palette */
 	TArray<FWidgetEditorToolbarExtender> WidgetEditorToolbarExtenders;
+	TArray<TSharedRef<IBlueprintWidgetCustomizationExtender>> WidgetCustomizationExtenders;
 
 	USequencerSettings* Settings;
 
diff --git a/Engine/Source/Editor/UMGEditor/Public/Customizations/IBlueprintWidgetCustomizationExtender.h b/Engine/Source/Editor/UMGEditor/Public/Customizations/IBlueprintWidgetCustomizationExtender.h
new file mode 100644
--- /dev/null
+++ b/Engine/Source/Editor/UMGEditor/Public/Customizations/IBlueprintWidgetCustomizationExtender.h
@@ -0,0 +1,26 @@
+// Copyright Epic Games, Inc. All Rights Reserved.
+
+#pragma once
+
+#include "CoreMinimal.h"
+
+class IDetailLayoutBuilder;
+class UWidget;
+class FWidgetBlueprintEditor;
+
+/**
+ * Details customization extender for UWidgets.
+ */
+class UMGEDITOR_API IBlueprintWidgetCustomizationExtender : public TSharedFromThis<IBlueprintWidgetCustomizationExtender>
+{
+public:
+	virtual ~IBlueprintWidgetCustomizationExtender() = default;
+
+	/**
+	 * Customizes given widgets.
+	 * @param InDetailLayout				detail layout builder for customization.
+	 * @param InWidget						widgets to be customized.
+	 * @param InWidgetBlueprintEditor		widget blueprint editor that is used for widgets being customized.
+	 */
+	virtual void CustomizeDetails(IDetailLayoutBuilder& InDetailLayout, const TArrayView<UWidget*> InWidgets, const TSharedRef<FWidgetBlueprintEditor>& InWidgetBlueprintEditor) = 0;
+};
\ No newline at end of file
diff --git a/Engine/Source/Editor/UMGEditor/Public/UMGEditorModule.h b/Engine/Source/Editor/UMGEditor/Public/UMGEditorModule.h
--- a/Engine/Source/Editor/UMGEditor/Public/UMGEditorModule.h
+++ b/Engine/Source/Editor/UMGEditor/Public/UMGEditorModule.h
@@ -12,6 +12,7 @@ extern const FName UMGEditorAppIdentifier;
 class FUMGEditor;
 class FWidgetBlueprintApplicationMode;
 class FWorkflowAllowedTabSet;
+class IBlueprintWidgetCustomizationExtender;
 
 /** The public interface of the UMG editor module. */
 class IUMGEditorModule : 
@@ -35,6 +36,15 @@ public:
 	/** Get current toolbar extenders for widget editor */
 	virtual TArrayView<FWidgetEditorToolbarExtender> GetAllWidgetEditorToolbarExtenders() = 0;
 
+	/** Add customization for widget */
+	virtual void AddWidgetCustomizationExtender(const TSharedRef<IBlueprintWidgetCustomizationExtender>& WidgetCustomizationExtender) = 0;
+
+	/** Remove customization for widget */
+	virtual void RemoveWidgetCustomizationExtender(const TSharedRef<IBlueprintWidgetCustomizationExtender>& WidgetCustomizationExtender) = 0;
+
+	/** Get current customization extenders for widget */
+	virtual TArrayView<TSharedRef<IBlueprintWidgetCustomizationExtender>> GetAllWidgetCustomizationExtenders() = 0;
+
 	/** Support for general layout extenders */
 	DECLARE_EVENT_OneParam(IStaticMeshEditor, FOnRegisterLayoutExtensions, FLayoutExtender&);
 	virtual FOnRegisterLayoutExtensions& OnRegisterLayoutExtensions() = 0;
