UMG Designer: Fix to "Replace With" crash.

diff --git a/Engine/Source/Editor/UMGEditor/Private/Templates/WidgetTemplateBlueprintClass.cpp b/Engine/Source/Editor/UMGEditor/Private/Templates/WidgetTemplateBlueprintClass.cpp
--- a/Engine/Source/Editor/UMGEditor/Private/Templates/WidgetTemplateBlueprintClass.cpp
+++ b/Engine/Source/Editor/UMGEditor/Private/Templates/WidgetTemplateBlueprintClass.cpp
@@ -76,11 +76,10 @@ UWidget* FWidgetTemplateBlueprintClass::Create(UWidgetTree* Tree)
 		{
 			WidgetClass = LoadObject<UBlueprintGeneratedClass>(nullptr, *AssetPath);
 		}
-		else
+		else if (UBlueprint* LoadedWidget = LoadObject<UBlueprint>(nullptr, *AssetPath))
 		{
-			UWidgetBlueprint* LoadedWidget = LoadObject<UWidgetBlueprint>(nullptr, *AssetPath);
 			WidgetClass = *LoadedWidget->GeneratedClass;
-		}
+ 		}
 	}
 
 	return FWidgetTemplateClass::CreateNamed(Tree, FName(*FBlueprintEditorUtils::GetClassNameWithoutSuffix(WidgetClass.Get())));
diff --git a/Engine/Source/Editor/UMGEditor/Private/WidgetBlueprintEditorUtils.cpp b/Engine/Source/Editor/UMGEditor/Private/WidgetBlueprintEditorUtils.cpp
--- a/Engine/Source/Editor/UMGEditor/Private/WidgetBlueprintEditorUtils.cpp
+++ b/Engine/Source/Editor/UMGEditor/Private/WidgetBlueprintEditorUtils.cpp
@@ -1064,57 +1064,56 @@ bool FWidgetBlueprintEditorUtils::CanBeReplacedWithTemplate(TSharedRef<FWidgetBl
 	FAssetData SelectedUserWidget = BlueprintEditor->GetSelectedUserWidget();
 	UWidget* ThisWidget = Widget.GetTemplate();
 	UPanelWidget* ExistingPanel = Cast<UPanelWidget>(ThisWidget);
+
+	UClass* WidgetClass = nullptr;
 	// If selecting another widget blueprint
 	if (SelectedUserWidget.GetSoftObjectPath().IsValid())
 	{
-		if (ExistingPanel)
+		if (ExistingPanel && ExistingPanel->GetChildrenCount() != 0)
 		{
-			if (ExistingPanel->GetChildrenCount() != 0)
-			{
-				return false;
-			}
+			return false;
 		}
-		if (UUserWidget* NewUserWidget = CastChecked<UUserWidget>(FWidgetTemplateBlueprintClass(SelectedUserWidget).Create(BP->WidgetTree)))
+		if (UWidget* NewWidget = FWidgetTemplateBlueprintClass(SelectedUserWidget).Create(BP->WidgetTree))
 		{
-			const bool bFreeFromCircularRefs = BP->IsWidgetFreeFromCircularReferences(NewUserWidget);
-			NewUserWidget->Rename(nullptr, GetTransientPackage());
-			return bFreeFromCircularRefs;
+			// If we are creating a UserWidget, check for Circular references
+			if (UUserWidget* NewUserWidget = Cast<UUserWidget>(NewWidget))
+			{
+				const bool bFreeFromCircularRefs = BP->IsWidgetFreeFromCircularReferences(NewUserWidget);
+				NewWidget->Rename(nullptr, GetTransientPackage());
+				return bFreeFromCircularRefs;
+			}
+			WidgetClass = NewWidget->GetClass();
+			NewWidget->Rename(nullptr, GetTransientPackage());
 		}
 	}
 
-	UClass* WidgetClass = BlueprintEditor->GetSelectedTemplate().Get();
-	const bool bCanReplace = WidgetClass->IsChildOf(UPanelWidget::StaticClass());
-	if (!ExistingPanel && !bCanReplace)
+	// If we get here, the Widget selected is not a UserWidget and it's not a Blueprint.
+	if (!WidgetClass)
 	{
-		return true;
+		WidgetClass = BlueprintEditor->GetSelectedTemplate().Get();
 	}
-	else if (!ExistingPanel && bCanReplace)
+
+	// If the Widget to replace is not a Panel we can replace it with anything	
+	if (!ExistingPanel)
 	{
 		return true;
 	}
-	else if (ExistingPanel && !bCanReplace)
+
+	const bool bNewWidgetClassIsAPanel = WidgetClass->IsChildOf(UPanelWidget::StaticClass());
+	// If the Widget to replace is a Panel and the new widget is not, we allow to replace it only if it's empty;
+	if (!bNewWidgetClassIsAPanel)
 	{
-		if (ExistingPanel->GetChildrenCount() == 0)
-		{
-			return true;
-		}
-		else 
-		{
-			return false;
-		}
+		return ExistingPanel->GetChildrenCount() == 0;
 	}
-	else 
+
+	// If the Widget to replace is a Panel that can have multiple children, we allow to replace it with a Panel that can support multiple children only.
+	if (ExistingPanel->GetClass()->GetDefaultObject<UPanelWidget>()->CanHaveMultipleChildren() && bNewWidgetClassIsAPanel)
 	{
-		if (ExistingPanel->GetClass()->GetDefaultObject<UPanelWidget>()->CanHaveMultipleChildren() && bCanReplace)
-		{
-			const bool bChildAllowed = WidgetClass->GetDefaultObject<UPanelWidget>()->CanHaveMultipleChildren() || ExistingPanel->GetChildrenCount() == 0;
-			return bChildAllowed;
-		}
-		else
-		{
-			return true;
-		}
+		const bool bChildAllowed = WidgetClass->GetDefaultObject<UPanelWidget>()->CanHaveMultipleChildren() || ExistingPanel->GetChildrenCount() == 0;
+		return bChildAllowed;
 	}
+	
+	return true;
 }
 
 void FWidgetBlueprintEditorUtils::ReplaceWidgetWithChildren(TSharedRef<FWidgetBlueprintEditor> BlueprintEditor, UWidgetBlueprint* BP, FWidgetReference Widget)
