Add a missing include in SDesignerView to fix non unity build.

diff --git a/Engine/Source/Editor/UMGEditor/Private/Designer/SDesignerView.cpp b/Engine/Source/Editor/UMGEditor/Private/Designer/SDesignerView.cpp
--- a/Engine/Source/Editor/UMGEditor/Private/Designer/SDesignerView.cpp
+++ b/Engine/Source/Editor/UMGEditor/Private/Designer/SDesignerView.cpp
@@ -19,6 +19,8 @@
 #include "Widgets/Input/SButton.h"
 #include "Widgets/Input/SComboButton.h"
 
+#include "Animation/WidgetAnimation.h"
+
 #include "Components/CanvasPanelSlot.h"
 #include "Blueprint/WidgetTree.h"
 #include "Settings/WidgetDesignerSettings.h"
