Add FindReferenced command/option to widget blueprint editor context menus. Will do a graph reference search if a widget is marked as also being a variable.

#tests Tested context menu and shortcut keys

diff --git a/Engine/Source/Editor/UMGEditor/Private/WidgetBlueprintEditor.cpp b/Engine/Source/Editor/UMGEditor/Private/WidgetBlueprintEditor.cpp
--- a/Engine/Source/Editor/UMGEditor/Private/WidgetBlueprintEditor.cpp
+++ b/Engine/Source/Editor/UMGEditor/Private/WidgetBlueprintEditor.cpp
@@ -67,6 +67,7 @@
 #include "Library/SLibraryViewModel.h"
 
 #include "DesktopPlatformModule.h"
+#include "Engine/MemberReference.h"
 #include "Engine/TextureRenderTarget2D.h"
 #include "IDesktopPlatform.h"
 #include "IImageWrapper.h"
@@ -82,6 +83,7 @@
 #include "Editor/UnrealEdEngine.h"
 #include "Preferences/UnrealEdOptions.h"
 #include "UnrealEdGlobals.h"
+#include "GraphEditorActions.h"
 
 #define LOCTEXT_NAMESPACE "UMG"
 
@@ -196,6 +198,10 @@ void FWidgetBlueprintEditor::InitWidgetBlueprintEditor(const EToolkitMode::Type
 		FCanExecuteAction::CreateSP(this, &FWidgetBlueprintEditor::CanDuplicateSelectedWidgets)
 		);
 
+	DesignerCommandList->MapAction(FGraphEditorCommands::Get().FindReferences,
+		FExecuteAction::CreateSP(this, &FWidgetBlueprintEditor::OnFindWidgetReferences),
+		FCanExecuteAction::CreateSP(this, &FWidgetBlueprintEditor::CanFindWidgetReferences));
+
 	TSharedPtr<class IToolkitHost> PinnedToolkitHost = ToolkitHost.Pin();
 	check(PinnedToolkitHost.IsValid());
 	ModeUILayer = MakeShared<FWidgetEditorModeUILayer>(PinnedToolkitHost.Get());
@@ -877,6 +883,24 @@ void FWidgetBlueprintEditor::DuplicateSelectedWidgets()
 	SelectWidgets(DuplicatedWidgetRefs, false);
 }
 
+void FWidgetBlueprintEditor::OnFindWidgetReferences()
+{
+	FWidgetReference WidgetReference = *GetSelectedWidgets().CreateConstIterator();
+	const FString VariableName = WidgetReference.GetTemplate()->GetName();
+
+	FMemberReference MemberReference;
+	MemberReference.SetSelfMember(*VariableName);
+	const FString SearchTerm = MemberReference.GetReferenceSearchString(GetBlueprintObj()->SkeletonGeneratedClass);
+
+	SetCurrentMode(FWidgetBlueprintApplicationModes::GraphMode);
+	SummonSearchUI(true, SearchTerm);
+}
+
+bool FWidgetBlueprintEditor::CanFindWidgetReferences() const
+{
+	return GetSelectedWidgets().Num() == 1 && GetSelectedWidgets().CreateConstIterator()->GetTemplate()->bIsVariable;
+}
+
 bool FWidgetBlueprintEditor::CanCreateNativeBaseClass() const
 {
 	return ensure(GUnrealEd) && GUnrealEd->GetUnrealEdOptions()->IsCPPAllowed() && IsParentClassNative();
diff --git a/Engine/Source/Editor/UMGEditor/Private/WidgetBlueprintEditorUtils.cpp b/Engine/Source/Editor/UMGEditor/Private/WidgetBlueprintEditorUtils.cpp
--- a/Engine/Source/Editor/UMGEditor/Private/WidgetBlueprintEditorUtils.cpp
+++ b/Engine/Source/Editor/UMGEditor/Private/WidgetBlueprintEditorUtils.cpp
@@ -54,6 +54,7 @@
 #include "Rendering/SlateDrawBuffer.h"
 #include "Slate/WidgetRenderer.h"
 #include "Widgets/SVirtualWindow.h"
+#include "GraphEditorActions.h"
 
 #define LOCTEXT_NAMESPACE "UMG"
 
@@ -377,18 +378,19 @@ void FWidgetBlueprintEditorUtils::CreateWidgetContextMenu(FMenuBuilder& MenuBuil
 	TSet<FWidgetReference> Widgets = BlueprintEditor->GetSelectedWidgets();
 	UWidgetBlueprint* BP = BlueprintEditor->GetWidgetBlueprintObj();
 
-	MenuBuilder.PushCommandList(BlueprintEditor->DesignerCommandList.ToSharedRef());
-
 	MenuBuilder.BeginSection("Edit", LOCTEXT("Edit", "Edit"));
 	{
-		MenuBuilder.AddMenuEntry(FGenericCommands::Get().Cut);
-		MenuBuilder.AddMenuEntry(FGenericCommands::Get().Copy);
-		MenuBuilder.AddMenuEntry(FGenericCommands::Get().Paste);
-		MenuBuilder.AddMenuEntry(FGenericCommands::Get().Duplicate);
-		MenuBuilder.AddMenuEntry(FGenericCommands::Get().Delete);
-	}
-	MenuBuilder.PopCommandList();
-	{
+		MenuBuilder.PushCommandList(BlueprintEditor->DesignerCommandList.ToSharedRef());
+		{
+			MenuBuilder.AddMenuEntry(FGenericCommands::Get().Cut);
+			MenuBuilder.AddMenuEntry(FGenericCommands::Get().Copy);
+			MenuBuilder.AddMenuEntry(FGenericCommands::Get().Paste);
+			MenuBuilder.AddMenuEntry(FGenericCommands::Get().Duplicate);
+			MenuBuilder.AddMenuEntry(FGenericCommands::Get().Delete);
+			MenuBuilder.AddMenuEntry(FGraphEditorCommands::Get().FindReferences);
+		}
+		MenuBuilder.PopCommandList();
+
 		MenuBuilder.AddMenuEntry(FGenericCommands::Get().Rename);
 	}
 	MenuBuilder.EndSection();
diff --git a/Engine/Source/Editor/UMGEditor/Public/WidgetBlueprintEditor.h b/Engine/Source/Editor/UMGEditor/Public/WidgetBlueprintEditor.h
--- a/Engine/Source/Editor/UMGEditor/Public/WidgetBlueprintEditor.h
+++ b/Engine/Source/Editor/UMGEditor/Public/WidgetBlueprintEditor.h
@@ -339,6 +339,9 @@ private:
 	bool CanDuplicateSelectedWidgets();
 	void DuplicateSelectedWidgets();
 
+	void OnFindWidgetReferences();
+	bool CanFindWidgetReferences() const;
+
 	/** Is creating a native base class for the current widget blueprint allowed */
 	bool CanCreateNativeBaseClass() const;
 
