Code quality: small code cleanup on access to UUserInterfaceSettings.
A wrong pattern used in a lot of places, that was copied by me and noticed in a review.

diff --git a/Engine/Plugins/Experimental/CommonUI/Source/CommonInput/Private/CommonInputBaseTypes.cpp b/Engine/Plugins/Experimental/CommonUI/Source/CommonInput/Private/CommonInputBaseTypes.cpp
--- a/Engine/Plugins/Experimental/CommonUI/Source/CommonInput/Private/CommonInputBaseTypes.cpp
+++ b/Engine/Plugins/Experimental/CommonUI/Source/CommonInput/Private/CommonInputBaseTypes.cpp
@@ -27,15 +27,13 @@ FCommonInputKeySetBrushConfiguration::FCommonInputKeySetBrushConfiguration()
 
 bool UCommonUIInputData::NeedsLoadForServer() const
 {
-	const UUserInterfaceSettings* UISettings = GetDefault<UUserInterfaceSettings>(UUserInterfaceSettings::StaticClass());
-	check(UISettings);
+	const UUserInterfaceSettings* UISettings = GetDefault<UUserInterfaceSettings>();
 	return UISettings->bLoadWidgetsOnDedicatedServer;
 }
 
 bool UCommonInputBaseControllerData::NeedsLoadForServer() const
 {
-	const UUserInterfaceSettings* UISettings = GetDefault<UUserInterfaceSettings>(UUserInterfaceSettings::StaticClass());
-	check(UISettings);
+	const UUserInterfaceSettings* UISettings = GetDefault<UUserInterfaceSettings>();
 	return UISettings->bLoadWidgetsOnDedicatedServer;
 }
 
diff --git a/Engine/Plugins/Experimental/CommonUI/Source/CommonUI/Private/CommonButtonBase.cpp b/Engine/Plugins/Experimental/CommonUI/Source/CommonUI/Private/CommonButtonBase.cpp
--- a/Engine/Plugins/Experimental/CommonUI/Source/CommonUI/Private/CommonButtonBase.cpp
+++ b/Engine/Plugins/Experimental/CommonUI/Source/CommonUI/Private/CommonButtonBase.cpp
@@ -34,9 +34,7 @@
 
 bool UCommonButtonStyle::NeedsLoadForServer() const
 {
-	const UUserInterfaceSettings* UISettings = GetDefault<UUserInterfaceSettings>(UUserInterfaceSettings::StaticClass());
-	check(UISettings);
-	return UISettings->bLoadWidgetsOnDedicatedServer;
+	return GetDefault<UUserInterfaceSettings>()->bLoadWidgetsOnDedicatedServer;
 }
 
 void UCommonButtonStyle::GetButtonPadding(FMargin& OutButtonPadding) const
diff --git a/Engine/Plugins/Experimental/VirtualProductionUtilities/Source/VPUtilities/Private/VPFullScreenUserWidget.cpp b/Engine/Plugins/Experimental/VirtualProductionUtilities/Source/VPUtilities/Private/VPFullScreenUserWidget.cpp
--- a/Engine/Plugins/Experimental/VirtualProductionUtilities/Source/VPUtilities/Private/VPFullScreenUserWidget.cpp
+++ b/Engine/Plugins/Experimental/VirtualProductionUtilities/Source/VPUtilities/Private/VPFullScreenUserWidget.cpp
@@ -899,13 +899,8 @@ float UVPFullScreenUserWidget::GetViewportDPIScale()
 	else
 	{
 		// Otherwise when in Editor mode, the editor automatically scales to the platform size, so we only care about the UI scale
-		FIntPoint ViewportSize = FindSceneViewportSize().IntPoint();
-
-		const UUserInterfaceSettings* UserInterfaceSettings = GetDefault<UUserInterfaceSettings>(UUserInterfaceSettings::StaticClass());
-		if (UserInterfaceSettings)
-		{
-			UIScale = UserInterfaceSettings->GetDPIScaleBasedOnSize(ViewportSize);
-		}
+		const FIntPoint ViewportSize = FindSceneViewportSize().IntPoint();
+		UIScale = GetDefault<UUserInterfaceSettings>()->GetDPIScaleBasedOnSize(ViewportSize);
 	}
 
 	return UIScale;
diff --git a/Engine/Source/Editor/UMGEditor/Private/Designer/SDesignerView.cpp b/Engine/Source/Editor/UMGEditor/Private/Designer/SDesignerView.cpp
--- a/Engine/Source/Editor/UMGEditor/Private/Designer/SDesignerView.cpp
+++ b/Engine/Source/Editor/UMGEditor/Private/Designer/SDesignerView.cpp
@@ -3031,8 +3031,8 @@ FText SDesignerView::GetCurrentDPIScaleText() const
 	Options.MaximumFractionalDigits = 2;
 	Options.MinimumFractionalDigits = 1;
 
-	const UUserInterfaceSettings* UISettings = GetDefault<UUserInterfaceSettings>(UUserInterfaceSettings::StaticClass());
-	if (UISettings && UISettings->UIScaleRule == EUIScalingRule::Custom)
+	const UUserInterfaceSettings* UISettings = GetDefault<UUserInterfaceSettings>();
+	if (UISettings->UIScaleRule == EUIScalingRule::Custom)
 	{
 		UClass* CustomScalingRuleClassInstance = UISettings->CustomScalingRuleClass.TryLoadClass<UDPICustomScalingRule>();
 
@@ -3048,8 +3048,8 @@ FText SDesignerView::GetCurrentDPIScaleText() const
 
 FSlateColor SDesignerView::GetCurrentDPIScaleColor() const
 {
-	const UUserInterfaceSettings* UISettings = GetDefault<UUserInterfaceSettings>(UUserInterfaceSettings::StaticClass());
-	if (UISettings && UISettings->UIScaleRule == EUIScalingRule::Custom)
+	const UUserInterfaceSettings* UISettings = GetDefault<UUserInterfaceSettings>();
+	if (UISettings->UIScaleRule == EUIScalingRule::Custom)
 	{
 		UClass* CustomScalingRuleClassInstance = UISettings->CustomScalingRuleClass.TryLoadClass<UDPICustomScalingRule>();
 
diff --git a/Engine/Source/Editor/UMGEditor/Private/WidgetBlueprintEditorUtils.cpp b/Engine/Source/Editor/UMGEditor/Private/WidgetBlueprintEditorUtils.cpp
--- a/Engine/Source/Editor/UMGEditor/Private/WidgetBlueprintEditorUtils.cpp
+++ b/Engine/Source/Editor/UMGEditor/Private/WidgetBlueprintEditorUtils.cpp
@@ -2128,7 +2128,7 @@ float FWidgetBlueprintEditorUtils::GetWidgetPreviewDPIScale(UUserWidget* UserWid
 		}
 	}
 
-	return GetDefault<UUserInterfaceSettings>(UUserInterfaceSettings::StaticClass())->GetDPIScaleBasedOnSize(FIntPoint(PreviewSize.X, PreviewSize.Y));
+	return GetDefault<UUserInterfaceSettings>()->GetDPIScaleBasedOnSize(FIntPoint(PreviewSize.X, PreviewSize.Y));
 }
 
 
diff --git a/Engine/Source/Runtime/Engine/Private/Slate/SGameLayerManager.cpp b/Engine/Source/Runtime/Engine/Private/Slate/SGameLayerManager.cpp
--- a/Engine/Source/Runtime/Engine/Private/Slate/SGameLayerManager.cpp
+++ b/Engine/Source/Runtime/Engine/Private/Slate/SGameLayerManager.cpp
@@ -415,14 +415,8 @@ float SGameLayerManager::GetGameViewportDPIScale() const
 		return 1;
 	}
 
-	const auto UserInterfaceSettings = GetDefault<UUserInterfaceSettings>(UUserInterfaceSettings::StaticClass());
-
-	if (UserInterfaceSettings == nullptr)
-	{
-		return 1;
-	}
-
-	FIntPoint ViewportSize = Viewport->GetSize();
+	const UUserInterfaceSettings* UserInterfaceSettings = GetDefault<UUserInterfaceSettings>();
+	const FIntPoint ViewportSize = Viewport->GetSize();
 	float GameUIScale;
 
 	if (bUseScaledDPI)
diff --git a/Engine/Source/Runtime/UMG/Private/Components/Visual.cpp b/Engine/Source/Runtime/UMG/Private/Components/Visual.cpp
--- a/Engine/Source/Runtime/UMG/Private/Components/Visual.cpp
+++ b/Engine/Source/Runtime/UMG/Private/Components/Visual.cpp
@@ -27,7 +27,6 @@ void UVisual::BeginDestroy()
 
 bool UVisual::NeedsLoadForServer() const
 {
-	const UUserInterfaceSettings* UISettings = GetDefault<UUserInterfaceSettings>(UUserInterfaceSettings::StaticClass());
-	check(UISettings);
+	const UUserInterfaceSettings* UISettings = GetDefault<UUserInterfaceSettings>();
 	return UISettings->bLoadWidgetsOnDedicatedServer;
 }
diff --git a/Engine/Source/Runtime/UMG/Private/WidgetBlueprintGeneratedClass.cpp b/Engine/Source/Runtime/UMG/Private/WidgetBlueprintGeneratedClass.cpp
--- a/Engine/Source/Runtime/UMG/Private/WidgetBlueprintGeneratedClass.cpp
+++ b/Engine/Source/Runtime/UMG/Private/WidgetBlueprintGeneratedClass.cpp
@@ -401,8 +401,7 @@ void UWidgetBlueprintGeneratedClass::PurgeClass(bool bRecompilingOnLoad)
 
 bool UWidgetBlueprintGeneratedClass::NeedsLoadForServer() const
 {
-	const UUserInterfaceSettings* UISettings = GetDefault<UUserInterfaceSettings>(UUserInterfaceSettings::StaticClass());
-	check(UISettings);
+	const UUserInterfaceSettings* UISettings = GetDefault<UUserInterfaceSettings>();
 	return UISettings->bLoadWidgetsOnDedicatedServer;
 }
 
