Initialize the UMG widget Blueprint editor according to the standard non-defaults-only workflow.

diff --git a/Engine/Source/Editor/UMGEditor/Private/AssetTypeActions_WidgetBlueprint.cpp b/Engine/Source/Editor/UMGEditor/Private/AssetTypeActions_WidgetBlueprint.cpp
--- a/Engine/Source/Editor/UMGEditor/Private/AssetTypeActions_WidgetBlueprint.cpp
+++ b/Engine/Source/Editor/UMGEditor/Private/AssetTypeActions_WidgetBlueprint.cpp
@@ -22,7 +22,7 @@ void FAssetTypeActions_WidgetBlueprint::OpenAssetEditor( const TArray<UObject*>&
 
 			TArray<UBlueprint*> Blueprints;
 			Blueprints.Add(Blueprint);
-			NewBlueprintEditor->InitWidgetBlueprintEditor(Mode, EditWithinLevelEditor, Blueprints, true);
+			NewBlueprintEditor->InitWidgetBlueprintEditor(Mode, EditWithinLevelEditor, Blueprints, /*bShouldOpenInDefaultsMode = */ false);
 		}
 		else
 		{
