Disable real-time thumbnail on widget blueprint assets to avoid repeated calls to TakeWidget on mouse hover.

diff --git a/Engine/Source/Editor/UMGEditor/Private/WidgetBlueprintThumbnailRenderer.cpp b/Engine/Source/Editor/UMGEditor/Private/WidgetBlueprintThumbnailRenderer.cpp
--- a/Engine/Source/Editor/UMGEditor/Private/WidgetBlueprintThumbnailRenderer.cpp
+++ b/Engine/Source/Editor/UMGEditor/Private/WidgetBlueprintThumbnailRenderer.cpp
@@ -210,6 +210,10 @@ void UWidgetBlueprintThumbnailRenderer::Draw(UObject* Object, int32 X, int32 Y,
 #endif
 }
 
+bool UWidgetBlueprintThumbnailRenderer::AllowsRealtimeThumbnails(UObject* Object) const 
+{ 
+	return false; 
+}
 
 void UWidgetBlueprintThumbnailRenderer::OnBlueprintUnloaded(UBlueprint* Blueprint)
 {
diff --git a/Engine/Source/Editor/UMGEditor/Public/WidgetBlueprintThumbnailRenderer.h b/Engine/Source/Editor/UMGEditor/Public/WidgetBlueprintThumbnailRenderer.h
--- a/Engine/Source/Editor/UMGEditor/Public/WidgetBlueprintThumbnailRenderer.h
+++ b/Engine/Source/Editor/UMGEditor/Public/WidgetBlueprintThumbnailRenderer.h
@@ -22,6 +22,7 @@ class UMGEDITOR_API UWidgetBlueprintThumbnailRenderer : public UDefaultSizedThum
 	//~ Begin UThumbnailRenderer Object
 	bool CanVisualizeAsset(UObject* Object) override;
 	void Draw(UObject* Object, int32 X, int32 Y, uint32 Width, uint32 Height, FRenderTarget* RenderTarget, FCanvas* Canvas, bool bAdditionalViewFamily) override;
+	bool AllowsRealtimeThumbnails(UObject* Object) const override;
 	//~ End UThumbnailRenderer Object
 
 private:
