Fix duplicate moved content issue on inherited name

diff --git a/Engine/Source/Editor/UMGEditor/Private/Hierarchy/SHierarchyViewItem.cpp b/Engine/Source/Editor/UMGEditor/Private/Hierarchy/SHierarchyViewItem.cpp
--- a/Engine/Source/Editor/UMGEditor/Private/Hierarchy/SHierarchyViewItem.cpp
+++ b/Engine/Source/Editor/UMGEditor/Private/Hierarchy/SHierarchyViewItem.cpp
@@ -388,6 +388,24 @@ TOptional<EItemDropZone> ProcessHierarchyDragDrop(const FDragDropEvent& DragDrop
 						}
 					}
 
+					// If this widget inherits from another one, we can't access the inherited named slots by traversing the widget tree from its root.
+					// So we have to look at the NamedSlotBindings to find a named slot for the moved content.
+					else if (Blueprint->ParentClass && Blueprint->ParentClass != UUserWidget::StaticClass())
+					{
+						TArray<FName> SlotNames;
+						Blueprint->WidgetTree->GetSlotNames(SlotNames);
+						for (FName SlotName : SlotNames)
+						{
+							if (UWidget* SlotContent = Blueprint->WidgetTree->GetContentForSlot(SlotName))
+							{
+								if (SlotContent == TemplateWidget)
+								{
+									Blueprint->WidgetTree->SetContentForSlot(SlotName, nullptr);
+								}
+							}
+						}
+					}
+
 					UPanelWidget* OriginalParent = TemplateWidget->GetParent();
 					UBlueprint* OriginalBP = nullptr;
 
