Fix compilation error when binding animation to an inherited widget in UMG.

diff --git a/Engine/Source/Editor/UMGEditor/Private/WidgetBlueprintCompiler.cpp b/Engine/Source/Editor/UMGEditor/Private/WidgetBlueprintCompiler.cpp
--- a/Engine/Source/Editor/UMGEditor/Private/WidgetBlueprintCompiler.cpp
+++ b/Engine/Source/Editor/UMGEditor/Private/WidgetBlueprintCompiler.cpp
@@ -1151,22 +1151,24 @@ void FWidgetBlueprintCompilerContext::ValidateWidgetAnimations()
 {
 	UWidgetBlueprintGeneratedClass* WidgetClass = NewWidgetBlueprintClass;
 	UWidgetBlueprint* WidgetBP = WidgetBlueprint();
-	UUserWidget& UserWidget = *WidgetClass->GetDefaultObject<UUserWidget>();
+	UUserWidget* UserWidget = WidgetClass->GetDefaultObject<UUserWidget>();
 	FBlueprintCompilerLog BlueprintLog(MessageLog, WidgetClass);
-
+	
+	UWidgetTree* LatestWidgetTree = FWidgetBlueprintEditorUtils::FindLatestWidgetTree(WidgetBP, UserWidget);
+	
 	for (const UWidgetAnimation* InAnimation : WidgetBP->Animations)
 	{
 		for (const FWidgetAnimationBinding& Binding : InAnimation->AnimationBindings)
 		{
 			// Look for the object bindings within the widget
-			UObject* FoundObject = Binding.FindRuntimeObject(*WidgetBP->WidgetTree, UserWidget);
+			UObject* FoundObject = Binding.FindRuntimeObject(*LatestWidgetTree, *UserWidget);
 
 			// If any of the FoundObjects is null, we do not play the animation.
 			if (FoundObject == nullptr)
 			{
 				// Notify the user of the null track in the editor
 				const FText AnimationNullTrackMessage = LOCTEXT("AnimationNullTrack", "UMG Animation '{0}' from '{1}' is trying to animate a non-existent widget through binding '{2}'. Please re-bind or delete this object from the animation.");
-				BlueprintLog.Warning(FText::Format(AnimationNullTrackMessage, InAnimation->GetDisplayName(), FText::FromString(UserWidget.GetClass()->GetName()), FText::FromName(Binding.WidgetName)));
+				BlueprintLog.Warning(FText::Format(AnimationNullTrackMessage, InAnimation->GetDisplayName(), FText::FromString(UserWidget->GetClass()->GetName()), FText::FromName(Binding.WidgetName)));
 			}
 		}
 	}
diff --git a/Engine/Source/Editor/UMGEditor/Private/WidgetBlueprintEditor.cpp b/Engine/Source/Editor/UMGEditor/Private/WidgetBlueprintEditor.cpp
--- a/Engine/Source/Editor/UMGEditor/Private/WidgetBlueprintEditor.cpp
+++ b/Engine/Source/Editor/UMGEditor/Private/WidgetBlueprintEditor.cpp
@@ -1661,19 +1661,7 @@ void FWidgetBlueprintEditor::UpdatePreview(UBlueprint* InBlueprint, bool bInForc
 				PreviewUserWidget->SetPlayerContext(FLocalPlayerContext(Player));
 			}
 
-			UWidgetTree* LatestWidgetTree = PreviewBlueprint->WidgetTree;
-
-			// If there is no RootWidget, we look for a WidgetTree in the parents classes until we find one.
-			if (LatestWidgetTree->RootWidget == nullptr)
-			{
-				UWidgetBlueprintGeneratedClass* BGClass = PreviewUserWidget->GetWidgetTreeOwningClass();
-				// If we find a class that owns the widget tree, just make sure it's not our current class, that would imply we've removed all the widgets
-				// from this current tree, and if we use this classes compiled tree it's going to be the outdated old version.
-				if (BGClass && BGClass != PreviewBlueprint->GeneratedClass)
-				{
-					LatestWidgetTree = BGClass->GetWidgetTreeArchetype();
-				}
-			}
+			UWidgetTree* LatestWidgetTree = FWidgetBlueprintEditorUtils::FindLatestWidgetTree(PreviewBlueprint, PreviewUserWidget);
 
 			TMap<FName, UWidget*> NamedSlotContentToMerge;
 			UWidgetBlueprint* WidgetBPIt = PreviewBlueprint;
diff --git a/Engine/Source/Editor/UMGEditor/Private/WidgetBlueprintEditorUtils.cpp b/Engine/Source/Editor/UMGEditor/Private/WidgetBlueprintEditorUtils.cpp
--- a/Engine/Source/Editor/UMGEditor/Private/WidgetBlueprintEditorUtils.cpp
+++ b/Engine/Source/Editor/UMGEditor/Private/WidgetBlueprintEditorUtils.cpp
@@ -2028,6 +2028,24 @@ FString FWidgetBlueprintEditorUtils::FindNextValidName(UWidgetTree* WidgetTree,
 	return Name;
 }
 
+UWidgetTree* FWidgetBlueprintEditorUtils::FindLatestWidgetTree(UWidgetBlueprint* Blueprint, UUserWidget* UserWidget)
+{
+	UWidgetTree* LatestWidgetTree = Blueprint->WidgetTree;
+
+	// If there is no RootWidget, we look for a WidgetTree in the parents classes until we find one.
+	if (LatestWidgetTree->RootWidget == nullptr)
+	{
+		UWidgetBlueprintGeneratedClass* BGClass = UserWidget->GetWidgetTreeOwningClass();
+		// If we find a class that owns the widget tree, just make sure it's not our current class, that would imply we've removed all the widgets
+		// from this current tree, and if we use this classes compiled tree it's going to be the outdated old version.
+		if (BGClass && BGClass != Blueprint->GeneratedClass)
+		{
+			LatestWidgetTree = BGClass->GetWidgetTreeArchetype();
+		}
+	}
+	return LatestWidgetTree;
+}
+
 int32 FWidgetBlueprintEditorUtils::UpdateHittestGrid(FHittestGrid& HitTestGrid, TSharedRef<SWindow> Window, float Scale, FVector2D DrawSize, float DeltaTime)
 {
 	FSlateApplication::Get().InvalidateAllWidgets(false);
diff --git a/Engine/Source/Editor/UMGEditor/Public/WidgetBlueprintEditorUtils.h b/Engine/Source/Editor/UMGEditor/Public/WidgetBlueprintEditorUtils.h
--- a/Engine/Source/Editor/UMGEditor/Public/WidgetBlueprintEditorUtils.h
+++ b/Engine/Source/Editor/UMGEditor/Public/WidgetBlueprintEditorUtils.h
@@ -74,6 +74,8 @@ public:
 
 	static bool ReplaceNamedSlotHostContent(UWidget* WidgetTemplate, TScriptInterface<INamedSlotInterface> NamedSlotHost, UWidget* NewContentWidget);
 
+	static UWidgetTree* FindLatestWidgetTree(UWidgetBlueprint* Blueprint, UUserWidget* UserWidget);
+
 	static int32 UpdateHittestGrid(FHittestGrid& HitTestGrid, TSharedRef<SWindow> Window, float Scale, FVector2D DrawSize, float DeltaTime);
 
 	static TTuple<FVector2D, FVector2D> GetWidgetPreviewAreaAndSize(UUserWidget* UserWidget, FVector2D DesiredSize, FVector2D PreviewSize, EDesignPreviewSizeMode ThumbnailSizeMode, TOptional<FVector2D> ThumbnailCustomSize);
