PR #9766: Default preview resolution of widgets in UMG can be set in Editor Preferences

diff --git a/Engine/Source/Editor/UMGEditor/Private/Designer/SDesignerView.cpp b/Engine/Source/Editor/UMGEditor/Private/Designer/SDesignerView.cpp
--- a/Engine/Source/Editor/UMGEditor/Private/Designer/SDesignerView.cpp
+++ b/Engine/Source/Editor/UMGEditor/Private/Designer/SDesignerView.cpp
@@ -217,9 +217,6 @@ UWidget* SDesignerView::GetWidgetInDesignScopeFromSlateWidget(TSharedRef<SWidget
 // SDesignerView
  
 const FString SDesignerView::ConfigSectionName = "UMGEditor.Designer";
-const uint32 SDesignerView::DefaultResolutionWidth = 1280;
-const uint32 SDesignerView::DefaultResolutionHeight = 720;
-const FString SDesignerView::DefaultAspectRatio = "16:9";
 const FString SDesignerView::DefaultPreviewOverrideName = "";
 
 void SDesignerView::Construct(const FArguments& InArgs, TSharedPtr<FWidgetBlueprintEditor> InBlueprintEditor)
@@ -874,28 +871,38 @@ bool SDesignerView::IsRespectingLocks() const
 
 void SDesignerView::SetStartupResolution()
 {
-	// Use previously set resolution (or create new entries using default values)
+	// Whether the user selected a common resolution.
+	if (!GConfig->GetBool(*ConfigSectionName, TEXT("bCommonResolutionSelected"), bCommonResolutionSelected, GEditorPerProjectIni))
+	{
+		GConfig->SetBool(*ConfigSectionName, TEXT("bCommonResolutionSelected"), false, GEditorPerProjectIni);
+		bCommonResolutionSelected = false;
+	}
+	// Use user-set resolution
+	const UWidgetDesignerSettings* DesignerSettings = GetDefault<const UWidgetDesignerSettings>();
+	const FUintVector2 DefaultPreviewResolution = DesignerSettings->DefaultPreviewResolution;
 	// Width
-	if (!GConfig->GetInt(*ConfigSectionName, TEXT("PreviewWidth"), PreviewWidth, GEditorPerProjectIni))
+	if (!GConfig->GetInt(*ConfigSectionName, TEXT("PreviewWidth"), PreviewWidth, GEditorPerProjectIni) || !bCommonResolutionSelected)
 	{
-		GConfig->SetInt(*ConfigSectionName, TEXT("PreviewWidth"), DefaultResolutionWidth, GEditorPerProjectIni);
-		PreviewWidth = DefaultResolutionWidth;
+		GConfig->SetInt(*ConfigSectionName, TEXT("PreviewWidth"), DefaultPreviewResolution.X, GEditorPerProjectIni);
+		PreviewWidth = DefaultPreviewResolution.X;
 	}
 	// Initially assign WidthReadFromSettings to PreviewWidth
 	WidthReadFromSettings = PreviewWidth;
 	// Height
-	if (!GConfig->GetInt(*ConfigSectionName, TEXT("PreviewHeight"), PreviewHeight, GEditorPerProjectIni))
+	PreviewOverrideName = DefaultPreviewOverrideName;
+	if (!GConfig->GetInt(*ConfigSectionName, TEXT("PreviewHeight"), PreviewHeight, GEditorPerProjectIni) || !bCommonResolutionSelected)
 	{
-		GConfig->SetInt(*ConfigSectionName, TEXT("PreviewHeight"), DefaultResolutionHeight, GEditorPerProjectIni);
-		PreviewHeight = DefaultResolutionHeight;
+		GConfig->SetInt(*ConfigSectionName, TEXT("PreviewHeight"), DefaultPreviewResolution.Y, GEditorPerProjectIni);
+		PreviewHeight = DefaultPreviewResolution.Y;
 	}
 	// Initially assign HeightReadFromSettings to PreviewHeight
 	HeightReadFromSettings = PreviewHeight;
 	// Aspect Ratio
-	if (!GConfig->GetString(*ConfigSectionName, TEXT("PreviewAspectRatio"), PreviewAspectRatio, GEditorPerProjectIni))
+	if (!GConfig->GetString(*ConfigSectionName, TEXT("PreviewAspectRatio"), PreviewAspectRatio, GEditorPerProjectIni) || !bCommonResolutionSelected)
 	{
-		GConfig->SetString(*ConfigSectionName, TEXT("PreviewAspectRatio"), *DefaultAspectRatio, GEditorPerProjectIni);
-		PreviewAspectRatio = DefaultAspectRatio;
+		const int32 GCD = FMath::GreatestCommonDivisor(PreviewWidth, PreviewHeight);
+		PreviewAspectRatio = FString::Printf(TEXT("%d:%d"), PreviewWidth / GCD, PreviewHeight / GCD);
+		GConfig->SetString(*ConfigSectionName, TEXT("PreviewAspectRatio"), *PreviewAspectRatio, GEditorPerProjectIni);
 	}
 	// Portrait Mode
 	if (!GConfig->GetBool(*ConfigSectionName, TEXT("bIsInPortraitMode"), bPreviewIsPortrait, GEditorPerProjectIni))
@@ -1036,7 +1043,7 @@ void SDesignerView::SetPreviewAreaSize(int32 Width, int32 Height)
 
 				PreviewWidth = Width;
 				PreviewHeight = Height;
-				PreviewAspectRatio = FString::Printf(TEXT("%d:%d"), Height / GCD, Width / GCD);
+				PreviewAspectRatio = FString::Printf(TEXT("%d:%d"), Width / GCD, Height / GCD);
 
 				const bool bSaveChanges = false;
 				if (bSaveChanges)
@@ -1048,6 +1055,7 @@ void SDesignerView::SetPreviewAreaSize(int32 Width, int32 Height)
 					GConfig->SetString(*ConfigSectionName, TEXT("ProfileName"), *PreviewOverrideName, GEditorPerProjectIni);
 					GConfig->SetFloat(*ConfigSectionName, TEXT("ScaleFactor"), ScaleFactor, GEditorPerProjectIni);
 					GConfig->SetBool(*ConfigSectionName, TEXT("bCanPreviewSwapAspectRatio"), bCanPreviewSwapAspectRatio, GEditorPerProjectIni);
+					GConfig->SetBool(*ConfigSectionName, TEXT("bCommonResolutionSelected"), false, GEditorPerProjectIni);
 				}
 				break;
 			}
@@ -3276,7 +3284,7 @@ void SDesignerView::HandleOnCommonResolutionSelected(const FPlayScreenResolution
 	GConfig->SetString(*ConfigSectionName, TEXT("ProfileName"), *PreviewOverrideName, GEditorPerProjectIni);
 	GConfig->SetFloat(*ConfigSectionName, TEXT("ScaleFactor"), ScaleFactor, GEditorPerProjectIni);
 	GConfig->SetBool(*ConfigSectionName, TEXT("bCanPreviewSwapAspectRatio"), bCanPreviewSwapAspectRatio, GEditorPerProjectIni);
-
+	GConfig->SetBool(*ConfigSectionName, TEXT("bCommonResolutionSelected"), true, GEditorPerProjectIni);
 	if (!PreviewOverrideName.IsEmpty())
 	{
 		ULevelEditorPlaySettings* PlayInSettings = GetMutableDefault<ULevelEditorPlaySettings>();
diff --git a/Engine/Source/Editor/UMGEditor/Private/Settings/WidgetDesignerSettings.cpp b/Engine/Source/Editor/UMGEditor/Private/Settings/WidgetDesignerSettings.cpp
--- a/Engine/Source/Editor/UMGEditor/Private/Settings/WidgetDesignerSettings.cpp
+++ b/Engine/Source/Editor/UMGEditor/Private/Settings/WidgetDesignerSettings.cpp
@@ -16,6 +16,7 @@ UWidgetDesignerSettings::UWidgetDesignerSettings()
 	bRespectLocks = true;
 	CreateOnCompile = EDisplayOnCompile::DoC_ErrorsOrWarnings;
 	DismissOnCompile = EDisplayOnCompile::DoC_ErrorsOrWarnings;
+	DefaultPreviewResolution = FUintVector2(1280, 720);
 }
 
 #if WITH_EDITOR
diff --git a/Engine/Source/Editor/UMGEditor/Public/Designer/SDesignerView.h b/Engine/Source/Editor/UMGEditor/Public/Designer/SDesignerView.h
--- a/Engine/Source/Editor/UMGEditor/Public/Designer/SDesignerView.h
+++ b/Engine/Source/Editor/UMGEditor/Public/Designer/SDesignerView.h
@@ -265,9 +265,6 @@ private:
 
 private:
 	static const FString ConfigSectionName;
-	static const uint32 DefaultResolutionWidth;
-	static const uint32 DefaultResolutionHeight;
-	static const FString DefaultAspectRatio;
 	static const FString DefaultPreviewOverrideName;
 
 	/** Extensions for the designer to allow for custom widgets to be inserted onto the design surface as selection changes. */
@@ -366,6 +363,9 @@ private:
 	/** The original Height of the preview area, read from the settings file. */
 	int32 HeightReadFromSettings;
 
+	// Whether we have selected one of the common resolutions for the preview.
+	bool bCommonResolutionSelected;
+
 	/***/
 	bool bShowResolutionOutlines;
 
diff --git a/Engine/Source/Editor/UMGEditor/Public/Settings/WidgetDesignerSettings.h b/Engine/Source/Editor/UMGEditor/Public/Settings/WidgetDesignerSettings.h
--- a/Engine/Source/Editor/UMGEditor/Public/Settings/WidgetDesignerSettings.h
+++ b/Engine/Source/Editor/UMGEditor/Public/Settings/WidgetDesignerSettings.h
@@ -60,6 +60,10 @@ public:
 	 */
 	UPROPERTY(EditAnywhere, config, Category = Dragging)
 	bool bLockToPanelOnDragByDefault;
+	
+	/** The default preview resolution in the designer.*/
+	UPROPERTY(EditAnywhere, config, Category = Visuals)
+	FUintVector2 DefaultPreviewResolution;
 
 	/**
 	 * Should the designer show outlines by default?
