The Bind menu is now shown even for properties that MVVM can bind to, but legacy UMG cannot.

diff --git a/Engine/Plugins/Runtime/ModelViewViewModel/Source/ModelViewViewModelEditor/Private/Customizations/MVVMPropertyBindingExtension.cpp b/Engine/Plugins/Runtime/ModelViewViewModel/Source/ModelViewViewModelEditor/Private/Customizations/MVVMPropertyBindingExtension.cpp
--- a/Engine/Plugins/Runtime/ModelViewViewModel/Source/ModelViewViewModelEditor/Private/Customizations/MVVMPropertyBindingExtension.cpp
+++ b/Engine/Plugins/Runtime/ModelViewViewModel/Source/ModelViewViewModelEditor/Private/Customizations/MVVMPropertyBindingExtension.cpp
@@ -2,6 +2,7 @@
 
 #include "Customizations/MVVMPropertyBindingExtension.h"
 
+#include "BlueprintEditor.h"
 #include "Components/Widget.h"
 #include "Framework/MultiBox/MultiBoxBuilder.h"
 #include "Kismet2/BlueprintEditorUtils.h"
@@ -163,11 +164,46 @@ TOptional<FName> FMVVMPropertyBindingExtension::GetCurrentValue(const UWidgetBlu
 	TArray<FName> Names = Binding->ViewModelPath.GetPaths();
 	if (Names.Num() > 0)
 	{
-		return Names[0];
+		return Names.Last();
 	}
 	return TOptional<FName>();
 }
 
+const FSlateBrush* FMVVMPropertyBindingExtension::GetCurrentIcon(const UWidgetBlueprint* WidgetBlueprint, const UWidget* Widget, const FProperty* Property) const
+{
+	const UMVVMWidgetBlueprintExtension_View* MVVMExtensionPtr = UMVVMWidgetBlueprintExtension_View::GetExtension<UMVVMWidgetBlueprintExtension_View>(WidgetBlueprint);
+	if (MVVMExtensionPtr == nullptr)
+	{
+		return nullptr;
+	}
+	const UMVVMBlueprintView* MVVMBlueprintView = MVVMExtensionPtr->GetBlueprintView();
+
+	const FMVVMBlueprintViewBinding* Binding = MVVMBlueprintView ? MVVMBlueprintView->FindBinding(Widget, Property) : nullptr;
+	if (Binding == nullptr)
+	{
+		return nullptr;
+	}
+
+	TArray<UE::MVVM::FMVVMConstFieldVariant> Fields = Binding->ViewModelPath.GetFields();
+	if (Fields.IsEmpty())
+	{
+		return nullptr;
+	}
+
+	UE::MVVM::FMVVMConstFieldVariant Field = Fields.Last();
+	if (Field.IsFunction() && Field.GetFunction() != nullptr)
+	{
+		return FAppStyle::Get().GetBrush("GraphEditor.Function_16x");
+	}
+	else if (Field.IsProperty() && Field.GetProperty() != nullptr)
+	{
+		FSlateColor PrimaryColor, SecondaryColor;
+		const FSlateBrush* SecondaryBrush = nullptr;
+		return FBlueprintEditor::GetVarIconAndColorFromProperty(Field.GetProperty(), PrimaryColor, SecondaryBrush, SecondaryColor);
+	}
+
+	return nullptr;
+}
 
 void FMVVMPropertyBindingExtension::ClearCurrentValue(const UWidgetBlueprint* WidgetBlueprint, const UWidget* Widget, const FProperty* Property)
 {
diff --git a/Engine/Plugins/Runtime/ModelViewViewModel/Source/ModelViewViewModelEditor/Private/Customizations/MVVMPropertyBindingExtension.h b/Engine/Plugins/Runtime/ModelViewViewModel/Source/ModelViewViewModelEditor/Private/Customizations/MVVMPropertyBindingExtension.h
--- a/Engine/Plugins/Runtime/ModelViewViewModel/Source/ModelViewViewModelEditor/Private/Customizations/MVVMPropertyBindingExtension.h
+++ b/Engine/Plugins/Runtime/ModelViewViewModel/Source/ModelViewViewModelEditor/Private/Customizations/MVVMPropertyBindingExtension.h
@@ -11,6 +11,8 @@ class FMVVMPropertyBindingExtension
 	: public IPropertyBindingExtension
 {
 	virtual TOptional<FName> GetCurrentValue(const UWidgetBlueprint* WidgetBlueprint, const UWidget* Widget, const FProperty* Property) const override;
+	virtual const FSlateBrush* GetCurrentIcon(const UWidgetBlueprint* WidgetBlueprint, const UWidget* Widget, const FProperty* Property) const override;
+
 	virtual void ClearCurrentValue(const UWidgetBlueprint* WidgetBlueprint, const UWidget* Widget, const FProperty* Property) override;
 	virtual TSharedPtr<FExtender> CreateMenuExtender(const UWidgetBlueprint* WidgetBlueprint, const UWidget* Widget, const FProperty* Property) override;
 	virtual bool CanExtend(const UWidgetBlueprint* WidgetBlueprint, const UWidget* Widget, const FProperty* Property) const;
diff --git a/Engine/Source/Editor/UMGEditor/Private/Customizations/UMGDetailCustomizations.cpp b/Engine/Source/Editor/UMGEditor/Private/Customizations/UMGDetailCustomizations.cpp
--- a/Engine/Source/Editor/UMGEditor/Private/Customizations/UMGDetailCustomizations.cpp
+++ b/Engine/Source/Editor/UMGEditor/Private/Customizations/UMGDetailCustomizations.cpp
@@ -101,7 +101,7 @@ private:
 	TSharedPtr<FEdGraphSchemaAction> Action;
 };
 
-TSharedRef<SWidget> FBlueprintWidgetCustomization::MakePropertyBindingWidget(TWeakPtr<FWidgetBlueprintEditor> InEditor, FDelegateProperty* InDelegateProperty, TSharedRef<IPropertyHandle> InPropertyHandle, bool bInGeneratePureBindings)
+TSharedRef<SWidget> FBlueprintWidgetCustomization::MakePropertyBindingWidget(TWeakPtr<FWidgetBlueprintEditor> InEditor, UFunction* SignatureFunction, TSharedRef<IPropertyHandle> InPropertyHandle, bool bInGeneratePureBindings)
 {
 	if (!IModularFeatures::Get().IsModularFeatureAvailable("PropertyAccessEditor"))
 	{
@@ -123,7 +123,7 @@ TSharedRef<SWidget> FBlueprintWidgetCustomization::MakePropertyBindingWidget(TWe
 	UWidget* Widget = Objects.Num() ? Cast<UWidget>(Objects[0]) : nullptr;
 
 	FString WidgetName;
-	if ( Widget && !Widget->IsGeneratedName() )
+	if (Widget && !Widget->IsGeneratedName())
 	{
 		WidgetName = TEXT("_") + Widget->GetName() + TEXT("_");
 	}
@@ -132,19 +132,23 @@ TSharedRef<SWidget> FBlueprintWidgetCustomization::MakePropertyBindingWidget(TWe
 
 	TArray<TSharedPtr<FExtender>> MenuExtenders;
 
+	// cached list of extensions for which CanExtend() returned true
+	TArray<TSharedPtr<IPropertyBindingExtension>> ActiveExtensions;
+
 	IUMGEditorModule& EditorModule = FModuleManager::LoadModuleChecked<IUMGEditorModule>("UMGEditor");
 	for (const TSharedPtr<IPropertyBindingExtension>& Extension : EditorModule.GetPropertyBindingExtensibilityManager()->GetExtensions())
 	{
 		if (Extension->CanExtend(WidgetBlueprint, Widget, InPropertyHandle->GetProperty()))
 		{
 			MenuExtenders.Add(Extension->CreateMenuExtender(WidgetBlueprint, Widget, InPropertyHandle->GetProperty()));
+			ActiveExtensions.Add(Extension);
 		}
 	}
 
 	FPropertyBindingWidgetArgs Args;
 	Args.MenuExtender = FExtender::Combine(MenuExtenders);
 	Args.Property = InPropertyHandle->GetProperty();
-	Args.BindableSignature = InDelegateProperty->SignatureFunction;
+	Args.BindableSignature = SignatureFunction;
 	Args.OnGenerateBindingName = FOnGenerateBindingName::CreateLambda([WidgetName]()
 	{
 		return WidgetName;
@@ -230,23 +234,26 @@ TSharedRef<SWidget> FBlueprintWidgetCustomization::MakePropertyBindingWidget(TWe
 		return false;
 	});
 
-	Args.OnCanBindProperty = FOnCanBindProperty::CreateLambda([InDelegateProperty](FProperty* InProperty)
+	Args.OnCanBindProperty = FOnCanBindProperty::CreateLambda([SignatureFunction](FProperty* InProperty)
 	{
-		if ( FProperty* ReturnProperty = InDelegateProperty->SignatureFunction->GetReturnProperty() )
+		if (SignatureFunction != nullptr)
 		{
-			// Find the binder that can handle the delegate return type.
-			TSubclassOf<UPropertyBinding> Binder = UWidget::FindBinderClassForDestination(ReturnProperty);
-			if ( Binder != nullptr )
+			if (FProperty* ReturnProperty = SignatureFunction->GetReturnProperty() )
 			{
-				// Ensure that the binder also can handle binding from the property we care about.
-				return ( Binder->GetDefaultObject<UPropertyBinding>()->IsSupportedSource(InProperty) );
+				// Find the binder that can handle the delegate return type.
+				TSubclassOf<UPropertyBinding> Binder = UWidget::FindBinderClassForDestination(ReturnProperty);
+				if ( Binder != nullptr )
+				{
+					// Ensure that the binder also can handle binding from the property we care about.
+					return ( Binder->GetDefaultObject<UPropertyBinding>()->IsSupportedSource(InProperty) );
+				}
 			}
 		}
 
 		return false;
 	});
 	
-	Args.OnCanBindFunction = FOnCanBindFunction::CreateLambda([InDelegateProperty](UFunction* InFunction)
+	Args.OnCanBindFunction = FOnCanBindFunction::CreateLambda([SignatureFunction](UFunction* InFunction)
 	{
 		auto HasFunctionBinder = [InFunction](UFunction* InBindableSignature)
 		{
@@ -273,9 +280,14 @@ TSharedRef<SWidget> FBlueprintWidgetCustomization::MakePropertyBindingWidget(TWe
 			return false;
 		};
 
+		if (SignatureFunction == nullptr)
+		{
+			return false;
+		}
+
 		// We ignore CPF_ReturnParm because all that matters for binding to script functions is that the number of out parameters match.
-		return ( InFunction->IsSignatureCompatibleWith(InDelegateProperty->SignatureFunction, UFunction::GetDefaultIgnoredSignatureCompatibilityFlags() | CPF_ReturnParm) ||
-				HasFunctionBinder(InDelegateProperty->SignatureFunction) );
+		return ( InFunction->IsSignatureCompatibleWith(SignatureFunction, UFunction::GetDefaultIgnoredSignatureCompatibilityFlags() | CPF_ReturnParm) ||
+				HasFunctionBinder(SignatureFunction) );
 	});
 
 	Args.OnCanBindToClass = FOnCanBindToClass::CreateLambda([](UClass* InClass)
@@ -366,7 +378,7 @@ TSharedRef<SWidget> FBlueprintWidgetCustomization::MakePropertyBindingWidget(TWe
 		FBlueprintEditorUtils::MarkBlueprintAsStructurallyModified(ThisBlueprint);	
 	});
 
-	Args.OnRemoveBinding = FOnRemoveBinding::CreateLambda([InEditor, Objects, InPropertyHandle](FName InPropertyName)
+	Args.OnRemoveBinding = FOnRemoveBinding::CreateLambda([InEditor, Objects, InPropertyHandle, ActiveExtensions](FName InPropertyName)
 	{
 		UWidgetBlueprint* ThisBlueprint = InEditor.Pin()->GetWidgetBlueprintObj();
 
@@ -390,8 +402,7 @@ TSharedRef<SWidget> FBlueprintWidgetCustomization::MakePropertyBindingWidget(TWe
 
 			if (UWidget* Widget = Cast<UWidget>(Object))
 			{
-				IUMGEditorModule& EditorModule = FModuleManager::LoadModuleChecked<IUMGEditorModule>("UMGEditor");
-				for (const TSharedPtr<IPropertyBindingExtension>& Extension : EditorModule.GetPropertyBindingExtensibilityManager()->GetExtensions())
+				for (const TSharedPtr<IPropertyBindingExtension>& Extension : ActiveExtensions)
 				{
 					Extension->ClearCurrentValue(ThisBlueprint, Widget, InPropertyHandle->GetProperty());
 				}
@@ -399,11 +410,9 @@ TSharedRef<SWidget> FBlueprintWidgetCustomization::MakePropertyBindingWidget(TWe
 		}
 
 		FBlueprintEditorUtils::MarkBlueprintAsStructurallyModified(ThisBlueprint);
-
-		
 	});
 
-	Args.OnCanRemoveBinding = FOnCanRemoveBinding::CreateLambda([InEditor, Objects, InPropertyHandle](FName InPropertyName)
+	Args.OnCanRemoveBinding = FOnCanRemoveBinding::CreateLambda([InEditor, Objects, InPropertyHandle, ActiveExtensions](FName InPropertyName)
 	{
 		UWidgetBlueprint* ThisBlueprint = InEditor.Pin()->GetWidgetBlueprintObj();
 
@@ -427,8 +436,7 @@ TSharedRef<SWidget> FBlueprintWidgetCustomization::MakePropertyBindingWidget(TWe
 
 			if (UWidget* Widget = Cast<UWidget>(Object))
 			{
-				IUMGEditorModule& EditorModule = FModuleManager::LoadModuleChecked<IUMGEditorModule>("UMGEditor");
-				for (const TSharedPtr<IPropertyBindingExtension>& Extension : EditorModule.GetPropertyBindingExtensibilityManager()->GetExtensions())
+				for (const TSharedPtr<IPropertyBindingExtension>& Extension : ActiveExtensions)
 				{
 					TOptional<FName> Name = Extension->GetCurrentValue(ThisBlueprint, Widget, InPropertyHandle->GetProperty());
 					if (Name.IsSet())
@@ -442,7 +450,7 @@ TSharedRef<SWidget> FBlueprintWidgetCustomization::MakePropertyBindingWidget(TWe
 		return false;
 	});
 
-	Args.CurrentBindingText = MakeAttributeLambda([InEditor, Objects, InPropertyHandle]()
+	Args.CurrentBindingText = MakeAttributeLambda([InEditor, Objects, InPropertyHandle, ActiveExtensions]()
 	{
 		UWidgetBlueprint* ThisBlueprint = InEditor.Pin()->GetWidgetBlueprintObj();
 
@@ -503,8 +511,7 @@ TSharedRef<SWidget> FBlueprintWidgetCustomization::MakePropertyBindingWidget(TWe
 
 			if (UWidget* Widget = Cast<UWidget>(Object))
 			{
-				IUMGEditorModule& EditorModule = FModuleManager::LoadModuleChecked<IUMGEditorModule>("UMGEditor");
-				for (const TSharedPtr<IPropertyBindingExtension>& Extension : EditorModule.GetPropertyBindingExtensibilityManager()->GetExtensions())
+				for (const TSharedPtr<IPropertyBindingExtension>& Extension : ActiveExtensions)
 				{
 					TOptional<FName> Name = Extension->GetCurrentValue(ThisBlueprint, Widget, InPropertyHandle->GetProperty());
 					if (Name.IsSet())
@@ -522,7 +529,7 @@ TSharedRef<SWidget> FBlueprintWidgetCustomization::MakePropertyBindingWidget(TWe
 		return LOCTEXT("Bind", "Bind");
 	});
 
-	Args.CurrentBindingImage = MakeAttributeLambda([InEditor, Objects, InPropertyHandle]() -> const FSlateBrush*
+	Args.CurrentBindingImage = MakeAttributeLambda([InEditor, Objects, InPropertyHandle, ActiveExtensions]() -> const FSlateBrush*
 	{
 		static FName PropertyIcon(TEXT("Kismet.Tabs.Variables"));
 		static FName FunctionIcon(TEXT("GraphEditor.Function_16x"));
@@ -557,6 +564,18 @@ TSharedRef<SWidget> FBlueprintWidgetCustomization::MakePropertyBindingWidget(TWe
 					}
 				}
 			}
+
+			if (UWidget* Widget = Cast<UWidget>(Object))
+			{
+				for (const TSharedPtr<IPropertyBindingExtension>& Extension : ActiveExtensions)
+				{
+					const FSlateBrush* Brush = Extension->GetCurrentIcon(ThisBlueprint, Widget, InPropertyHandle->GetProperty());
+					if (Brush != nullptr)
+					{
+						return Brush;
+					}
+				}
+			}
 		}
 
 		return nullptr;
@@ -586,8 +605,9 @@ bool FBlueprintWidgetCustomization::HasPropertyBindings(TWeakPtr<FWidgetBlueprin
 		ParentPropertyName = CurrentPropertyHandle->GetProperty()->GetFName();
 	}
 
-	//TODO UMG O(N) Isn't good for this, needs to be map, but map isn't serialized, need cached runtime map for fast lookups.
+	IUMGEditorModule& EditorModule = FModuleManager::LoadModuleChecked<IUMGEditorModule>("UMGEditor");
 
+	//TODO UMG O(N) Isn't good for this, needs to be map, but map isn't serialized, need cached runtime map for fast lookups.
 	for (const UObject* Object : Objects)
 	{
 		// Ignore null outer objects
@@ -603,8 +623,28 @@ bool FBlueprintWidgetCustomization::HasPropertyBindings(TWeakPtr<FWidgetBlueprin
 				return true;
 			}
 		}
+	}
 
-		break;
+	// check property binding extensions
+	for (const UObject* Object : Objects)
+	{
+		const UWidget* Widget = Cast<UWidget>(Object);
+		if (Widget == nullptr)
+		{
+			continue;
+		}
+
+		for (const TSharedPtr<IPropertyBindingExtension>& BindingExtension : EditorModule.GetPropertyBindingExtensibilityManager()->GetExtensions())
+		{
+			if (BindingExtension->CanExtend(ThisBlueprint, Widget, InPropertyHandle->GetProperty()))
+			{
+				TOptional<FName> CurrentValue = BindingExtension->GetCurrentValue(ThisBlueprint, Widget, InPropertyHandle->GetProperty());
+				if (CurrentValue.IsSet())
+				{
+					return true;
+				}
+			}
+		}
 	}
 
 	return false;
@@ -656,7 +696,7 @@ void FBlueprintWidgetCustomization::CreateEventCustomization( IDetailLayoutBuild
 		.MinDesiredWidth(200)
 		.MaxDesiredWidth(250)
 		[
-			MakePropertyBindingWidget(Editor.Pin(), Property, DelegatePropertyHandle, false)
+			MakePropertyBindingWidget(Editor.Pin(), Property->SignatureFunction, DelegatePropertyHandle, false)
 		];
 }
 
@@ -914,7 +954,7 @@ void FBlueprintWidgetCustomization::CustomizeAccessibilityProperty(IDetailLayout
 			return !HasPropertyBindings(ThisEditor, AccessibleTextPropertyHandle);
 		}));
 
-	TSharedRef<SWidget> BindingWidget = MakePropertyBindingWidget(Editor, AccessibleTextDelegateProperty, AccessibleTextPropertyHandle, false);
+	TSharedRef<SWidget> BindingWidget = MakePropertyBindingWidget(Editor, AccessibleTextDelegateProperty->SignatureFunction, AccessibleTextPropertyHandle, false);
 	TSharedRef<SHorizontalBox> CustomTextLayout = SNew(SHorizontalBox)
 	.Visibility(TAttribute<EVisibility>::Create([AccessibleBehaviorPropertyHandle]() -> EVisibility
 	{
diff --git a/Engine/Source/Editor/UMGEditor/Private/Customizations/UMGDetailCustomizations.h b/Engine/Source/Editor/UMGEditor/Private/Customizations/UMGDetailCustomizations.h
--- a/Engine/Source/Editor/UMGEditor/Private/Customizations/UMGDetailCustomizations.h
+++ b/Engine/Source/Editor/UMGEditor/Private/Customizations/UMGDetailCustomizations.h
@@ -9,6 +9,7 @@
 
 class IDetailLayoutBuilder;
 class IPropertyHandle;
+class UFunction;
 
 /**
  * Provides the customization for all UWidgets.  Bindings, style disabling...etc.
@@ -32,7 +33,7 @@ public:
 	virtual void CustomizeDetails( IDetailLayoutBuilder& DetailLayout ) override;
 
 	/** Make a property binding widget */
-	static TSharedRef<SWidget> MakePropertyBindingWidget(TWeakPtr<FWidgetBlueprintEditor> InEditor, FDelegateProperty* InProperty, TSharedRef<IPropertyHandle> InDelegatePropertyHandle, bool bInGeneratePureBindings);
+	static TSharedRef<SWidget> MakePropertyBindingWidget(TWeakPtr<FWidgetBlueprintEditor> InEditor, UFunction* SignatureFunction, TSharedRef<IPropertyHandle> InDelegatePropertyHandle, bool bInGeneratePureBindings);
 
 	/** Whether the property (or its parent property, if this is an array/struct element) currently has bindings */
 	static bool HasPropertyBindings(TWeakPtr<FWidgetBlueprintEditor> InEditor, const TSharedRef<IPropertyHandle>& InPropertyHandle);
diff --git a/Engine/Source/Editor/UMGEditor/Private/Details/DetailWidgetExtensionHandler.cpp b/Engine/Source/Editor/UMGEditor/Private/Details/DetailWidgetExtensionHandler.cpp
--- a/Engine/Source/Editor/UMGEditor/Private/Details/DetailWidgetExtensionHandler.cpp
+++ b/Engine/Source/Editor/UMGEditor/Private/Details/DetailWidgetExtensionHandler.cpp
@@ -1,59 +1,118 @@
 // Copyright Epic Games, Inc. All Rights Reserved.
 
 #include "Details/DetailWidgetExtensionHandler.h"
-#include "UMGEditorProjectSettings.h"
-#include "WidgetBlueprintEditor.h"
-#include "Engine/Blueprint.h"
 #include "Binding/WidgetBinding.h"
-#include "WidgetBlueprint.h"
 #include "Customizations/UMGDetailCustomizations.h"
+#include "Engine/Blueprint.h"
+#include "UMGEditorModule.h"
+#include "UMGEditorProjectSettings.h"
+#include "WidgetBlueprint.h"
+#include "WidgetBlueprintEditor.h"
 
 FDetailWidgetExtensionHandler::FDetailWidgetExtensionHandler(TSharedPtr<FWidgetBlueprintEditor> InBlueprintEditor)
 	: BlueprintEditor( InBlueprintEditor )
 {}
 
-bool FDetailWidgetExtensionHandler::IsPropertyExtendable(const UClass* InObjectClass, const IPropertyHandle& InPropertyHandle) const
+namespace Private
 {
-	// TODO UMG make this work for multiple widgets.
-	if ( InPropertyHandle.GetNumOuterObjects() == 1 )
+	FDelegateProperty* FindDelegateProperty(const IPropertyHandle& PropertyHandle)
 	{
-		TArray<UObject*> Objects;
-		InPropertyHandle.GetOuterObjects(Objects);
-
 		// Make the extension handler run for child struct/array properties if the parent is bindable.
 		// This is so we can later disable any child rows of a bound property in ExtendWidgetRow().
-		// Check to see if there's a delegate for the parent property.
 		TSharedPtr<const IPropertyHandle> ParentPropertyHandle;
-		for (TSharedPtr<const IPropertyHandle> CurrentPropertyHandle = InPropertyHandle.AsShared();
+		for (TSharedPtr<const IPropertyHandle> CurrentPropertyHandle = PropertyHandle.AsShared();
 			CurrentPropertyHandle && CurrentPropertyHandle->GetProperty();
 			CurrentPropertyHandle = CurrentPropertyHandle->GetParentHandle())
 		{
 			ParentPropertyHandle = CurrentPropertyHandle;
 		}
 
-		// We don't allow bindings on the CDO.
-		if (Objects[0] != nullptr && Objects[0]->HasAnyFlags(RF_ClassDefaultObject) )
+		// Check to see if there's a delegate for the parent property.
+		FProperty* ParentProperty = ParentPropertyHandle->GetProperty();
+		FString DelegateName = ParentProperty->GetName() + "Delegate";
+
+		if (UClass* ContainerClass = ParentProperty->GetOwner<UClass>())
+		{
+			FDelegateProperty* DelegateProperty = FindFProperty<FDelegateProperty>(ContainerClass, FName(*DelegateName));
+			return DelegateProperty;
+		} 
+
+		return nullptr;
+	}
+
+	bool ShouldShowOldBindingWidget(const UWidgetBlueprint* WidgetBlueprint, const IPropertyHandle& PropertyHandle)
+	{
+		FProperty* Property = PropertyHandle.GetProperty();
+		FString DelegateName = Property->GetName() + "Delegate";
+		FDelegateProperty* DelegateProperty = FindDelegateProperty(PropertyHandle);
+		if (DelegateProperty == nullptr)
 		{
 			return false;
 		}
 
-		TSharedPtr<FWidgetBlueprintEditor> BPEd = BlueprintEditor.Pin();
-		if (BPEd == nullptr || Objects[0] == BPEd->GetPreview())
+		const bool bIsEditable = Property->HasAnyPropertyFlags(CPF_Edit | CPF_EditConst);
+		const bool bDoSignaturesMatch = DelegateProperty->SignatureFunction->GetReturnProperty()->SameType(Property);
+
+		if (!ensure(bIsEditable && bDoSignaturesMatch))
 		{
 			return false;
 		}
 
-		FProperty* ParentProperty = ParentPropertyHandle->GetProperty();
-		FString DelegateName = ParentProperty->GetName() + "Delegate";
-
-		if ( UClass* ContainerClass = ParentProperty->GetOwner<UClass>() )
+		if (!WidgetBlueprint->ArePropertyBindingsAllowed())
 		{
-			FDelegateProperty* DelegateProperty = FindFProperty<FDelegateProperty>(ContainerClass, FName(*DelegateName));
-			if ( DelegateProperty )
+			// Even if they don't want them on, we need to show them so they can remove them if they had any.
+			if (WidgetBlueprint->Bindings.Num() == 0)
 			{
-				return true;
+				return false;
 			}
 		}
+
+		return true;
+	}
+}
+
+bool FDetailWidgetExtensionHandler::IsPropertyExtendable(const UClass* ObjectClass, const IPropertyHandle& PropertyHandle) const
+{
+	if (PropertyHandle.GetNumOuterObjects() != 1)
+	{
+		return false;
+	}
+
+	TArray<UObject*> Objects;
+	PropertyHandle.GetOuterObjects(Objects);
+
+	// We don't allow bindings on the CDO.
+	if (Objects[0] != nullptr && Objects[0]->HasAnyFlags(RF_ClassDefaultObject))
+	{
+		return false;
+	}
+
+	TSharedPtr<FWidgetBlueprintEditor> BPEd = BlueprintEditor.Pin();
+	if (BPEd == nullptr || Objects[0] == BPEd->GetPreview())
+	{
+		return false;
+	}
+
+	if (FDelegateProperty* DelegateProperty = Private::FindDelegateProperty(PropertyHandle))
+	{
+		return true;
+	}
+
+	const UWidget* Widget = Cast<UWidget>(Objects[0]);
+	if (Widget == nullptr)
+	{
+		return false;
+	}
+
+	const UWidgetBlueprint* WidgetBlueprint = BPEd->GetWidgetBlueprintObj();
+
+	IUMGEditorModule& EditorModule = FModuleManager::LoadModuleChecked<IUMGEditorModule>("UMGEditor");
+	for (const TSharedPtr<IPropertyBindingExtension>& Extension : EditorModule.GetPropertyBindingExtensibilityManager()->GetExtensions())
+	{
+		if (Extension->CanExtend(WidgetBlueprint, Widget, PropertyHandle.GetProperty()))
+		{
+			return true;
+		}
 	}
 
 	return false;
@@ -79,41 +138,42 @@ void FDetailWidgetExtensionHandler::ExtendWidgetRow(
 			return !FBlueprintWidgetCustomization::HasPropertyBindings(BlueprintEditorPtr, PropertyHandleRef);
 		}));
 
-	FProperty* Property = InPropertyHandle->GetProperty();
-	FString DelegateName = Property->GetName() + "Delegate";
+	const UWidgetBlueprint* WidgetBlueprint = BlueprintEditor.Pin()->GetWidgetBlueprintObj();
 
-	// Don't show the property binding widget for child properties of a bindable property,
-	// i.e. if the property's owner isn't the class itself.
-	// This check is necessary because IsPropertyExtendable() accepts all bindable properties
-	// and their descendants, primarily so we can disable the descendants' value widgets above.
-	const UClass* OwnerClass = Property->GetOwner<UClass>();
-	if (!OwnerClass) {
-		return;
-	}
+	bool bShouldShowWidget = false;
+	bShouldShowWidget |= Private::ShouldShowOldBindingWidget(WidgetBlueprint, PropertyHandleRef.Get());
+	
+	TArray<UObject*> Objects;
+	InPropertyHandle->GetOuterObjects(Objects);
 
-	FDelegateProperty* DelegateProperty = FindFieldChecked<FDelegateProperty>(OwnerClass, FName(*DelegateName));
+	for (const UObject* Object : Objects)
+	{
+		const UWidget* Widget = Cast<UWidget>(Object);
+		if (Widget == nullptr)
+		{
+			continue;
+		}
 
-	const bool bIsEditable = Property->HasAnyPropertyFlags(CPF_Edit | CPF_EditConst);
-	const bool bDoSignaturesMatch = DelegateProperty->SignatureFunction->GetReturnProperty()->SameType(Property);
+		IUMGEditorModule& EditorModule = FModuleManager::LoadModuleChecked<IUMGEditorModule>("UMGEditor");
+		for (const TSharedPtr<IPropertyBindingExtension>& Extension : EditorModule.GetPropertyBindingExtensibilityManager()->GetExtensions())
+		{
+			bShouldShowWidget |= Extension->CanExtend(WidgetBlueprint, Widget, InPropertyHandle->GetProperty());
+		}
+	}
 
-	if ( !ensure(bIsEditable && bDoSignaturesMatch) )
+	if (!bShouldShowWidget)
 	{
 		return;
 	}
 
-	UWidgetBlueprint* WidgetBlueprint = BlueprintEditor.Pin()->GetWidgetBlueprintObj();
-
-	if (!WidgetBlueprint->ArePropertyBindingsAllowed())
+	UFunction* SignatureFunction = nullptr;
+	if (FDelegateProperty* DelegateProperty = Private::FindDelegateProperty(PropertyHandleRef.Get()))
 	{
-		// Even if they don't want them on, we need to show them so they can remove them if they had any.
-		if (BlueprintEditor.Pin()->GetWidgetBlueprintObj()->Bindings.Num() == 0)
-		{
-			return;
-		}
+		SignatureFunction = DelegateProperty->SignatureFunction;
 	}
 
 	InWidgetRow.ExtensionContent()
 	[
-		FBlueprintWidgetCustomization::MakePropertyBindingWidget(BlueprintEditor, DelegateProperty, InPropertyHandle.ToSharedRef(), true)
+		FBlueprintWidgetCustomization::MakePropertyBindingWidget(BlueprintEditor, SignatureFunction, InPropertyHandle.ToSharedRef(), true)
 	];
 }
diff --git a/Engine/Source/Editor/UMGEditor/Private/Details/DetailWidgetExtensionHandler.h b/Engine/Source/Editor/UMGEditor/Private/Details/DetailWidgetExtensionHandler.h
--- a/Engine/Source/Editor/UMGEditor/Private/Details/DetailWidgetExtensionHandler.h
+++ b/Engine/Source/Editor/UMGEditor/Private/Details/DetailWidgetExtensionHandler.h
@@ -3,18 +3,16 @@
 #pragma once
 
 #include "CoreMinimal.h"
-#include "Widgets/SWidget.h"
-#include "WidgetBlueprintEditor.h"
 #include "IDetailPropertyExtensionHandler.h"
 
-class IPropertyHandle;
+class FWidgetBlueprintEditor;
 
 class FDetailWidgetExtensionHandler : public IDetailPropertyExtensionHandler
 {
 public:
-	FDetailWidgetExtensionHandler(TSharedPtr<class FWidgetBlueprintEditor> InBlueprintEditor);
+	FDetailWidgetExtensionHandler(TSharedPtr<FWidgetBlueprintEditor> InBlueprintEditor);
 
-	virtual bool IsPropertyExtendable(const UClass* InObjectClass, const class IPropertyHandle& PropertyHandle) const override;
+	virtual bool IsPropertyExtendable(const UClass* InObjectClass, const IPropertyHandle& PropertyHandle) const override;
 
 	virtual void ExtendWidgetRow(
 		FDetailWidgetRow& InWidgetRow,
@@ -23,5 +21,5 @@ public:
 		TSharedPtr<IPropertyHandle> PropertyHandle) override;
 
 private:
-	TWeakPtr<class FWidgetBlueprintEditor> BlueprintEditor;
+	TWeakPtr<FWidgetBlueprintEditor> BlueprintEditor;
 };
diff --git a/Engine/Source/Editor/UMGEditor/Public/IHasPropertyBindingExtensibility.h b/Engine/Source/Editor/UMGEditor/Public/IHasPropertyBindingExtensibility.h
--- a/Engine/Source/Editor/UMGEditor/Public/IHasPropertyBindingExtensibility.h
+++ b/Engine/Source/Editor/UMGEditor/Public/IHasPropertyBindingExtensibility.h
@@ -9,6 +9,7 @@
 
 class UWidget;
 class UWidgetBlueprint;
+struct FSlateBrush;
 
 class UMGEDITOR_API IPropertyBindingExtension
 {
@@ -16,8 +17,9 @@ public:
 	/** Does this extension want to extend this property in the widget? */
 	virtual bool CanExtend(const UWidgetBlueprint* WidgetBlueprint, const UWidget* Widget, const FProperty* Property) const = 0;
 	virtual TSharedPtr<FExtender> CreateMenuExtender(const UWidgetBlueprint* WidgetBlueprint, const UWidget* Widget, const FProperty* Property) = 0;
-	virtual TOptional<FName> GetCurrentValue(const UWidgetBlueprint* WidgetBlueprint, const UWidget* Widget, const FProperty* Property) const = 0;
 	virtual void ClearCurrentValue(const UWidgetBlueprint* WidgetBlueprint, const UWidget* Widget, const FProperty* Property) = 0;
+	virtual TOptional<FName> GetCurrentValue(const UWidgetBlueprint* WidgetBlueprint, const UWidget* Widget, const FProperty* Property) const = 0;
+	virtual const FSlateBrush* GetCurrentIcon(const UWidgetBlueprint* WidgetBlueprint, const UWidget* Widget, const FProperty* Property) const = 0;
 };
 
 /**
