PR #9634: Add a list of favorite widget parent class to project settings

diff --git a/Engine/Source/Editor/UMGEditor/Private/Settings/UMGEditorProjectSettings.cpp b/Engine/Source/Editor/UMGEditor/Private/Settings/UMGEditorProjectSettings.cpp
--- a/Engine/Source/Editor/UMGEditor/Private/Settings/UMGEditorProjectSettings.cpp
+++ b/Engine/Source/Editor/UMGEditor/Private/Settings/UMGEditorProjectSettings.cpp
@@ -21,7 +21,7 @@ UUMGEditorProjectSettings::UUMGEditorProjectSettings()
 
 	bUseWidgetTemplateSelector = false;
 	DefaultRootWidget = nullptr;
-	DefaultWidgetParentClass = UUserWidget::StaticClass();
+	FavoriteWidgetParentClasses.Add(UUserWidget::StaticClass());
 }
 
 #if WITH_EDITOR
diff --git a/Engine/Source/Editor/UMGEditor/Private/WidgetBlueprintFactory.cpp b/Engine/Source/Editor/UMGEditor/Private/WidgetBlueprintFactory.cpp
--- a/Engine/Source/Editor/UMGEditor/Private/WidgetBlueprintFactory.cpp
+++ b/Engine/Source/Editor/UMGEditor/Private/WidgetBlueprintFactory.cpp
@@ -71,16 +71,38 @@ bool UWidgetBlueprintFactory::ConfigureProperties()
 
 		TSharedPtr<FWidgetClassFilter> Filter = MakeShareable(new FWidgetClassFilter);
 		Options.ClassFilters.Add(Filter.ToSharedRef());
-		Options.ExtraPickerCommonClasses.Add(UUserWidget::StaticClass());
+
+		TArray<TSoftClassPtr<UUserWidget>> FavoriteWidgetParentClasses = GetDefault <UUMGEditorProjectSettings>()->FavoriteWidgetParentClasses;
+		for (int32 Index = 0; Index < FavoriteWidgetParentClasses.Num(); ++Index)
+		{
+			UClass* FavoriteWidgetParentClass = FavoriteWidgetParentClasses[Index].LoadSynchronous();
+			if (FavoriteWidgetParentClass && FavoriteWidgetParentClass->IsChildOf(UUserWidget::StaticClass()))
+			{
+				if (!Options.ExtraPickerCommonClasses.Contains(FavoriteWidgetParentClass))
+				{
+					Options.ExtraPickerCommonClasses.Add(FavoriteWidgetParentClass);
+				}
+			}
+		}
+
+		if (Options.ExtraPickerCommonClasses.Num() == 0)
+		{
+			Options.ExtraPickerCommonClasses.Add(UUserWidget::StaticClass());
+		}
 
 		Filter->DisallowedClassFlags = CLASS_Deprecated | CLASS_NewerVersionExists;
 		Filter->AllowedChildrenOfClasses.Add(UUserWidget::StaticClass());
 
-		const FText TitleText = LOCTEXT("CreateWidgetBlueprint", "Pick Root Widget for New Widget Blueprint");
+		const FText TitleText = LOCTEXT("CreateWidgetBlueprint", "Pick Parent Class for New Widget Blueprint");
 
 		UClass* ChosenParentClass = nullptr;
-		SClassPickerDialog::PickClass(TitleText, Options, ChosenParentClass, UUserWidget::StaticClass());
+		bool isSuccessful = SClassPickerDialog::PickClass(TitleText, Options, ChosenParentClass, UUserWidget::StaticClass());
 		ParentClass = ChosenParentClass ? ChosenParentClass : UUserWidget::StaticClass();
+
+		if (!isSuccessful)
+		{
+			return false;
+		}
 	}
 	if (GetDefault<UUMGEditorProjectSettings>()->bUseWidgetTemplateSelector)
 	{
@@ -123,7 +145,7 @@ UObject* UWidgetBlueprintFactory::FactoryCreateNew(UClass* Class, UObject* InPar
 	UClass* CurrentParentClass = ParentClass;
 	if (CurrentParentClass == nullptr)
 	{
-		CurrentParentClass = GetDefault <UUMGEditorProjectSettings>()->DefaultWidgetParentClass.LoadSynchronous();
+		CurrentParentClass = UUserWidget::StaticClass();
 	}
 
 	// If they selected an interface, force the parent class to be UInterface
diff --git a/Engine/Source/Editor/UMGEditor/Public/UMGEditorProjectSettings.h b/Engine/Source/Editor/UMGEditor/Public/UMGEditorProjectSettings.h
--- a/Engine/Source/Editor/UMGEditor/Public/UMGEditorProjectSettings.h
+++ b/Engine/Source/Editor/UMGEditor/Public/UMGEditorProjectSettings.h
@@ -185,11 +185,11 @@ public:
 	bool bUseEditorConfigPaletteFiltering;
 
 	/**
-	 * The default parent class for all newly constructed widget blueprints.
-	 * The WidgetParentClass must have an empty widget hierarchy.
+	 * The list of parent classes to choose from for newly constructed widget blueprints.
+	 * The classes must have empty widget hierarchies.
 	 */
-	UPROPERTY(config, meta = (AllowAbstract = ""))
-	TSoftClassPtr<UUserWidget> DefaultWidgetParentClass;
+	UPROPERTY(EditAnywhere, config, Category = Designer, meta = (AllowAbstract = ""))
+	TArray<TSoftClassPtr<UUserWidget>> FavoriteWidgetParentClasses;
 
 	UPROPERTY(EditAnywhere, config, Category=Designer)
 	TArray<FDebugResolution> DebugResolutions;
