Fix regression introduced in CL16061554, where in case of rename of a widget, the display name of the widgets was sanitize like object name, preventing to use many characters (space, &,...).

diff --git a/Engine/Source/Editor/UMGEditor/Private/WidgetBlueprintEditorUtils.cpp b/Engine/Source/Editor/UMGEditor/Private/WidgetBlueprintEditorUtils.cpp
--- a/Engine/Source/Editor/UMGEditor/Private/WidgetBlueprintEditorUtils.cpp
+++ b/Engine/Source/Editor/UMGEditor/Private/WidgetBlueprintEditorUtils.cpp
@@ -256,14 +256,14 @@ bool FWidgetBlueprintEditorUtils::RenameWidget(TSharedRef<FWidgetBlueprintEditor
 		UWidget* WidgetPreview = BlueprintEditor->GetReferenceFromTemplate(Widget).GetPreview();
 		if (WidgetPreview)
 		{
-			WidgetPreview->SetDisplayLabel(NewNameStr);
+			WidgetPreview->SetDisplayLabel(NewDisplayName);
 			WidgetPreview->Rename(*NewNameStr);
 		}
 
 		if (!WidgetPreview || WidgetPreview != Widget)
 		{
 			// Find and update all variable references in the graph
-			Widget->SetDisplayLabel(NewNameStr);
+			Widget->SetDisplayLabel(NewDisplayName);
 			Widget->Rename(*NewNameStr);
 		}
 
