Fix LOCTEXT key conflict in widget blueprint factory.

diff --git a/Engine/Source/Editor/UMGEditor/Private/WidgetBlueprintFactory.cpp b/Engine/Source/Editor/UMGEditor/Private/WidgetBlueprintFactory.cpp
--- a/Engine/Source/Editor/UMGEditor/Private/WidgetBlueprintFactory.cpp
+++ b/Engine/Source/Editor/UMGEditor/Private/WidgetBlueprintFactory.cpp
@@ -125,7 +125,7 @@ bool UWidgetBlueprintFactory::ConfigureProperties()
 		Filter->DisallowedClassFlags = CLASS_Abstract | CLASS_Deprecated | CLASS_NewerVersionExists;
 		Filter->AllowedChildrenOfClasses.Add(UPanelWidget::StaticClass());
 
-		const FText TitleText = LOCTEXT("CreateWidgetBlueprint", "Pick Root Widget for New Widget Blueprint");
+		const FText TitleText = LOCTEXT("CreateRootWidgetBlueprint", "Pick Root Widget for New Widget Blueprint");
 		return SClassPickerDialog::PickClass(TitleText, Options, static_cast<UClass*&>(RootWidgetClass), UPanelWidget::StaticClass());
 
 	}
