Added first pass Debug mode to UMG editor. Currently disabled behind CVar: UMG.EnableDebugMode.

diff --git a/Engine/Source/Editor/UMGEditor/Private/BlueprintModes/WidgetBlueprintApplicationModes.cpp b/Engine/Source/Editor/UMGEditor/Private/BlueprintModes/WidgetBlueprintApplicationModes.cpp
--- a/Engine/Source/Editor/UMGEditor/Private/BlueprintModes/WidgetBlueprintApplicationModes.cpp
+++ b/Engine/Source/Editor/UMGEditor/Private/BlueprintModes/WidgetBlueprintApplicationModes.cpp
@@ -2,11 +2,13 @@
 
 #include "BlueprintModes/WidgetBlueprintApplicationModes.h"
 
+#include "HAL/IConsoleManager.h"
 #include "Internationalization/Internationalization.h"
 
 // Mode constants
 const FName FWidgetBlueprintApplicationModes::DesignerMode("DesignerName");
 const FName FWidgetBlueprintApplicationModes::GraphMode("GraphName");
+const FName FWidgetBlueprintApplicationModes::DebugMode("DebugName");
 
 FText FWidgetBlueprintApplicationModes::GetLocalizedMode(const FName InMode)
 {
@@ -18,5 +20,21 @@ FText FWidgetBlueprintApplicationModes::GetLocalizedMode(const FName InMode)
 	{
 		return NSLOCTEXT("WidgetBlueprintModes", "GraphMode", "Graph");
 	}
+	else if (InMode == FWidgetBlueprintApplicationModes::DebugMode)
+	{
+		return NSLOCTEXT("WidgetBlueprintModes", "DebugMode", "Debug");
+	}
 	return FText::GetEmpty();
 }
+
+static bool EnableDebugMode = false;
+static FAutoConsoleVariableRef CVarEnableDebugMode(
+	TEXT("UMG.EnableDebugMode"), 
+	EnableDebugMode, 
+	TEXT("Whether or not to enable the UMG Debug mode.")
+);
+
+bool FWidgetBlueprintApplicationModes::IsDebugModeEnabled()
+{
+	return EnableDebugMode;
+}
diff --git a/Engine/Source/Editor/UMGEditor/Private/BlueprintModes/WidgetDebugApplicationMode.cpp b/Engine/Source/Editor/UMGEditor/Private/BlueprintModes/WidgetDebugApplicationMode.cpp
new file mode 100644
--- /dev/null
+++ b/Engine/Source/Editor/UMGEditor/Private/BlueprintModes/WidgetDebugApplicationMode.cpp
@@ -0,0 +1,81 @@
+// Copyright Epic Games, Inc. All Rights Reserved.
+
+#include "BlueprintModes/WidgetDebugApplicationMode.h"
+
+#include "BlueprintModes/WidgetBlueprintApplicationModes.h"
+#include "SBlueprintEditorToolbar.h"
+#include "TabFactory/DebugLogTabSummoner.h"
+#include "TabFactory/DetailsTabSummoner.h"
+#include "TabFactory/PreviewTabSummoner.h"
+#include "ToolMenus.h"
+#include "UMGEditorModule.h"
+#include "WidgetBlueprintEditorToolbar.h"
+
+FWidgetDebugApplicationMode::FWidgetDebugApplicationMode(TSharedPtr<FWidgetBlueprintEditor> InWidgetEditor)
+	: FWidgetBlueprintApplicationMode(InWidgetEditor, FWidgetBlueprintApplicationModes::DebugMode)
+{
+	TabLayout = FTabManager::NewLayout( "WidgetBlueprintEditor_Debug_Layout_v2" )
+		->AddArea
+		(
+			FTabManager::NewPrimaryArea()
+			->SetOrientation(Orient_Horizontal)
+			->Split
+			(
+				FTabManager::NewSplitter()
+				->SetOrientation(Orient_Vertical)
+				->SetSizeCoefficient(0.7f)
+				->Split
+				(
+					FTabManager::NewStack()
+					->SetSizeCoefficient(0.8f)
+					->AddTab(FPreviewTabSummoner::TabID, ETabState::OpenedTab)
+				)
+				->Split
+				(
+					FTabManager::NewStack()
+					->SetSizeCoefficient(0.2f)
+					->AddTab(FDebugLogTabSummoner::TabID, ETabState::OpenedTab)
+				)
+			)
+			->Split
+			(
+				FTabManager::NewSplitter()
+				->SetOrientation(Orient_Vertical)
+				->SetSizeCoefficient(0.15f)
+				->Split
+				(
+					FTabManager::NewStack()
+					->AddTab(FDetailsTabSummoner::TabID, ETabState::OpenedTab)
+				)
+			)
+		);
+	
+	TabFactories.RegisterFactory(MakeShareable(new FPreviewTabSummoner(InWidgetEditor)));
+	TabFactories.RegisterFactory(MakeShareable(new FDebugLogTabSummoner(InWidgetEditor, "Model View Viewmodel")));
+	TabFactories.RegisterFactory(MakeShareable(new FDetailsTabSummoner(InWidgetEditor)));
+
+	IUMGEditorModule& UMGEditorModule = FModuleManager::LoadModuleChecked<IUMGEditorModule>("UMGEditor");
+	UMGEditorModule.OnRegisterTabsForEditor().Broadcast(*this, TabFactories);
+
+	ToolbarExtender = UMGEditorModule.GetToolBarExtensibilityManager()->GetAllExtenders();
+	InWidgetEditor->GetWidgetToolbarBuilder()->AddWidgetBlueprintEditorModesToolbar(ToolbarExtender);
+
+	if (UToolMenu* Toolbar = InWidgetEditor->RegisterModeToolbarIfUnregistered(GetModeName()))
+	{
+		InWidgetEditor->GetWidgetToolbarBuilder()->AddWidgetReflector(Toolbar);
+		InWidgetEditor->GetToolbarBuilder()->AddCompileToolbar(Toolbar);
+	}
+}
+
+void FWidgetDebugApplicationMode::RegisterTabFactories(TSharedPtr<FTabManager> InTabManager)
+{
+	TSharedPtr<FBlueprintEditor> BP = GetBlueprintEditor();
+	
+	BP->RegisterToolbarTab(InTabManager.ToSharedRef());
+	BP->PushTabFactories(TabFactories);
+}
+
+void FWidgetDebugApplicationMode::PostActivateMode()
+{
+
+}
diff --git a/Engine/Source/Editor/UMGEditor/Private/BlueprintModes/WidgetDebugApplicationMode.h b/Engine/Source/Editor/UMGEditor/Private/BlueprintModes/WidgetDebugApplicationMode.h
new file mode 100644
--- /dev/null
+++ b/Engine/Source/Editor/UMGEditor/Private/BlueprintModes/WidgetDebugApplicationMode.h
@@ -0,0 +1,19 @@
+// Copyright Epic Games, Inc. All Rights Reserved.
+
+#pragma once
+
+#include "CoreMinimal.h"
+#include "BlueprintModes/WidgetBlueprintApplicationMode.h"
+
+class FWidgetDebugApplicationMode : public FWidgetBlueprintApplicationMode
+{
+public:
+	FWidgetDebugApplicationMode(TSharedPtr<FWidgetBlueprintEditor> InWidgetEditor);
+
+	static bool IsDebugModeEnabled();
+
+	// FApplicationMode interface
+	virtual void RegisterTabFactories(TSharedPtr<FTabManager> InTabManager) override;
+	virtual void PostActivateMode() override;
+	// End of FApplicationMode interface
+};
diff --git a/Engine/Source/Editor/UMGEditor/Private/Debug/SDebugLog.cpp b/Engine/Source/Editor/UMGEditor/Private/Debug/SDebugLog.cpp
new file mode 100644
--- /dev/null
+++ b/Engine/Source/Editor/UMGEditor/Private/Debug/SDebugLog.cpp
@@ -0,0 +1,24 @@
+// Copyright Epic Games, Inc. All Rights Reserved.
+
+#include "SDebugLog.h"
+#include "Widgets/Views/SListView.h"
+#include "Widgets/Views/STableRow.h"
+
+namespace UE::UMG
+{
+
+void SDebugLog::Construct(const FArguments& Args, FName LogName)
+{
+	ChildSlot
+	[
+		SNew(SListView<TSharedPtr<FTokenizedMessage>>)
+		.OnGenerateRow(this, &SDebugLog::OnGenerateRow)
+	];
+}
+
+TSharedRef<ITableRow> SDebugLog::OnGenerateRow(TSharedPtr<FTokenizedMessage> Message, const TSharedRef<STableViewBase>& OwnerTable) const
+{
+	return SNew(STableRow<TSharedPtr<FTokenizedMessage>>, OwnerTable);
+}
+
+} // namespace UE::UMG
diff --git a/Engine/Source/Editor/UMGEditor/Private/Debug/SDebugLog.h b/Engine/Source/Editor/UMGEditor/Private/Debug/SDebugLog.h
new file mode 100644
--- /dev/null
+++ b/Engine/Source/Editor/UMGEditor/Private/Debug/SDebugLog.h
@@ -0,0 +1,25 @@
+// Copyright Epic Games, Inc. All Rights Reserved.
+
+#pragma once
+
+#include "Widgets/SCompoundWidget.h"
+
+class FTokenizedMessage;
+class ITableRow;
+class STableViewBase;
+
+namespace UE::UMG
+{
+
+class SDebugLog : public SCompoundWidget
+{
+	SLATE_BEGIN_ARGS(SDebugLog) {}
+	SLATE_END_ARGS()
+
+	void Construct(const FArguments& Args, FName LogName);
+
+private:
+	TSharedRef<ITableRow> OnGenerateRow(TSharedPtr<FTokenizedMessage> Message, const TSharedRef<STableViewBase>& OwnerTable) const;
+};
+
+} // namespace UE::UMG
diff --git a/Engine/Source/Editor/UMGEditor/Private/Debug/SDebugPreview.cpp b/Engine/Source/Editor/UMGEditor/Private/Debug/SDebugPreview.cpp
new file mode 100644
--- /dev/null
+++ b/Engine/Source/Editor/UMGEditor/Private/Debug/SDebugPreview.cpp
@@ -0,0 +1,46 @@
+// Copyright Epic Games, Inc. All Rights Reserved.
+
+#include "SDebugPreview.h"
+#include "Blueprint/UserWidget.h"
+#include "Editor.h"
+#include "Engine/Engine.h"
+#include "Engine/World.h"
+#include "WidgetBlueprint.h"
+#include "WidgetBlueprintEditor.h"
+
+namespace UE::UMG
+{
+
+void SDebugPreview::Construct(const FArguments& Args, TSharedPtr<FWidgetBlueprintEditor> InWidgetBlueprintEditor)
+{
+	WidgetBlueprintEditor = InWidgetBlueprintEditor;
+
+	TSharedPtr<SWidget> SlateWidget;
+
+	if (UWidgetBlueprint* WidgetBlueprint = InWidgetBlueprintEditor->GetWidgetBlueprintObj())
+	{
+		if (UWorld* World = GEditor->GetEditorWorldContext().World())
+		{
+			if (UWidgetBlueprintGeneratedClass* WidgetClass = Cast<UWidgetBlueprintGeneratedClass>(WidgetBlueprint->GeneratedClass))
+			{
+				CreatedWidget = CreateWidget(World, WidgetClass, WidgetClass->GetFName());
+				if (CreatedWidget != nullptr)
+				{
+					CreatedWidget->SetDesignerFlags(EWidgetDesignFlags::Previewing);
+					SlateWidget = CreatedWidget->TakeWidget();
+				}
+			}
+		}
+	}
+
+	ChildSlot
+	[
+		SlateWidget.IsValid() ? SlateWidget.ToSharedRef() : SNullWidget::NullWidget
+	];
+}
+
+SDebugPreview::~SDebugPreview()
+{
+}
+
+} // namespace UE::UMG
diff --git a/Engine/Source/Editor/UMGEditor/Private/Debug/SDebugPreview.h b/Engine/Source/Editor/UMGEditor/Private/Debug/SDebugPreview.h
new file mode 100644
--- /dev/null
+++ b/Engine/Source/Editor/UMGEditor/Private/Debug/SDebugPreview.h
@@ -0,0 +1,31 @@
+// Copyright Epic Games, Inc. All Rights Reserved.
+
+#pragma once
+
+#include "UObject/ObjectPtr.h"
+#include "Widgets/SCompoundWidget.h"
+
+class FWidgetBlueprintEditor;
+class UPlayer;
+class UUserWidget;
+
+namespace UE::UMG
+{
+
+class SDebugPreview : public SCompoundWidget
+{
+public:
+	SLATE_BEGIN_ARGS(SDebugPreview) {}
+	SLATE_END_ARGS()
+
+	~SDebugPreview();
+
+	void Construct(const FArguments& Args, TSharedPtr<FWidgetBlueprintEditor> InWidgetBlueprintEditor);
+
+private:
+
+	TSharedPtr<FWidgetBlueprintEditor> WidgetBlueprintEditor;
+	TObjectPtr<UUserWidget> CreatedWidget;
+};
+
+} // namespace UE::UMG
\ No newline at end of file
diff --git a/Engine/Source/Editor/UMGEditor/Private/TabFactory/DebugLogTabSummoner.cpp b/Engine/Source/Editor/UMGEditor/Private/TabFactory/DebugLogTabSummoner.cpp
new file mode 100644
--- /dev/null
+++ b/Engine/Source/Editor/UMGEditor/Private/TabFactory/DebugLogTabSummoner.cpp
@@ -0,0 +1,32 @@
+// Copyright Epic Games, Inc. All Rights Reserved.
+
+#include "TabFactory/DebugLogTabSummoner.h"
+
+#include "Debug/SDebugLog.h"
+#include "UMGStyle.h"
+#include "WidgetBlueprintEditor.h"
+
+#define LOCTEXT_NAMESPACE "UMG"
+
+const FName FDebugLogTabSummoner::TabID("DebugLog");
+
+FDebugLogTabSummoner::FDebugLogTabSummoner(TSharedPtr<FWidgetBlueprintEditor> InBlueprintEditor, FName InLogName)
+		: FWorkflowTabFactory(TabID, InBlueprintEditor)
+{
+	LogName = InLogName;
+
+	TabLabel = LOCTEXT("DebugLog_TabLabel", "Debug Log");
+	TabIcon = FSlateIcon(FAppStyle::GetAppStyleSetName(), "BlueprintDebugger.TabIcon");
+
+	bIsSingleton = true;
+
+	ViewMenuDescription = LOCTEXT("DebugLog_ViewMenu_Desc", "Debug Log");
+	ViewMenuTooltip = LOCTEXT("DebugLog_ViewMenu_ToolTip", "Show the Debug Log");
+}
+
+TSharedRef<SWidget> FDebugLogTabSummoner::CreateTabBody(const FWorkflowTabSpawnInfo& Info) const
+{
+	return SNew(UE::UMG::SDebugLog, LogName);
+}
+
+#undef LOCTEXT_NAMESPACE 
diff --git a/Engine/Source/Editor/UMGEditor/Private/TabFactory/DebugLogTabSummoner.h b/Engine/Source/Editor/UMGEditor/Private/TabFactory/DebugLogTabSummoner.h
new file mode 100644
--- /dev/null
+++ b/Engine/Source/Editor/UMGEditor/Private/TabFactory/DebugLogTabSummoner.h
@@ -0,0 +1,23 @@
+// Copyright Epic Games, Inc. All Rights Reserved.
+
+#pragma once
+
+#include "WorkflowOrientedApp/WorkflowTabFactory.h"
+
+class FWidgetBlueprintEditor;
+class SWidget;
+
+struct FDebugLogTabSummoner : public FWorkflowTabFactory
+{
+public:
+	static const FName TabID;
+	
+public:
+	FDebugLogTabSummoner(TSharedPtr<FWidgetBlueprintEditor> InBlueprintEditor, FName InLogName);
+	
+	virtual TSharedRef<SWidget> CreateTabBody(const FWorkflowTabSpawnInfo& Info) const override;
+
+private:
+
+	FName LogName;
+};
diff --git a/Engine/Source/Editor/UMGEditor/Private/TabFactory/PreviewTabSummoner.cpp b/Engine/Source/Editor/UMGEditor/Private/TabFactory/PreviewTabSummoner.cpp
new file mode 100644
--- /dev/null
+++ b/Engine/Source/Editor/UMGEditor/Private/TabFactory/PreviewTabSummoner.cpp
@@ -0,0 +1,33 @@
+// Copyright Epic Games, Inc. All Rights Reserved.
+
+#include "TabFactory/PreviewTabSummoner.h"
+
+#include "Debug/SDebugPreview.h"
+#include "UMGStyle.h"
+
+#define LOCTEXT_NAMESPACE "UMG"
+
+const FName FPreviewTabSummoner::TabID("WidgetPreview");
+
+FPreviewTabSummoner::FPreviewTabSummoner(TSharedPtr<FWidgetBlueprintEditor> InBlueprintEditor)
+		: FWorkflowTabFactory(TabID, InBlueprintEditor)
+		, BlueprintEditor(InBlueprintEditor)
+{
+	TabLabel = LOCTEXT("DebugPreview_TabLabel", "Preview");
+	TabIcon = FSlateIcon(FUMGStyle::GetStyleSetName(), "BlueprintDebugger.TabIcon");
+
+	bIsSingleton = true;
+
+	ViewMenuDescription = LOCTEXT("DebugPreview_ViewMenu_Desc", "Preview");
+	ViewMenuTooltip = LOCTEXT("DebugPreview_ViewMenu_ToolTip", "Show the Preview");
+}
+
+TSharedRef<SWidget> FPreviewTabSummoner::CreateTabBody(const FWorkflowTabSpawnInfo& Info) const
+{
+	TSharedPtr<FWidgetBlueprintEditor> BlueprintEditorPtr = StaticCastSharedPtr<FWidgetBlueprintEditor>(BlueprintEditor.Pin());
+
+	return SNew(UE::UMG::SDebugPreview, BlueprintEditorPtr)
+		.AddMetaData<FTagMetaData>(FTagMetaData(TEXT("Preview")));
+}
+
+#undef LOCTEXT_NAMESPACE 
diff --git a/Engine/Source/Editor/UMGEditor/Private/TabFactory/PreviewTabSummoner.h b/Engine/Source/Editor/UMGEditor/Private/TabFactory/PreviewTabSummoner.h
new file mode 100644
--- /dev/null
+++ b/Engine/Source/Editor/UMGEditor/Private/TabFactory/PreviewTabSummoner.h
@@ -0,0 +1,21 @@
+// Copyright Epic Games, Inc. All Rights Reserved.
+
+#pragma once
+
+#include "WidgetBlueprintEditor.h"
+#include "Widgets/SWidget.h"
+#include "WorkflowOrientedApp/WorkflowTabFactory.h"
+
+struct FPreviewTabSummoner : public FWorkflowTabFactory
+{
+public:
+	static const FName TabID;
+	
+public:
+	FPreviewTabSummoner(TSharedPtr<FWidgetBlueprintEditor> InBlueprintEditor);
+	
+	virtual TSharedRef<SWidget> CreateTabBody(const FWorkflowTabSpawnInfo& Info) const override;
+	
+protected:
+	TWeakPtr<FWidgetBlueprintEditor> BlueprintEditor;
+};
diff --git a/Engine/Source/Editor/UMGEditor/Private/WidgetBlueprintEditor.cpp b/Engine/Source/Editor/UMGEditor/Private/WidgetBlueprintEditor.cpp
--- a/Engine/Source/Editor/UMGEditor/Private/WidgetBlueprintEditor.cpp
+++ b/Engine/Source/Editor/UMGEditor/Private/WidgetBlueprintEditor.cpp
@@ -38,6 +38,7 @@
 #include "Blueprint/WidgetTree.h"
 #include "WidgetBlueprintEditorUtils.h"
 #include "WorkflowOrientedApp/ApplicationMode.h"
+#include "BlueprintModes/WidgetDebugApplicationMode.h"
 #include "BlueprintModes/WidgetDesignerApplicationMode.h"
 #include "BlueprintModes/WidgetGraphApplicationMode.h"
 #include "WidgetModeManager.h"
@@ -547,8 +548,13 @@ void FWidgetBlueprintEditor::RegisterApplicationModes(const TArray<UBlueprint*>&
 
 		// Create the modes and activate one (which will populate with a real layout)
 		TArray< TSharedRef<FApplicationMode> > TempModeList;
-		TempModeList.Add(MakeShareable(new FWidgetDesignerApplicationMode(ThisPtr)));
-		TempModeList.Add(MakeShareable(new FWidgetGraphApplicationMode(ThisPtr)));
+		TempModeList.Add(MakeShared<FWidgetDesignerApplicationMode>(ThisPtr));
+		TempModeList.Add(MakeShared<FWidgetGraphApplicationMode>(ThisPtr));
+
+		if (FWidgetBlueprintApplicationModes::IsDebugModeEnabled())
+		{
+			TempModeList.Add(MakeShared<FWidgetDebugApplicationMode>(ThisPtr));
+		}
 
 		for ( TSharedRef<FApplicationMode>& AppMode : TempModeList )
 		{
diff --git a/Engine/Source/Editor/UMGEditor/Private/WidgetBlueprintEditorToolbar.cpp b/Engine/Source/Editor/UMGEditor/Private/WidgetBlueprintEditorToolbar.cpp
--- a/Engine/Source/Editor/UMGEditor/Private/WidgetBlueprintEditorToolbar.cpp
+++ b/Engine/Source/Editor/UMGEditor/Private/WidgetBlueprintEditorToolbar.cpp
@@ -94,7 +94,7 @@ void FWidgetBlueprintEditorToolbar::FillWidgetBlueprintEditorModesToolbar(FToolB
 			.OnSetActiveMode(SetActiveMode)
 			.ToolTip(IDocumentation::Get()->CreateToolTip(
 				LOCTEXT("DesignerModeButtonTooltip", "Switch to Blueprint Designer Mode"),
-				NULL,
+				nullptr,
 				TEXT("Shared/Editors/BlueprintEditor"),
 				TEXT("DesignerMode")))
 			.IconImage(FAppStyle::GetBrush("UMGEditor.SwitchToDesigner"))
@@ -110,16 +110,35 @@ void FWidgetBlueprintEditorToolbar::FillWidgetBlueprintEditorModesToolbar(FToolB
 			.CanBeSelected(BlueprintEditorPtr.Get(), &FBlueprintEditor::IsEditingSingleBlueprint)
 			.ToolTip(IDocumentation::Get()->CreateToolTip(
 				LOCTEXT("GraphModeButtonTooltip", "Switch to Graph Editing Mode"),
-				NULL,
+				nullptr,
 				TEXT("Shared/Editors/BlueprintEditor"),
 				TEXT("GraphMode")))
-			.ToolTipText(LOCTEXT("GraphModeButtonTooltip", "Switch to Graph Editing Mode"))
 			.IconImage(FAppStyle::GetBrush("FullBlueprintEditor.SwitchToScriptingMode"))
 			.AddMetaData<FTagMetaData>(FTagMetaData(TEXT("GraphMode")))
 		);
 		
 		// Right side padding
 		BlueprintEditorPtr->AddToolbarWidget(SNew(SSpacer).Size(FVector2D(10.0f, 1.0f)));
+
+		if (FWidgetBlueprintApplicationModes::IsDebugModeEnabled())
+		{
+			BlueprintEditorPtr->AddToolbarWidget(
+				SNew(SModeWidget, FWidgetBlueprintApplicationModes::GetLocalizedMode(FWidgetBlueprintApplicationModes::DebugMode), FWidgetBlueprintApplicationModes::DebugMode)
+				.OnGetActiveMode(GetActiveMode)
+				.OnSetActiveMode(SetActiveMode)
+				.CanBeSelected(BlueprintEditorPtr.Get(), &FBlueprintEditor::IsEditingSingleBlueprint)
+				.ToolTip(IDocumentation::Get()->CreateToolTip(
+					LOCTEXT("DebugModeButtonTooltip", "Switch to Debugging Mode"),
+					nullptr,
+					TEXT("Shared/Editors/BlueprintEditor"),
+					TEXT("DebugMode")))
+				.IconImage(FAppStyle::GetBrush("BlueprintDebugger.TabIcon"))
+				.AddMetaData<FTagMetaData>(FTagMetaData(TEXT("DebugMode")))
+			);
+		
+			// Right side padding
+			BlueprintEditorPtr->AddToolbarWidget(SNew(SSpacer).Size(FVector2D(10.0f, 1.0f)));
+		}
 	}
 }
 
diff --git a/Engine/Source/Editor/UMGEditor/Public/BlueprintModes/WidgetBlueprintApplicationModes.h b/Engine/Source/Editor/UMGEditor/Public/BlueprintModes/WidgetBlueprintApplicationModes.h
--- a/Engine/Source/Editor/UMGEditor/Public/BlueprintModes/WidgetBlueprintApplicationModes.h
+++ b/Engine/Source/Editor/UMGEditor/Public/BlueprintModes/WidgetBlueprintApplicationModes.h
@@ -15,9 +15,12 @@ struct UMGEDITOR_API FWidgetBlueprintApplicationModes
 	// Mode constants
 	static const FName DesignerMode;
 	static const FName GraphMode;
+	static const FName DebugMode;
 
 	static FText GetLocalizedMode(const FName InMode);
 
+	static bool IsDebugModeEnabled();
+
 private:
 	FWidgetBlueprintApplicationModes() = delete;
 };
diff --git a/Engine/Source/Runtime/UMG/Public/Components/Widget.h b/Engine/Source/Runtime/UMG/Public/Components/Widget.h
--- a/Engine/Source/Runtime/UMG/Public/Components/Widget.h
+++ b/Engine/Source/Runtime/UMG/Public/Components/Widget.h
@@ -168,7 +168,8 @@ enum class EWidgetDesignFlags : uint8
 	None				= 0,
 	Designing			= 1 << 0,
 	ShowOutline			= 1 << 1,
-	ExecutePreConstruct	= 1 << 2
+	ExecutePreConstruct	= 1 << 2,
+	Previewing			= 1 << 3
 };
 
 ENUM_CLASS_FLAGS(EWidgetDesignFlags);
@@ -918,6 +919,11 @@ public:
 	{
 		return EnumHasAnyFlags(GetDesignerFlags(), FlagsToCheck);
 	}
+	
+	FORCEINLINE bool IsPreviewTime() const 
+	{
+		return HasAnyDesignerFlags(EWidgetDesignFlags::Previewing);
+	}
 
 	/** Returns the friendly name of the widget to display in the editor */
 	const FString& GetDisplayLabel() const
@@ -945,6 +951,7 @@ public:
 	virtual bool Modify(bool bAlwaysMarkDirty = true) override;
 #else
 	FORCEINLINE bool IsDesignTime() const { return false; }
+	FORCEINLINE bool IsPreviewTime() const { return false; }
 #endif
 	
 	/**
