Expose WidgetBlueprintEditorUtils

diff --git a/Engine/Source/Editor/UMGEditor/Public/WidgetBlueprintEditorUtils.h b/Engine/Source/Editor/UMGEditor/Public/WidgetBlueprintEditorUtils.h
--- a/Engine/Source/Editor/UMGEditor/Public/WidgetBlueprintEditorUtils.h
+++ b/Engine/Source/Editor/UMGEditor/Public/WidgetBlueprintEditorUtils.h
@@ -18,7 +18,7 @@ class SWindow;
 //////////////////////////////////////////////////////////////////////////
 // FWidgetBlueprintEditorUtils
 
-class FWidgetBlueprintEditorUtils
+class UMGEDITOR_API FWidgetBlueprintEditorUtils
 {
 
 public:
