PR #9996: Expose the list of common root widgets displayed on the root widget class picker dialog (Contributed by [REDACTED_PERSON])

diff --git a/Engine/Source/Editor/UMGEditor/Private/Settings/UMGEditorProjectSettings.cpp b/Engine/Source/Editor/UMGEditor/Private/Settings/UMGEditorProjectSettings.cpp
--- a/Engine/Source/Editor/UMGEditor/Private/Settings/UMGEditorProjectSettings.cpp
+++ b/Engine/Source/Editor/UMGEditor/Private/Settings/UMGEditorProjectSettings.cpp
@@ -7,7 +7,10 @@
 #include "UObject/UObjectIterator.h"
 
 #include "Blueprint/UserWidget.h"
-#include "Components/CanvasPanel.h"
+#include "Components/CanvasPanel.h"	
+#include "Components/GridPanel.h"
+#include "Components/HorizontalBox.h"
+#include "Components/VerticalBox.h"
 
 UUMGEditorProjectSettings::UUMGEditorProjectSettings()
 {
@@ -20,6 +23,12 @@ UUMGEditorProjectSettings::UUMGEditorProjectSettings()
 	bUseEditorConfigPaletteFiltering = false;
 
 	bUseWidgetTemplateSelector = false;
+	CommonRootWidgetClasses = {
+		UHorizontalBox::StaticClass(),
+		UVerticalBox::StaticClass(),
+		UGridPanel::StaticClass(),
+		UCanvasPanel::StaticClass()
+	};
 	DefaultRootWidget = nullptr;
 	FavoriteWidgetParentClasses.Add(UUserWidget::StaticClass());
 }
diff --git a/Engine/Source/Editor/UMGEditor/Private/WidgetBlueprintFactory.cpp b/Engine/Source/Editor/UMGEditor/Private/WidgetBlueprintFactory.cpp
--- a/Engine/Source/Editor/UMGEditor/Private/WidgetBlueprintFactory.cpp
+++ b/Engine/Source/Editor/UMGEditor/Private/WidgetBlueprintFactory.cpp
@@ -15,9 +15,6 @@
 #include "Kismet2/SClassPickerDialog.h"
 #include "ClassViewerFilter.h"
 #include "Components/CanvasPanel.h"
-#include "Components/HorizontalBox.h"
-#include "Components/VerticalBox.h"
-#include "Components/GridPanel.h"
 
 #define LOCTEXT_NAMESPACE "UWidgetBlueprintFactory"
 
@@ -114,10 +111,23 @@ bool UWidgetBlueprintFactory::ConfigureProperties()
 		Options.Mode = EClassViewerMode::ClassPicker;
 		Options.bShowNoneOption = true;
 
-		Options.ExtraPickerCommonClasses.Add(UHorizontalBox::StaticClass());
-		Options.ExtraPickerCommonClasses.Add(UVerticalBox::StaticClass());
-		Options.ExtraPickerCommonClasses.Add(UGridPanel::StaticClass());
-		Options.ExtraPickerCommonClasses.Add(UCanvasPanel::StaticClass());
+		TArray<TSoftClassPtr<UPanelWidget>> CommonRootWidgetClasses = GetDefault <UUMGEditorProjectSettings>()->CommonRootWidgetClasses;
+		for (int32 Index = 0; Index < CommonRootWidgetClasses.Num(); ++Index)
+		{
+			UClass* PanelWidgetClass = CommonRootWidgetClasses[Index].LoadSynchronous();
+			if (PanelWidgetClass && PanelWidgetClass->IsChildOf(UPanelWidget::StaticClass()))
+			{
+				if (!Options.ExtraPickerCommonClasses.Contains(PanelWidgetClass))
+				{
+					Options.ExtraPickerCommonClasses.Add(PanelWidgetClass);
+				}
+			}
+		}
+
+		if (Options.ExtraPickerCommonClasses.Num() == 0)
+		{
+			Options.ExtraPickerCommonClasses.Add(UCanvasPanel::StaticClass());
+		}
 
 		TSharedPtr<FWidgetClassFilter> Filter = MakeShareable(new FWidgetClassFilter);
 		Options.ClassFilters.Add(Filter.ToSharedRef());
diff --git a/Engine/Source/Editor/UMGEditor/Public/UMGEditorProjectSettings.h b/Engine/Source/Editor/UMGEditor/Public/UMGEditorProjectSettings.h
--- a/Engine/Source/Editor/UMGEditor/Public/UMGEditorProjectSettings.h
+++ b/Engine/Source/Editor/UMGEditor/Public/UMGEditorProjectSettings.h
@@ -168,6 +168,10 @@ public:
 	UPROPERTY(EditAnywhere, config, Category = Designer)
 	bool bUseWidgetTemplateSelector;
 
+	/** This list populates the common class section of the root widget selection dialog */
+	UPROPERTY(EditAnywhere, config, Category = Designer, meta = (EditCondition = "bUseWidgetTemplateSelector"))
+	TArray<TSoftClassPtr<UPanelWidget>> CommonRootWidgetClasses;
+
 	/** The panel widget to place at the root of all newly constructed widget blueprints. Can be empty. */
 	UPROPERTY(EditAnywhere, config, Category = Designer)
 	TSubclassOf<UPanelWidget> DefaultRootWidget;
