Transact the CDO to record undo transactions of the root user widget.

diff --git a/Engine/Source/Editor/UMGEditor/Private/WidgetBlueprintEditor.cpp b/Engine/Source/Editor/UMGEditor/Private/WidgetBlueprintEditor.cpp
--- a/Engine/Source/Editor/UMGEditor/Private/WidgetBlueprintEditor.cpp
+++ b/Engine/Source/Editor/UMGEditor/Private/WidgetBlueprintEditor.cpp
@@ -1001,6 +1001,7 @@ static bool MigratePropertyValue(UObject* SourceObject, UObject* DestinationObje
 		{
 			if (DestinationObject)
 			{
+				DestinationObject->SetFlags(RF_Transactional);
 				DestinationObject->Modify();
 			}
 			return true;
@@ -1076,6 +1077,7 @@ void FWidgetBlueprintEditor::MigrateFromChain(FEditPropertyChain* PropertyThatCh
 void FWidgetBlueprintEditor::PostUndo(bool bSuccessful)
 {
 	Super::PostUndo(bSuccessful);
+	InvalidatePreview();
 
 	OnWidgetBlueprintTransaction.Broadcast();
 }
@@ -1083,6 +1085,7 @@ void FWidgetBlueprintEditor::PostUndo(bool bSuccessful)
 void FWidgetBlueprintEditor::PostRedo(bool bSuccessful)
 {
 	Super::PostRedo(bSuccessful);
+	InvalidatePreview();
 
 	OnWidgetBlueprintTransaction.Broadcast();
 }
