Removed ensure from ShouldShowOldBindingWidget() that was incorrectly triggering when legacy bindings aren't available.

The method is used both when extending and when checking whether we want to extend a row, and therefore it is not an error to fail to be able to bind to it.

diff --git a/Engine/Source/Editor/UMGEditor/Private/Details/DetailWidgetExtensionHandler.cpp b/Engine/Source/Editor/UMGEditor/Private/Details/DetailWidgetExtensionHandler.cpp
--- a/Engine/Source/Editor/UMGEditor/Private/Details/DetailWidgetExtensionHandler.cpp
+++ b/Engine/Source/Editor/UMGEditor/Private/Details/DetailWidgetExtensionHandler.cpp
@@ -53,7 +53,7 @@ namespace Private
 		const bool bIsEditable = Property->HasAnyPropertyFlags(CPF_Edit | CPF_EditConst);
 		const bool bDoSignaturesMatch = DelegateProperty->SignatureFunction->GetReturnProperty()->SameType(Property);
 
-		if (!ensure(bIsEditable && bDoSignaturesMatch))
+		if (!bIsEditable || !bDoSignaturesMatch)
 		{
 			return false;
 		}
