UMG: Refactor the filter asset data in the palette widget to be reuseable. Hide widget in wrap option.

diff --git a/Engine/Source/Editor/UMGEditor/Private/Palette/SPaletteViewModel.cpp b/Engine/Source/Editor/UMGEditor/Private/Palette/SPaletteViewModel.cpp
--- a/Engine/Source/Editor/UMGEditor/Private/Palette/SPaletteViewModel.cpp
+++ b/Engine/Source/Editor/UMGEditor/Private/Palette/SPaletteViewModel.cpp
@@ -32,125 +32,6 @@
 
 #define LOCTEXT_NAMESPACE "UMG"
 
-namespace UE::Editor::SPaletteViewModel::Private
-{
-	/** Helper class to perform path based filtering for unloaded BP's */
-	class FUnloadedBlueprintData : public IUnloadedBlueprintData
-	{
-	public:
-		FUnloadedBlueprintData(const FAssetData& InAssetData)
-			:ClassPath()
-			,ClassFlags(CLASS_None)
-			,bIsNormalBlueprintType(false)
-		{
-			ClassName = MakeShared<FString>(InAssetData.AssetName.ToString());
-
-			FString GeneratedClassPath;
-			const UClass* AssetClass = InAssetData.GetClass();
-			if (AssetClass && AssetClass->IsChildOf(UBlueprintGeneratedClass::StaticClass()))
-			{
-				ClassPath = InAssetData.ToSoftObjectPath().GetAssetPathString();
-			}
-			else if (InAssetData.GetTagValue(FBlueprintTags::GeneratedClassPath, GeneratedClassPath))
-			{
-				ClassPath = FTopLevelAssetPath(*FPackageName::ExportTextPathToObjectPath(GeneratedClassPath));
-			}
-
-			FEditorClassUtils::GetImplementedInterfaceClassPathsFromAsset(InAssetData, ImplementedInterfaces);
-		}
-
-		virtual ~FUnloadedBlueprintData()
-		{
-		}
-
-		// Begin IUnloadedBlueprintData interface
-		virtual bool HasAnyClassFlags(uint32 InFlagsToCheck) const
-		{
-			return (ClassFlags & InFlagsToCheck) != 0;
-		}
-
-		virtual bool HasAllClassFlags(uint32 InFlagsToCheck) const
-		{
-			return ((ClassFlags & InFlagsToCheck) == InFlagsToCheck);
-		}
-
-		virtual void SetClassFlags(uint32 InFlags)
-		{
-			ClassFlags = InFlags;
-		}
-
-		virtual bool ImplementsInterface(const UClass* InInterface) const
-		{
-			FString InterfacePath = InInterface->GetPathName();
-			for (const FString& ImplementedInterface : ImplementedInterfaces)
-			{
-				if (ImplementedInterface == InterfacePath)
-				{
-					return true;
-				}
-			}
-
-			return false;
-		}
-
-		virtual bool IsChildOf(const UClass* InClass) const
-		{
-			return false;
-		}
-
-		virtual bool IsA(const UClass* InClass) const
-		{
-			// Unloaded blueprint classes should always be a BPGC, so this just checks against the expected type.
-			return UBlueprintGeneratedClass::StaticClass()->UObject::IsA(InClass);
-		}
-
-		virtual const UClass* GetClassWithin() const
-		{
-			return nullptr;
-		}
-
-		virtual const UClass* GetNativeParent() const
-		{
-			return nullptr;
-		}
-
-		virtual void SetNormalBlueprintType(bool bInNormalBPType)
-		{
-			bIsNormalBlueprintType = bInNormalBPType;
-		}
-
-		virtual bool IsNormalBlueprintType() const
-		{
-			return bIsNormalBlueprintType;
-		}
-
-		virtual TSharedPtr<FString> GetClassName() const
-		{
-			return ClassName;
-		}
-
-		virtual FName GetClassPath() const
-		{
-PRAGMA_DISABLE_DEPRECATION_WARNINGS
-			return ClassPath.ToFName();
-PRAGMA_ENABLE_DEPRECATION_WARNINGS
-		}
-
-		virtual FTopLevelAssetPath GetClassPathName() const
-		{
-			return ClassPath;
-		}
-		// End IUnloadedBlueprintData interface
-
-	private:
-		TSharedPtr<FString> ClassName;
-		FTopLevelAssetPath ClassPath;
-		uint32 ClassFlags;
-		TArray<FString> ImplementedInterfaces;
-		bool bIsNormalBlueprintType;
-	};
-}
-
 FWidgetTemplateViewModel::FWidgetTemplateViewModel()
 	: FavortiesViewModel(nullptr),
 	bIsFavorite(false)
@@ -332,310 +213,124 @@ void FWidgetCatalogViewModel::BuildWidgetList()
 
 void FWidgetCatalogViewModel::BuildClassWidgetList()
 {
-	static const FName DevelopmentStatusKey(TEXT("DevelopmentStatus"));
-	static const FString InvalidCategoryName("InvalidCategoryName");
-
-	TMap<FName, TSubclassOf<UUserWidget>> LoadedWidgetBlueprintClassesByName;
-	TMap<FName, TSubclassOf<UBlueprintGeneratedClass>> LoadedGeneratedBlueprintClassesByName;
-
-	auto ActiveWidgetBlueprintClass = GetBlueprint()->GeneratedClass;
+	const UClass* ActiveWidgetBlueprintClass = GetBlueprint()->GeneratedClass;
 	FName ActiveWidgetBlueprintClassName = ActiveWidgetBlueprintClass->GetFName();
 
-	UUMGEditorProjectSettings* UMGEditorProjectSettings = GetMutableDefault<UUMGEditorProjectSettings>();
-	bool bUseEditorConfigPaletteFiltering = UMGEditorProjectSettings->bUseEditorConfigPaletteFiltering;
-	const FNamePermissionList& AllowedPaletteCategories = UMGEditorProjectSettings->GetAllowedPaletteCategories();
-	const FPathPermissionList& AllowedPaletteWidgets = UMGEditorProjectSettings->GetAllowedPaletteWidgets();
-	const TArray<FSoftClassPath>& WidgetClassesToHide = UMGEditorProjectSettings->WidgetClassesToHide;
-
-	// Since allowing all widgets in a category can be excessive when using permission lists, ensure widgets pass class viewing filters
-	FClassViewerModule& ClassViewerModule = FModuleManager::LoadModuleChecked<FClassViewerModule>("ClassViewer");
-	const TSharedPtr<IClassViewerFilter>& GlobalClassFilter = ClassViewerModule.GetGlobalClassViewerFilter();
-	TSharedRef<FClassViewerFilterFuncs> ClassFilterFuncs = ClassViewerModule.CreateFilterFuncs();
-	FClassViewerInitializationOptions ClassViewerOptions = {};
-
-	auto WidgetPassesConfigFiltering = [&](const FAssetData& InWidgetAssetData, const FString& InCategoryName, TWeakObjectPtr<UClass> InWidgetClass)
-	{
-		if (AllowedPaletteWidgets.PassesFilter(InWidgetAssetData.GetObjectPathString()))
-		{
-			return true;
-		}
-
-		if (AllowedPaletteCategories.PassesFilter(*InCategoryName) && GlobalClassFilter.IsValid())
-		{
-			if (InWidgetClass.IsValid())
-			{
-				return GlobalClassFilter->IsClassAllowed(ClassViewerOptions, InWidgetClass.Get(), ClassFilterFuncs);
-			}
-			else if (InWidgetAssetData.IsValid())
-			{
-				using namespace UE::Editor::SPaletteViewModel::Private;
-				TSharedRef<FUnloadedBlueprintData> UnloadedBlueprint = MakeShared<FUnloadedBlueprintData>(InWidgetAssetData);
-				return GlobalClassFilter->IsUnloadedClassAllowed(ClassViewerOptions, UnloadedBlueprint, ClassFilterFuncs);
-			}
-		}
-
-		return false;
-	};
-
 	// Locate all UWidget classes from code and loaded widget BPs
 	for (TObjectIterator<UClass> ClassIt; ClassIt; ++ClassIt)
 	{
 		UClass* WidgetClass = *ClassIt;
-
-		if (!FWidgetBlueprintEditorUtils::IsUsableWidgetClass(WidgetClass))
+		const bool bIsSameClass = WidgetClass->GetFName() == ActiveWidgetBlueprintClassName;
+		if (bIsSameClass)
 		{
 			continue;
 		}
 
-		if (WidgetClass->HasAnyClassFlags(CLASS_HideDropDown))
+		if (!FWidgetBlueprintEditorUtils::IsUsableWidgetClass(WidgetClass))
 		{
 			continue;
 		}
 
-		// Initialize AssetData for checking PackagePath
-		FAssetData WidgetAssetData = FAssetData(WidgetClass);
-
-		// Excludes engine content if user sets it to false
-		if (!GetDefault<UContentBrowserSettings>()->GetDisplayEngineFolder() || !GetDefault<UUMGEditorProjectSettings>()->bShowWidgetsFromEngineContent)
-		{
-			if (WidgetAssetData.PackagePath.ToString().Find(TEXT("/Engine")) == 0)
-			{
-				continue;
-			}
-		}
-
-		// Excludes developer content if user sets it to false
-		if (!GetDefault<UContentBrowserSettings>()->GetDisplayDevelopersFolder() || !GetDefault<UUMGEditorProjectSettings>()->bShowWidgetsFromDeveloperContent)
-		{
-			if (WidgetAssetData.PackagePath.ToString().Find(TEXT("/Game/Developers")) == 0)
-			{
-				continue;
-			}
-		}
-
-		if (bUseEditorConfigPaletteFiltering)
-		{
-			if (!WidgetPassesConfigFiltering(WidgetAssetData, WidgetClass->GetDefaultObject<UWidget>()->GetPaletteCategory().ToString(), WidgetClass))
-			{
-				continue;
-			}
-		}
-		else
-		{
-			// Excludes this widget if it is on the hide list
-			bool bIsOnList = false;
-			for (FSoftClassPath Widget : WidgetClassesToHide)
-			{
-				if (WidgetAssetData.GetObjectPathString().Find(Widget.ToString()) == 0)
-				{
-					bIsOnList = true;
-					break;
-				}
-			}
-			if (bIsOnList)
-			{
-				continue;
-			}
-		}
-
-		const bool bIsSameClass = WidgetClass->GetFName() == ActiveWidgetBlueprintClassName;
-
-		// Check that the asset that generated this class is valid (necessary b/c of a larger issue wherein force delete does not wipe the generated class object)
-		if ( bIsSameClass )
+		if (WidgetClass->HasAnyClassFlags(CLASS_HideDropDown | CLASS_Hidden))
 		{
 			continue;
 		}
 
 		if (WidgetClass->IsChildOf(UUserWidget::StaticClass()))
 		{
-			if ( WidgetClass->ClassGeneratedBy )
-			{
-				// Track the widget blueprint classes that are already loaded
-				LoadedWidgetBlueprintClassesByName.Add(WidgetClass->ClassGeneratedBy->GetFName()) = WidgetClass;
-			}
-		}
-		else if (Cast<UBlueprintGeneratedClass>(WidgetClass))
-		{
-			// Note: Don't check IsChildOf above, since superstruct of BPGC is the C++ class
-			LoadedGeneratedBlueprintClassesByName.Add(WidgetClass->GetFName()) = WidgetClass;
-			TSharedPtr<FWidgetTemplateClass> Template = MakeShareable(new FWidgetTemplateClass(WidgetClass));
-			AddWidgetTemplate(Template);
+			bool bIsBlueprintGeneratedClass = Cast<UBlueprintGeneratedClass>(WidgetClass) != nullptr;
+			AddWidgetTemplate(MakeShared<FWidgetTemplateBlueprintClass>(FAssetData(WidgetClass), WidgetClass, bIsBlueprintGeneratedClass));
 		}
 		else
 		{
-			TSharedPtr<FWidgetTemplateClass> Template = MakeShareable(new FWidgetTemplateClass(WidgetClass));
-			AddWidgetTemplate(Template);
+			// For UWidget
+			AddWidgetTemplate(MakeShared<FWidgetTemplateClass>(WidgetClass));
 		}
-
-		//TODO UMG does not prevent deep nested circular references
 	}
 
-	// Locate all widget BP assets (include unloaded)
-	const FAssetRegistryModule& AssetRegistryModule = FModuleManager::LoadModuleChecked<FAssetRegistryModule>(TEXT("AssetRegistry"));
+	// Locate all UWidget BP assets, include loaded and unloaded. Only parsed the unloaded.
+	const FAssetRegistryModule* AssetRegistryModule = FModuleManager::GetModulePtr<FAssetRegistryModule>(TEXT("AssetRegistry"));
 	TArray<FAssetData> AllBPsAssetData;
-	AssetRegistryModule.Get().GetAssetsByClass(UBlueprint::StaticClass()->GetClassPathName(), AllBPsAssetData, true);
+	if (AssetRegistryModule)
+	{
+		AssetRegistryModule->Get().GetAssetsByClass(UBlueprint::StaticClass()->GetClassPathName(), AllBPsAssetData, true);
+	}
 
-	for (FAssetData& BPAssetData : AllBPsAssetData)
+	for (const FAssetData& BPAssetData : AllBPsAssetData)
 	{
-		// Blueprints get the class type actions for their parent native class - this avoids us having to load the blueprint
-		UClass* ParentClass = nullptr;
-		FString ParentClassName;
-		if (!BPAssetData.GetTagValue(FBlueprintTags::NativeParentClassPath, ParentClassName))
+		const bool bIsSameClass = BPAssetData.AssetName == ActiveWidgetBlueprintClassName;
+		if (bIsSameClass)
 		{
-			BPAssetData.GetTagValue(FBlueprintTags::ParentClassPath, ParentClassName);
+			continue;
 		}
-		if (!ParentClassName.IsEmpty())
+
+		// Was already parsed by the TObjectIterator<UClass>
+		if (BPAssetData.IsAssetLoaded())
 		{
-			ParentClass = UClass::TryFindTypeSlow<UClass>(FPackageName::ExportTextPathToObjectPath(ParentClassName));
-			// UUserWidgets have their own loading section, and we don't want to process any blueprints that don't have UWidget parents
-			if (ParentClass)
-			{
-				if (!ParentClass->IsChildOf(UWidget::StaticClass()) || ParentClass->IsChildOf(UUserWidget::StaticClass()))
-				{
-					continue;
-				}
-			}
+			continue;
 		}
 
-		if (!FilterAssetData(BPAssetData))
+		TValueOrError<FWidgetBlueprintEditorUtils::FUsableWidgetClassResult, void> Usable = FWidgetBlueprintEditorUtils::IsUsableWidgetClass(BPAssetData);
+		if (Usable.HasError())
 		{
-			// If this object isn't currently loaded, add it to the palette view
-			if (BPAssetData.ToSoftObjectPath().ResolveObject() == nullptr)
-			{
-				if (bUseEditorConfigPaletteFiltering)
-				{
-					// Still check the soft object path even if we can't check category due to BP asset being unloaded
-					if (!WidgetPassesConfigFiltering(BPAssetData, InvalidCategoryName, nullptr))
-					{
-						continue;
-					}
-				}
-
-				auto Template = MakeShareable(new FWidgetTemplateClass(BPAssetData, nullptr));
-				AddWidgetTemplate(Template);
-			}
+			continue;
 		}
-	}
-
-	TArray<FAssetData> AllWidgetBPsAssetData;
-	AssetRegistryModule.Get().GetAssetsByClass(UWidgetBlueprint::StaticClass()->GetClassPathName(), AllWidgetBPsAssetData, true);
 
-	FName ActiveWidgetBlueprintName = ActiveWidgetBlueprintClass->ClassGeneratedBy->GetFName();
-	for (FAssetData& WidgetBPAssetData : AllWidgetBPsAssetData)
-	{
-		// Excludes the blueprint you're currently in
-		if (WidgetBPAssetData.AssetName == ActiveWidgetBlueprintName)
+		if ((Usable.GetValue().AssetClassFlags & (CLASS_Hidden | CLASS_HideDropDown)) != 0)
 		{
 			continue;
-		}
+		}		
 
-		if (!FilterAssetData(WidgetBPAssetData))
+		if (Usable.GetValue().NativeParentClass->IsChildOf(UUserWidget::StaticClass()))
+		{
+			bool bIsBlueprintGeneratedClass = false;
+			AddWidgetTemplate(MakeShared<FWidgetTemplateBlueprintClass>(BPAssetData, nullptr, bIsBlueprintGeneratedClass));
+		}
+		else
 		{
-			// If the blueprint generated class was found earlier, pass it to the template
-			TSubclassOf<UUserWidget> WidgetBPClass = nullptr;
-			auto LoadedWidgetBPClass = LoadedWidgetBlueprintClassesByName.Find(WidgetBPAssetData.AssetName);
-			if (LoadedWidgetBPClass)
-			{
-				WidgetBPClass = *LoadedWidgetBPClass;
-			}
-
-			if (bUseEditorConfigPaletteFiltering)
-			{
-				FText Category = FWidgetTemplateBlueprintClass(WidgetBPAssetData, WidgetBPClass).GetCategory();
-				if (!WidgetPassesConfigFiltering(WidgetBPAssetData, Category.ToString(), WidgetBPClass))
-				{
-					continue;
-				}
-			}
-			else
-			{
-				// Excludes this widget if it is on the hide list
-				bool bIsOnList = false;
-				for (FSoftClassPath Widget : WidgetClassesToHide)
-				{
-					if (Widget.ToString().Find(WidgetBPAssetData.GetObjectPathString()) == 0)
-					{
-						bIsOnList = true;
-						break;
-					}
-				}
-				if (bIsOnList)
-				{
-					continue;
-				}
-			}
-
-			uint32 BPFlags = WidgetBPAssetData.GetTagValueRef<uint32>(FBlueprintTags::ClassFlags);
-			if ((BPFlags & (CLASS_Hidden | CLASS_HideDropDown | CLASS_Deprecated | CLASS_Abstract | CLASS_NewerVersionExists)) == 0)
-			{
-				auto Template = MakeShareable(new FWidgetTemplateBlueprintClass(WidgetBPAssetData, WidgetBPClass));
-
-				AddWidgetTemplate(Template);
-			}
+			AddWidgetTemplate(MakeShared<FWidgetTemplateClass>(BPAssetData, nullptr));
 		}
 	}
 
 
-	TArray<FAssetData> AllGeneratedBPsAssetData;
-	AssetRegistryModule.Get().GetAssetsByClass(UBlueprintGeneratedClass::StaticClass()->GetClassPathName(), AllGeneratedBPsAssetData, true);
+	TArray<FAssetData> AllGeneratedBPsAssetData; // if it's a widget already compiled
 
-	// Search generated class assets if using palette filtering, only do this with palette filtering since otherwise searching all generated BP's would be expensive
-	if (bUseEditorConfigPaletteFiltering)
+	if (AssetRegistryModule && GetDefault<UUMGEditorProjectSettings>()->bUseEditorConfigPaletteFiltering)
 	{
-		for (FAssetData& GeneratedBPAssetData : AllGeneratedBPsAssetData)
-		{
-			// Excludes the blueprint you're currently in
-			if (GeneratedBPAssetData.AssetName == ActiveWidgetBlueprintName)
-			{
-				continue;
-			}
-
-			if (!FilterAssetData(GeneratedBPAssetData) 
-				&& AllowedPaletteWidgets.PassesFilter(GeneratedBPAssetData.GetObjectPathString()) 
-				&& !LoadedGeneratedBlueprintClassesByName.Contains(GeneratedBPAssetData.AssetName))
-			{
-				uint32 BPFlags = GeneratedBPAssetData.GetTagValueRef<uint32>(FBlueprintTags::ClassFlags);
-				if ((BPFlags & (CLASS_Hidden | CLASS_HideDropDown | CLASS_Deprecated | CLASS_Abstract | CLASS_NewerVersionExists)) == 0)
-				{
-					bool bIsBlueprintGeneratedClass = true;
-					auto Template = MakeShareable(new FWidgetTemplateBlueprintClass(GeneratedBPAssetData, nullptr, bIsBlueprintGeneratedClass));
-
-					AddWidgetTemplate(Template);
-				}
-			}
-		}
+		AssetRegistryModule->Get().GetAssetsByClass(UBlueprintGeneratedClass::StaticClass()->GetClassPathName(), AllGeneratedBPsAssetData, true);
 	}
-}
 
-bool FWidgetCatalogViewModel::FilterAssetData(FAssetData &InAssetData)
-{
-	// Excludes engine content if user sets it to false
-	if (!GetDefault<UContentBrowserSettings>()->GetDisplayEngineFolder() || !GetDefault<UUMGEditorProjectSettings>()->bShowWidgetsFromEngineContent)
+	for (const FAssetData& BPAssetData : AllGeneratedBPsAssetData)
 	{
-		if (InAssetData.PackagePath.ToString().Find(TEXT("/Engine")) == 0)
+		const bool bIsSameClass = BPAssetData.AssetName == ActiveWidgetBlueprintClassName;
+		if (bIsSameClass)
 		{
-			return true;
+			continue;
 		}
-	}
 
-	// Excludes developer content if user sets it to false
-	if (!GetDefault<UContentBrowserSettings>()->GetDisplayDevelopersFolder() || !GetDefault<UUMGEditorProjectSettings>()->bShowWidgetsFromDeveloperContent)
-	{
-		if (InAssetData.PackagePath.ToString().Find(TEXT("/Game/Developers")) == 0)
+		TValueOrError<FWidgetBlueprintEditorUtils::FUsableWidgetClassResult, void> Usable = FWidgetBlueprintEditorUtils::IsUsableWidgetClass(BPAssetData);
+		if (Usable.HasError())
 		{
-			return true;
+			continue;
+		}
+
+		if ((Usable.GetValue().AssetClassFlags & (CLASS_Hidden | CLASS_HideDropDown)) != 0)
+		{
+			continue;
 		}
+
+		bool bIsBlueprintGeneratedClass = true;
+		AddWidgetTemplate(MakeShared<FWidgetTemplateBlueprintClass>(BPAssetData, nullptr, bIsBlueprintGeneratedClass));
 	}
-	return false;
 }
 
 void FWidgetCatalogViewModel::AddWidgetTemplate(TSharedPtr<FWidgetTemplate> Template)
 {
-	FString Category = Template->GetCategory().ToString();
+	FString Category = *Template->GetCategory().ToString();
 
 	// Hide user specific categories
-	TArray<FString> CategoriesToHide = GetDefault<UUMGEditorProjectSettings>()->CategoriesToHide;
-	for (FString CategoryName : CategoriesToHide)
+	const TArray<FString>& CategoriesToHide = GetDefault<UUMGEditorProjectSettings>()->CategoriesToHide;
+	for (const FString& CategoryName : CategoriesToHide)
 	{
 		if (Category == CategoryName)
 		{
diff --git a/Engine/Source/Editor/UMGEditor/Private/Palette/SPaletteViewModel.h b/Engine/Source/Editor/UMGEditor/Private/Palette/SPaletteViewModel.h
--- a/Engine/Source/Editor/UMGEditor/Private/Palette/SPaletteViewModel.h
+++ b/Engine/Source/Editor/UMGEditor/Private/Palette/SPaletteViewModel.h
@@ -182,8 +182,6 @@ protected:
 	virtual void BuildWidgetTemplateCategory(FString& Category, TArray<TSharedPtr<FWidgetTemplate>>& Templates, TArray<FString>& FavoritesList) = 0;
 	void BuildClassWidgetList();
 
-	static bool FilterAssetData(FAssetData& BPAssetData);
-
 	void AddWidgetTemplate(TSharedPtr<FWidgetTemplate> Template);
 
 	/** Called when a Blueprint is recompiled and live objects are swapped out for replacements */
diff --git a/Engine/Source/Editor/UMGEditor/Private/Settings/UMGEditorProjectSettings.cpp b/Engine/Source/Editor/UMGEditor/Private/Settings/UMGEditorProjectSettings.cpp
--- a/Engine/Source/Editor/UMGEditor/Private/Settings/UMGEditorProjectSettings.cpp
+++ b/Engine/Source/Editor/UMGEditor/Private/Settings/UMGEditorProjectSettings.cpp
@@ -89,11 +89,21 @@ FNamePermissionList& UUMGEditorProjectSettings::GetAllowedPaletteCategories()
 	return AllowedPaletteCategories;
 }
 
+const FNamePermissionList& UUMGEditorProjectSettings::GetAllowedPaletteCategories() const
+{
+	return AllowedPaletteCategories;
+}
+
 FPathPermissionList& UUMGEditorProjectSettings::GetAllowedPaletteWidgets()
 {
 	return AllowedPaletteWidgets;
 }
 
+const FPathPermissionList& UUMGEditorProjectSettings::GetAllowedPaletteWidgets() const
+{
+	return AllowedPaletteWidgets;
+}
+
 void UUMGEditorProjectSettings::GetCompilerOptionsForWidget(const UWidgetBlueprint* WidgetBlueprint, TFunctionRef<bool(const FWidgetCompilerOptions&)> Operator) const
 {
 	FString AssetPath = WidgetBlueprint->GetOutermost()->GetName();
diff --git a/Engine/Source/Editor/UMGEditor/Private/WidgetBlueprintEditorUtils.cpp b/Engine/Source/Editor/UMGEditor/Private/WidgetBlueprintEditorUtils.cpp
--- a/Engine/Source/Editor/UMGEditor/Private/WidgetBlueprintEditorUtils.cpp
+++ b/Engine/Source/Editor/UMGEditor/Private/WidgetBlueprintEditorUtils.cpp
@@ -9,13 +9,18 @@
 #include "UObject/UObjectIterator.h"
 #include "Internationalization/TextPackageNamespaceUtil.h"
 #include "UObject/PropertyPortFlags.h"
+#include "UObject/TopLevelAssetPath.h"
 #include "Blueprint/WidgetTree.h"
 #include "Misc/ConfigCacheIni.h"
 #include "Modules/ModuleManager.h"
 #include "MovieScene.h"
+#include "UMGEditorProjectSettings.h"
 #include "WidgetBlueprint.h"
 #include "HAL/PlatformApplicationMisc.h"
+#include "Settings/ContentBrowserSettings.h"
 
+#include "ClassViewerFilter.h"
+#include "EditorClassUtils.h"
 #include "Dialogs/Dialogs.h"
 #include "DragAndDrop/DecoratedDragDropOp.h"
 #include "DragAndDrop/AssetDragDropOp.h"
@@ -44,6 +49,7 @@
 #include "Components/Widget.h"
 #include "Blueprint/WidgetNavigation.h"
 #include "Subsystems/AssetEditorSubsystem.h"
+#include "UObject/CoreRedirects.h"
 #include "UObject/ScriptInterface.h"
 #include "Components/NamedSlotInterface.h"
 #include "K2Node_Variable.h"
@@ -1958,16 +1964,195 @@ bool FWidgetBlueprintEditorUtils::IsBindWidgetAnimProperty(FProperty* InProperty
 	return false;
 }
 
-bool FWidgetBlueprintEditorUtils::IsUsableWidgetClass(UClass* WidgetClass)
+namespace UE::UMG::Private
 {
-	if ( WidgetClass->IsChildOf(UWidget::StaticClass()) )
+/** Helper class to perform path based filtering for unloaded BP's */
+class FUnloadedBlueprintData : public IUnloadedBlueprintData
+{
+public:
+	FUnloadedBlueprintData(const FAssetData& InAssetData)
+		:ClassPath()
+		, ClassFlags(CLASS_None)
+		, bIsNormalBlueprintType(false)
+	{
+		ClassName = MakeShared<FString>(InAssetData.AssetName.ToString());
+
+		FString GeneratedClassPath;
+		const UClass* AssetClass = InAssetData.GetClass();
+		if (AssetClass && AssetClass->IsChildOf(UBlueprintGeneratedClass::StaticClass()))
+		{
+			ClassPath = InAssetData.ToSoftObjectPath().GetAssetPathString();
+		}
+		else if (InAssetData.GetTagValue(FBlueprintTags::GeneratedClassPath, GeneratedClassPath))
+		{
+			ClassPath = FTopLevelAssetPath(*FPackageName::ExportTextPathToObjectPath(GeneratedClassPath));
+		}
+
+		FEditorClassUtils::GetImplementedInterfaceClassPathsFromAsset(InAssetData, ImplementedInterfaces);
+	}
+
+	virtual ~FUnloadedBlueprintData()
+	{
+	}
+
+	// Begin IUnloadedBlueprintData interface
+	virtual bool HasAnyClassFlags(uint32 InFlagsToCheck) const
+	{
+		return (ClassFlags & InFlagsToCheck) != 0;
+	}
+
+	virtual bool HasAllClassFlags(uint32 InFlagsToCheck) const
+	{
+		return ((ClassFlags & InFlagsToCheck) == InFlagsToCheck);
+	}
+
+	virtual void SetClassFlags(uint32 InFlags)
+	{
+		ClassFlags = InFlags;
+	}
+
+	virtual bool ImplementsInterface(const UClass* InInterface) const
+	{
+		FString InterfacePath = InInterface->GetPathName();
+		for (const FString& ImplementedInterface : ImplementedInterfaces)
+		{
+			if (ImplementedInterface == InterfacePath)
+			{
+				return true;
+			}
+		}
+
+		return false;
+	}
+
+	virtual bool IsChildOf(const UClass* InClass) const
+	{
+		return false;
+	}
+
+	virtual bool IsA(const UClass* InClass) const
+	{
+		// Unloaded blueprint classes should always be a BPGC, so this just checks against the expected type.
+		return UBlueprintGeneratedClass::StaticClass()->UObject::IsA(InClass);
+	}
+
+	virtual const UClass* GetClassWithin() const
+	{
+		return nullptr;
+	}
+
+	virtual const UClass* GetNativeParent() const
+	{
+		return nullptr;
+	}
+
+	virtual void SetNormalBlueprintType(bool bInNormalBPType)
+	{
+		bIsNormalBlueprintType = bInNormalBPType;
+	}
+
+	virtual bool IsNormalBlueprintType() const
+	{
+		return bIsNormalBlueprintType;
+	}
+
+	virtual TSharedPtr<FString> GetClassName() const
+	{
+		return ClassName;
+	}
+
+	virtual FName GetClassPath() const
+	{
+PRAGMA_DISABLE_DEPRECATION_WARNINGS
+		return ClassPath.ToFName();
+PRAGMA_ENABLE_DEPRECATION_WARNINGS
+	}
+
+	virtual FTopLevelAssetPath GetClassPathName() const
+	{
+		return ClassPath;
+	}
+	// End IUnloadedBlueprintData interface
+
+private:
+	TSharedPtr<FString> ClassName;
+	FTopLevelAssetPath ClassPath;
+	uint32 ClassFlags;
+	TArray<FString> ImplementedInterfaces;
+	bool bIsNormalBlueprintType;
+};
+
+bool IsUsableWidgetClass(const FString& WidgetPathName, const FAssetData& WidgetAssetData, FName Category, const UClass* WidgetClass)
+{
+	const UUMGEditorProjectSettings* UMGEditorProjectSettings = GetDefault<UUMGEditorProjectSettings>();
+	const UContentBrowserSettings* ContentBrowserSettings = GetDefault<UContentBrowserSettings>();
+
+	// Excludes engine content if user sets it to false
+	if (!ContentBrowserSettings->GetDisplayEngineFolder() || !UMGEditorProjectSettings->bShowWidgetsFromEngineContent)
+	{
+		if (WidgetPathName.StartsWith(TEXT("/Engine")))
+		{
+			return false;
+		}
+	}
+
+	// Excludes developer content if user sets it to false
+	if (!ContentBrowserSettings->GetDisplayDevelopersFolder() || !UMGEditorProjectSettings->bShowWidgetsFromDeveloperContent)
+	{
+		if (WidgetPathName.StartsWith(TEXT("/Game/Developers")))
+		{
+			return false;
+		}
+	}
+
+	if (UMGEditorProjectSettings->bUseEditorConfigPaletteFiltering)
+	{
+		if (UMGEditorProjectSettings->GetAllowedPaletteWidgets().PassesFilter(WidgetPathName))
+		{
+			return true;
+		}
+
+		FClassViewerModule* ClassViewerModule = FModuleManager::GetModulePtr<FClassViewerModule>("ClassViewer");
+		const TSharedPtr<IClassViewerFilter> GlobalClassFilter = ClassViewerModule ? ClassViewerModule->GetGlobalClassViewerFilter() : TSharedPtr<IClassViewerFilter>();
+		if (UMGEditorProjectSettings->GetAllowedPaletteCategories().PassesFilter(Category) && GlobalClassFilter.IsValid())
+		{
+			if (WidgetClass)
+			{
+				return GlobalClassFilter->IsClassAllowed(FClassViewerInitializationOptions(), WidgetClass, ClassViewerModule->CreateFilterFuncs());
+			}
+			else if (WidgetAssetData.IsValid())
+			{
+				TSharedRef<FUnloadedBlueprintData> UnloadedBlueprint = MakeShared<FUnloadedBlueprintData>(WidgetAssetData);
+				return GlobalClassFilter->IsUnloadedClassAllowed(FClassViewerInitializationOptions(), UnloadedBlueprint, ClassViewerModule->CreateFilterFuncs());
+			}
+		}
+		return false;
+	}
+	else
+	{
+		// Excludes this widget if it is on the hide list
+		for (const FSoftClassPath& WidgetClassToHide : UMGEditorProjectSettings->WidgetClassesToHide)
+		{
+			if (WidgetPathName.Find(WidgetClassToHide.ToString()) == 0)
+			{
+				return false;
+			}
+		}
+	}
+	return true;
+}
+}
+
+bool FWidgetBlueprintEditorUtils::IsUsableWidgetClass(const UClass* WidgetClass)
+{
+	if (WidgetClass->IsChildOf(UWidget::StaticClass()))
 	{
 		// We aren't interested in classes that are experimental or cannot be instantiated
 		bool bIsExperimental, bIsEarlyAccess;
 		FString MostDerivedDevelopmentClassName;
-		FObjectEditorUtils::GetClassDevelopmentStatus(WidgetClass, bIsExperimental, bIsEarlyAccess, MostDerivedDevelopmentClassName);
+		FObjectEditorUtils::GetClassDevelopmentStatus(const_cast<UClass*>(WidgetClass), bIsExperimental, bIsEarlyAccess, MostDerivedDevelopmentClassName);
 		const bool bIsInvalid = WidgetClass->HasAnyClassFlags(CLASS_Abstract | CLASS_Deprecated | CLASS_NewerVersionExists);
-		if ( bIsExperimental || bIsEarlyAccess || bIsInvalid )
+		if (bIsExperimental || bIsEarlyAccess || bIsInvalid)
 		{
 			return false;
 		}
@@ -1976,17 +2161,75 @@ bool FWidgetBlueprintEditorUtils::IsUsableWidgetClass(UClass* WidgetClass)
 		const bool bIsSkeletonClass = WidgetClass->HasAnyFlags(RF_Transient) && WidgetClass->HasAnyClassFlags(CLASS_CompiledFromBlueprint);
 
 		// Check that the asset that generated this class is valid (necessary b/c of a larger issue wherein force delete does not wipe the generated class object)
-		if ( bIsSkeletonClass )
+		if (bIsSkeletonClass)
 		{
 			return false;
 		}
 
-		return true;
+		return UE::UMG::Private::IsUsableWidgetClass(WidgetClass->GetPathName(), FAssetData(), *WidgetClass->GetDefaultObject<UWidget>()->GetPaletteCategory().ToString(), WidgetClass);
 	}
 
 	return false;
 }
 
+TValueOrError<FWidgetBlueprintEditorUtils::FUsableWidgetClassResult, void> FWidgetBlueprintEditorUtils::IsUsableWidgetClass(const FAssetData& WidgetAsset)
+{
+	if (const UClass* WidgetAssetClass = WidgetAsset.GetClass(EResolveClass::No))
+	{
+		if (IsUsableWidgetClass(WidgetAssetClass))
+		{
+			FWidgetBlueprintEditorUtils::FUsableWidgetClassResult Result;
+			Result.NativeParentClass = WidgetAssetClass;
+			Result.AssetClassFlags = WidgetAssetClass->GetClassFlags();
+			return MakeValue(Result);
+		}
+	}
+
+	// Blueprints get the class type actions for their parent native class - this avoids us having to load the blueprint
+	UClass* NativeParentClass = nullptr;
+	FString NativeParentClassName;
+	WidgetAsset.GetTagValue(FBlueprintTags::NativeParentClassPath, NativeParentClassName);
+	if (NativeParentClassName.IsEmpty())
+	{
+		return MakeError();
+	}
+	else
+	{
+		const FString NativeParentClassPath = FPackageName::ExportTextPathToObjectPath(NativeParentClassName);
+		if (NativeParentClassPath.StartsWith(TEXT("/")))
+		{
+			// Metadata may be pointing to classes that no longer exist, so check for redirectors first
+			const FString RedirectedClassPath = FCoreRedirects::GetRedirectedName(ECoreRedirectFlags::Type_Class, FCoreRedirectObjectName(NativeParentClassPath)).ToString();
+			NativeParentClass = UClass::TryFindTypeSlow<UClass>(RedirectedClassPath);
+		}
+		if (NativeParentClass == nullptr)
+		{
+			return MakeError();
+		}
+		if (!NativeParentClass->IsChildOf(UWidget::StaticClass()))
+		{
+			return MakeError();
+		}
+	}
+
+	EClassFlags BPFlags = static_cast<EClassFlags>(WidgetAsset.GetTagValueRef<uint32>(FBlueprintTags::ClassFlags));
+	const bool bIsInvalid = (BPFlags & (CLASS_Deprecated | CLASS_Abstract | CLASS_NewerVersionExists)) != 0;
+	if (bIsInvalid)
+	{
+		return MakeError();
+	}
+
+	const FName CategoryName = *GetPaletteCategory(WidgetAsset, TSubclassOf<UWidget>(NativeParentClass)).ToString();
+	if (UE::UMG::Private::IsUsableWidgetClass(WidgetAsset.GetObjectPathString(), WidgetAsset, CategoryName, nullptr))
+	{
+		FWidgetBlueprintEditorUtils::FUsableWidgetClassResult Result;
+		Result.NativeParentClass = NativeParentClass;
+		Result.AssetClassFlags = BPFlags;
+		return MakeValue(Result);
+	}
+	return MakeError();
+}
+
 FString RemoveSuffixFromName(const FString OldName)
 {
 	int NameLen = OldName.Len();
@@ -2361,4 +2604,48 @@ void FWidgetBlueprintEditorUtils::SetTextureAsAssetThumbnail(UWidgetBlueprint* W
 	}
 }
 
+FText FWidgetBlueprintEditorUtils::GetPaletteCategory(const TSubclassOf<UWidget> WidgetClass)
+{
+	if (WidgetClass.Get())
+	{
+		return WidgetClass.GetDefaultObject()->GetPaletteCategory();
+	}
+	return GetMutableDefault<UWidget>()->GetPaletteCategory();
+}
+
+FText FWidgetBlueprintEditorUtils::GetPaletteCategory(const FAssetData& WidgetAsset, const TSubclassOf<UWidget> NativeClass)
+{
+	//The asset can be a UBlueprint, a UWidgetBlueprint or a UWidgetBlueprintGeneratedClass
+
+	if (UClass* WidgetAssetClass = WidgetAsset.GetClass(EResolveClass::No))
+	{
+		if (WidgetAssetClass->IsChildOf(UWidget::StaticClass()))
+		{
+			return GetPaletteCategory(TSubclassOf<UWidget>(WidgetAssetClass));
+		}
+	}
+
+	//If the blueprint is unloaded we need to extract it from the asset metadata.
+	FText FoundPaletteCategoryText = WidgetAsset.GetTagValueRef<FText>(GET_MEMBER_NAME_CHECKED(UWidgetBlueprint, PaletteCategory));
+	if (!FoundPaletteCategoryText.IsEmpty())
+	{
+		return FoundPaletteCategoryText;
+	}
+	else if (NativeClass.Get() != nullptr && NativeClass->IsChildOf(UWidget::StaticClass()) && !NativeClass->IsChildOf(UUserWidget::StaticClass()))
+	{
+		return NativeClass.GetDefaultObject()->GetPaletteCategory();
+	}
+
+	static const FTopLevelAssetPath BlueprintGeneratedClassAssetPath = UWidgetBlueprintGeneratedClass::StaticClass()->GetClassPathName();
+	static const FTopLevelAssetPath WidgetBlueprintAssetPath = UWidgetBlueprint::StaticClass()->GetClassPathName();
+	if (WidgetAsset.AssetClassPath == BlueprintGeneratedClassAssetPath || WidgetAsset.AssetClassPath == WidgetBlueprintAssetPath)
+	{
+		return GetMutableDefault<UUserWidget>()->GetPaletteCategory();
+	}
+	else
+	{
+		return GetMutableDefault<UWidget>()->GetPaletteCategory();
+	}
+}
+
 #undef LOCTEXT_NAMESPACE
diff --git a/Engine/Source/Editor/UMGEditor/Public/UMGEditorProjectSettings.h b/Engine/Source/Editor/UMGEditor/Public/UMGEditorProjectSettings.h
--- a/Engine/Source/Editor/UMGEditor/Public/UMGEditorProjectSettings.h
+++ b/Engine/Source/Editor/UMGEditor/Public/UMGEditorProjectSettings.h
@@ -205,9 +205,11 @@ public:
 
 	/** Get the permission list that controls which categories are exposed in config palette filtering */
 	FNamePermissionList& GetAllowedPaletteCategories();
+	const FNamePermissionList& GetAllowedPaletteCategories() const;
 
 	/** Get the permission list that controls which widgets are exposed in config palette filtering */
 	FPathPermissionList& GetAllowedPaletteWidgets();
+	const FPathPermissionList& GetAllowedPaletteWidgets() const;
 
 private:
 	template<typename ReturnType, typename T>
diff --git a/Engine/Source/Editor/UMGEditor/Public/WidgetBlueprintEditorUtils.h b/Engine/Source/Editor/UMGEditor/Public/WidgetBlueprintEditorUtils.h
--- a/Engine/Source/Editor/UMGEditor/Public/WidgetBlueprintEditorUtils.h
+++ b/Engine/Source/Editor/UMGEditor/Public/WidgetBlueprintEditorUtils.h
@@ -52,7 +52,13 @@ public:
 	static bool IsBindWidgetAnimProperty(FProperty* InProperty);
 	static bool IsBindWidgetAnimProperty(FProperty* InProperty, bool& bIsOptional);
 
-	static bool IsUsableWidgetClass(UClass* WidgetClass);
+	struct FUsableWidgetClassResult
+	{
+		const UClass* NativeParentClass = nullptr;
+		EClassFlags AssetClassFlags = EClassFlags::CLASS_None;
+	};
+	static bool IsUsableWidgetClass(const UClass* WidgetClass);
+	static TValueOrError<FUsableWidgetClassResult, void> IsUsableWidgetClass(const FAssetData& WidgetAsset);
 
 	static void ExportWidgetsToText(TArray<UWidget*> WidgetsToExport, /*out*/ FString& ExportedText);
 
@@ -90,6 +96,9 @@ public:
 
 	static void SetTextureAsAssetThumbnail(UWidgetBlueprint* WidgetBlueprint, UTexture2D* ThumbnailTexture);
 
+	static FText GetPaletteCategory(const TSubclassOf<UWidget> Widget);
+	static FText GetPaletteCategory(const FAssetData& WidgetAsset, const TSubclassOf<UWidget> NativeClass);
+
 	static TOptional<FWidgetThumbnailProperties> DrawSWidgetInRenderTargetForThumbnail(UUserWidget* WidgetInstance, FRenderTarget* RenderTarget2D, FVector2D ThumbnailSize, TOptional<FVector2D> ThumbnailCustomSize, EThumbnailPreviewSizeMode ThumbnailSizeMode = EThumbnailPreviewSizeMode::MatchDesignerMode);
 
 	static TOptional<FWidgetThumbnailProperties> DrawSWidgetInRenderTargetForThumbnail(UUserWidget* WidgetInstance, UTextureRenderTarget2D* RenderTarget2D, FVector2D ThumbnailSize, TOptional<FVector2D> ThumbnailCustomSize, EThumbnailPreviewSizeMode ThumbnailSizeMode);
