PR #9996: Expose the list of common root widgets displayed on the root widget class picker dialog (Contributed by [REDACTED_PERSON])

PR #9634: Add a list of favorite widget parent class to project settings

diff --git a/Engine/Source/Editor/UMGEditor/Private/WidgetBlueprintFactory.cpp b/Engine/Source/Editor/UMGEditor/Private/WidgetBlueprintFactory.cpp
--- a/Engine/Source/Editor/UMGEditor/Private/WidgetBlueprintFactory.cpp
+++ b/Engine/Source/Editor/UMGEditor/Private/WidgetBlueprintFactory.cpp
@@ -69,7 +69,7 @@ bool UWidgetBlueprintFactory::ConfigureProperties()
 		TSharedPtr<FWidgetClassFilter> Filter = MakeShareable(new FWidgetClassFilter);
 		Options.ClassFilters.Add(Filter.ToSharedRef());
 
-		TArray<TSoftClassPtr<UUserWidget>> FavoriteWidgetParentClasses = GetDefault <UUMGEditorProjectSettings>()->FavoriteWidgetParentClasses;
+		const TArray<TSoftClassPtr<UUserWidget>>& FavoriteWidgetParentClasses = GetDefault <UUMGEditorProjectSettings>()->FavoriteWidgetParentClasses;
 		for (int32 Index = 0; Index < FavoriteWidgetParentClasses.Num(); ++Index)
 		{
 			UClass* FavoriteWidgetParentClass = FavoriteWidgetParentClasses[Index].LoadSynchronous();
@@ -87,7 +87,8 @@ bool UWidgetBlueprintFactory::ConfigureProperties()
 			Options.ExtraPickerCommonClasses.Add(UUserWidget::StaticClass());
 		}
 
-		Filter->DisallowedClassFlags = CLASS_Deprecated | CLASS_NewerVersionExists;
+//Filter->DisallowedClassFlags = CLASS_Deprecated | CLASS_Abstract | CLASS_NewerVersionExists | CLASS_Hidden | CLASS_HideDropDown;
+		Filter->DisallowedClassFlags = CLASS_Deprecated | CLASS_Abstract | CLASS_NewerVersionExists;
 		Filter->AllowedChildrenOfClasses.Add(UUserWidget::StaticClass());
 
 		const FText TitleText = LOCTEXT("CreateWidgetBlueprint", "Pick Parent Class for New Widget Blueprint");
@@ -101,6 +102,7 @@ bool UWidgetBlueprintFactory::ConfigureProperties()
 			return false;
 		}
 	}
+
 	if (GetDefault<UUMGEditorProjectSettings>()->bUseWidgetTemplateSelector)
 	{
 		// Load the classviewer module to display a class picker
