UMG: Settings to display the widget picker or the favorite widge list in the factory.

diff --git a/Engine/Source/Editor/UMGEditor/Private/Settings/UMGEditorProjectSettings.cpp b/Engine/Source/Editor/UMGEditor/Private/Settings/UMGEditorProjectSettings.cpp
--- a/Engine/Source/Editor/UMGEditor/Private/Settings/UMGEditorProjectSettings.cpp
+++ b/Engine/Source/Editor/UMGEditor/Private/Settings/UMGEditorProjectSettings.cpp
@@ -21,6 +21,8 @@ UUMGEditorProjectSettings::UUMGEditorProjectSettings()
 	bGraphEditorHidden = false;
 	bHideWidgetAnimationEditor = false;
 	bUseEditorConfigPaletteFiltering = false;
+	bUseUserWidgetParentClassViewerSelector = true;
+	bUseUserWidgetParentDefaultClassViewerSelector = true;
 
 	bUseWidgetTemplateSelector = false;
 	CommonRootWidgetClasses = {
diff --git a/Engine/Source/Editor/UMGEditor/Private/WidgetBlueprintFactory.cpp b/Engine/Source/Editor/UMGEditor/Private/WidgetBlueprintFactory.cpp
--- a/Engine/Source/Editor/UMGEditor/Private/WidgetBlueprintFactory.cpp
+++ b/Engine/Source/Editor/UMGEditor/Private/WidgetBlueprintFactory.cpp
@@ -56,6 +56,7 @@ UWidgetBlueprintFactory::UWidgetBlueprintFactory(const FObjectInitializer& Objec
 
 bool UWidgetBlueprintFactory::ConfigureProperties()
 {
+	if (GetDefault<UUMGEditorProjectSettings>()->bUseUserWidgetParentClassViewerSelector || GetDefault<UUMGEditorProjectSettings>()->bUseUserWidgetParentDefaultClassViewerSelector)
 	{
 		FClassViewerModule& ClassViewerModule = FModuleManager::LoadModuleChecked<FClassViewerModule>("ClassViewer");
 
@@ -65,11 +66,13 @@ bool UWidgetBlueprintFactory::ConfigureProperties()
 		Options.Mode = EClassViewerMode::ClassPicker;
 		Options.bShowNoneOption = false;
 		Options.bExpandAllNodes = true;
+		Options.bShowDefaultClasses = GetDefault<UUMGEditorProjectSettings>()->bUseUserWidgetParentDefaultClassViewerSelector;
+		Options.bShowClassesViewer = GetDefault<UUMGEditorProjectSettings>()->bUseUserWidgetParentClassViewerSelector;
 
 		TSharedPtr<FWidgetClassFilter> Filter = MakeShareable(new FWidgetClassFilter);
 		Options.ClassFilters.Add(Filter.ToSharedRef());
 
-		const TArray<TSoftClassPtr<UUserWidget>>& FavoriteWidgetParentClasses = GetDefault <UUMGEditorProjectSettings>()->FavoriteWidgetParentClasses;
+		const TArray<TSoftClassPtr<UUserWidget>>& FavoriteWidgetParentClasses = GetDefault<UUMGEditorProjectSettings>()->FavoriteWidgetParentClasses;
 		for (int32 Index = 0; Index < FavoriteWidgetParentClasses.Num(); ++Index)
 		{
 			UClass* FavoriteWidgetParentClass = FavoriteWidgetParentClasses[Index].LoadSynchronous();
@@ -87,8 +90,7 @@ bool UWidgetBlueprintFactory::ConfigureProperties()
 			Options.ExtraPickerCommonClasses.Add(UUserWidget::StaticClass());
 		}
 
-//Filter->DisallowedClassFlags = CLASS_Deprecated | CLASS_Abstract | CLASS_NewerVersionExists | CLASS_Hidden | CLASS_HideDropDown;
-		Filter->DisallowedClassFlags = CLASS_Deprecated | CLASS_Abstract | CLASS_NewerVersionExists;
+		Filter->DisallowedClassFlags = CLASS_Deprecated | CLASS_Abstract | CLASS_NewerVersionExists | CLASS_Hidden | CLASS_HideDropDown;
 		Filter->AllowedChildrenOfClasses.Add(UUserWidget::StaticClass());
 
 		const FText TitleText = LOCTEXT("CreateWidgetBlueprint", "Pick Parent Class for New Widget Blueprint");
diff --git a/Engine/Source/Editor/UMGEditor/Public/UMGEditorProjectSettings.h b/Engine/Source/Editor/UMGEditor/Public/UMGEditorProjectSettings.h
--- a/Engine/Source/Editor/UMGEditor/Public/UMGEditorProjectSettings.h
+++ b/Engine/Source/Editor/UMGEditor/Public/UMGEditorProjectSettings.h
@@ -188,11 +188,19 @@ public:
 	UPROPERTY(EditAnywhere, config, Category = Designer)
 	bool bUseEditorConfigPaletteFiltering;
 
+	/** Enables a dialog that lets you select the parent class in a tree view. */
+	UPROPERTY(EditAnywhere, config, Category = Designer)
+	bool bUseUserWidgetParentClassViewerSelector;
+
+	/** Enables a dialog that lets you select the parent class in a default view. */
+	UPROPERTY(EditAnywhere, config, Category = Designer)
+	bool bUseUserWidgetParentDefaultClassViewerSelector;
+
 	/**
 	 * The list of parent classes to choose from for newly constructed widget blueprints.
 	 * The classes must have empty widget hierarchies.
 	 */
-	UPROPERTY(EditAnywhere, config, Category = Designer, meta = (AllowAbstract = ""))
+	UPROPERTY(EditAnywhere, config, Category = Designer, meta = (AllowAbstract = "", EditCondition="bUseUserWidgetParentDefaultClassViewerSelector"))
 	TArray<TSoftClassPtr<UUserWidget>> FavoriteWidgetParentClasses;
 
 	UPROPERTY(EditAnywhere, config, Category=Designer)
