UMG: Reuse the GetPaletteCategory helper function in the palette template.

diff --git a/Engine/Source/Editor/UMGEditor/Private/Templates/WidgetTemplateBlueprintClass.cpp b/Engine/Source/Editor/UMGEditor/Private/Templates/WidgetTemplateBlueprintClass.cpp
--- a/Engine/Source/Editor/UMGEditor/Private/Templates/WidgetTemplateBlueprintClass.cpp
+++ b/Engine/Source/Editor/UMGEditor/Private/Templates/WidgetTemplateBlueprintClass.cpp
@@ -4,8 +4,9 @@
 
 #include "Widgets/SToolTip.h"
 #include "IDocumentation.h"
+#include "UObject/CoreRedirects.h"
 #include "WidgetBlueprint.h"
-
+#include "WidgetBlueprintEditorUtils.h"
 
 #include "Kismet2/BlueprintEditorUtils.h"
 #include "Styling/SlateIconFinder.h"
@@ -29,7 +30,8 @@ FWidgetTemplateBlueprintClass::FWidgetTemplateBlueprintClass(const FAssetData& I
 		}
 		if (!ParentClassName.IsEmpty())
 		{
-			CachedParentClass = UClass::TryFindTypeSlow<UClass>(FPackageName::ExportTextPathToObjectPath(ParentClassName));
+			const FString RedirectedClassPath = FCoreRedirects::GetRedirectedName(ECoreRedirectFlags::Type_Class, FCoreRedirectObjectName(ParentClassName)).ToString();
+			CachedParentClass = UClass::TryFindTypeSlow<UClass>(FPackageName::ExportTextPathToObjectPath(RedirectedClassPath));
 		}
 	}
 }
@@ -40,28 +42,13 @@ FWidgetTemplateBlueprintClass::~FWidgetTemplateBlueprintClass()
 
 FText FWidgetTemplateBlueprintClass::GetCategory() const
 {
-	if ( WidgetClass.Get() )
+	if (WidgetClass.Get())
 	{
-		UUserWidget* DefaultUserWidget = WidgetClass->GetDefaultObject<UUserWidget>();
-		return DefaultUserWidget->GetPaletteCategory();
+		return FWidgetBlueprintEditorUtils::GetPaletteCategory(WidgetClass.Get());
 	}
 	else
 	{
-		//If the blueprint is unloaded we need to extract it from the asset metadata.
-		FText FoundPaletteCategoryText = WidgetAssetData.GetTagValueRef<FText>(GET_MEMBER_NAME_CHECKED(UWidgetBlueprint, PaletteCategory));
-		if (!FoundPaletteCategoryText.IsEmpty())
-		{
-			return FoundPaletteCategoryText;
-		}
-		else if (CachedParentClass.IsValid() && CachedParentClass->IsChildOf(UWidget::StaticClass()) && !CachedParentClass->IsChildOf(UUserWidget::StaticClass()))
-		{
-			return CachedParentClass->GetDefaultObject<UWidget>()->GetPaletteCategory();
-		}
-		else
-		{
-			auto DefaultUserWidget = UUserWidget::StaticClass()->GetDefaultObject<UUserWidget>();
-			return DefaultUserWidget->GetPaletteCategory();
-		}
+		return FWidgetBlueprintEditorUtils::GetPaletteCategory(WidgetAssetData, CachedParentClass.Get());
 	}
 }
 
diff --git a/Engine/Source/Editor/UMGEditor/Private/Templates/WidgetTemplateClass.cpp b/Engine/Source/Editor/UMGEditor/Private/Templates/WidgetTemplateClass.cpp
--- a/Engine/Source/Editor/UMGEditor/Private/Templates/WidgetTemplateClass.cpp
+++ b/Engine/Source/Editor/UMGEditor/Private/Templates/WidgetTemplateClass.cpp
@@ -8,6 +8,8 @@
 #endif // WITH_EDITOR
 #include "Widgets/SToolTip.h"
 #include "IDocumentation.h"
+#include "UObject/CoreRedirects.h"
+#include "WidgetBlueprintEditorUtils.h"
 
 #include "Blueprint/WidgetTree.h"
 #include "Styling/SlateIconFinder.h"
@@ -54,7 +56,9 @@ FWidgetTemplateClass::FWidgetTemplateClass(const FAssetData& InWidgetAssetData,
 			}
 			if (!ParentClassName.IsEmpty())
 			{
-				CachedParentClass = UClass::TryFindTypeSlow<UClass>(FPackageName::ExportTextPathToObjectPath(ParentClassName));
+				const FString RedirectedClassPath = FCoreRedirects::GetRedirectedName(ECoreRedirectFlags::Type_Class, FCoreRedirectObjectName(ParentClassName)).ToString();
+				CachedParentClass = UClass::TryFindTypeSlow<UClass>(FPackageName::ExportTextPathToObjectPath(RedirectedClassPath));
+				ensure(CachedParentClass == nullptr || CachedParentClass->IsChildOf(UWidget::StaticClass()));
 			}
 		}
 	}
@@ -69,17 +73,11 @@ FText FWidgetTemplateClass::GetCategory() const
 {
 	if (WidgetClass.Get())
 	{
-		auto DefaultWidget = WidgetClass->GetDefaultObject<UWidget>();
-		return DefaultWidget->GetPaletteCategory();
-	}
-	else if (CachedParentClass.IsValid() && CachedParentClass->IsChildOf(UWidget::StaticClass()) && !CachedParentClass->IsChildOf(UUserWidget::StaticClass()))
-	{
-		return CachedParentClass->GetDefaultObject<UWidget>()->GetPaletteCategory();
+		return FWidgetBlueprintEditorUtils::GetPaletteCategory(WidgetClass.Get());
 	}
 	else
 	{
-		auto DefaultWidget = UWidget::StaticClass()->GetDefaultObject<UWidget>();
-		return DefaultWidget->GetPaletteCategory();
+		return FWidgetBlueprintEditorUtils::GetPaletteCategory(WidgetAssetData, CachedParentClass.Get());
 	}
 }
 
diff --git a/Engine/Source/Editor/UMGEditor/Private/WidgetBlueprintEditorUtils.cpp b/Engine/Source/Editor/UMGEditor/Private/WidgetBlueprintEditorUtils.cpp
--- a/Engine/Source/Editor/UMGEditor/Private/WidgetBlueprintEditorUtils.cpp
+++ b/Engine/Source/Editor/UMGEditor/Private/WidgetBlueprintEditorUtils.cpp
@@ -2615,7 +2615,7 @@ FText FWidgetBlueprintEditorUtils::GetPaletteCategory(const TSubclassOf<UWidget>
 
 FText FWidgetBlueprintEditorUtils::GetPaletteCategory(const FAssetData& WidgetAsset, const TSubclassOf<UWidget> NativeClass)
 {
-	//The asset can be a UBlueprint, a UWidgetBlueprint or a UWidgetBlueprintGeneratedClass
+	//The asset can be a UBlueprint, UBlueprintGeneratedClass, a UWidgetBlueprint or a UWidgetBlueprintGeneratedClass
 
 	if (UClass* WidgetAssetClass = WidgetAsset.GetClass(EResolveClass::No))
 	{
