Change SReadOnlyHierarchy View to allow different behaviors:
- Disable Root Widget Selection
- Show Root as Self (better UX depending on the usage of the selected Widget

diff --git a/Engine/Source/Editor/UMGEditor/Private/Hierarchy/SReadOnlyHierarchyView.cpp b/Engine/Source/Editor/UMGEditor/Private/Hierarchy/SReadOnlyHierarchyView.cpp
--- a/Engine/Source/Editor/UMGEditor/Private/Hierarchy/SReadOnlyHierarchyView.cpp
+++ b/Engine/Source/Editor/UMGEditor/Private/Hierarchy/SReadOnlyHierarchyView.cpp
@@ -14,6 +14,7 @@ void SReadOnlyHierarchyView::Construct(const FArguments& InArgs, const UWidgetBl
 	OnSelectionChangedDelegate = InArgs._OnSelectionChanged;
 	WidgetBlueprint = InWidgetBlueprint;
 	ShowOnly = InArgs._ShowOnly;
+	RootSelectionMode = InArgs._RootSelectionMode;
 
 	SearchFilter = MakeShared<FTextFilter>(FTextFilter::FItemToStringArray::CreateSP(this, &SReadOnlyHierarchyView::GetFilterStringsForItem));
 	FilterHandler = MakeShared<FTreeFilterHandler>();
@@ -117,13 +118,26 @@ void SReadOnlyHierarchyView::SetRawFilterText(const FText& Text)
 
 FText SReadOnlyHierarchyView::GetItemText(TSharedPtr<FItem> Item) const
 {
-	if (const UWidget* Widget = Item->Widget.Get())
+	if (Item)
 	{
-		return Widget->GetLabelTextWithMetadata();
-	}
+		// The item is a Widget child in the tree
+		if (const UWidget* Widget = Item->Widget.Get())
+		{
+			return Widget->GetLabelTextWithMetadata();
+		}
+
+		// The widget is the Root WidgetBlueprint
+		if (const UWidgetBlueprint* Blueprint = Item->WidgetBlueprint.Get())
+		{
+			// If using self as the Selection mode, we use Self as the name, otherwise it's the WidgetBlueprint Name
+			FText RootWidgetName = RootSelectionMode == ERootSelectionMode::Self ? FText(LOCTEXT("SelfText", "Self")) : FText::FromString(Blueprint->GetName());
+
+			static const FText RootWidgetFormat = LOCTEXT("WidgetNameFormat", "[{0}]");
+			return FText::Format(RootWidgetFormat, RootWidgetName);
 
-	static const FText RootWidgetFormat = LOCTEXT("WidgetNameFormat", "[{0}]");
-	return FText::Format(RootWidgetFormat, FText::FromString(WidgetBlueprint->GetName()));
+		}
+	}
+	return FText::GetEmpty();
 }
 
 const FSlateBrush* SReadOnlyHierarchyView::GetIconBrush(TSharedPtr<FItem> Item) const
@@ -170,7 +184,18 @@ FText SReadOnlyHierarchyView::GetWidgetToolTipText(TSharedPtr<FItem> Item) const
 
 TSharedRef<ITableRow> SReadOnlyHierarchyView::GenerateRow(TSharedPtr<FItem> Item, const TSharedRef<STableViewBase>& OwnerTable) const
 {
+	bool bIsEnabled = true; 
+	// When Root Selection Mode is Disabled we must check if this is a WidgetBlueprint
+	if (RootSelectionMode == ERootSelectionMode::Disabled)
+	{
+		// If the Widget Blueprint is set, it's the root widget we disable it.
+		if (const UWidgetBlueprint* Blueprint = Item->WidgetBlueprint.Get())
+		{
+			bIsEnabled = false;
+		}
+	}
 	return SNew(STableRow<TSharedPtr<FItem>>, OwnerTable)
+		   .IsEnabled(bIsEnabled)
 		[
 			SNew(SHorizontalBox)
 			// Widget icon
diff --git a/Engine/Source/Editor/UMGEditor/Public/Hierarchy/SReadOnlyHierarchyView.h b/Engine/Source/Editor/UMGEditor/Public/Hierarchy/SReadOnlyHierarchyView.h
--- a/Engine/Source/Editor/UMGEditor/Public/Hierarchy/SReadOnlyHierarchyView.h
+++ b/Engine/Source/Editor/UMGEditor/Public/Hierarchy/SReadOnlyHierarchyView.h
@@ -18,6 +18,19 @@ template <typename ItemType> class STreeView;
 template <typename ItemType> class TreeFilterHandler;
 template <typename ItemType> class TTextFilter;
 
+UENUM()
+enum class ERootSelectionMode : uint8
+{
+	/** The Root Widget is not selectable */
+	Disabled,
+
+	/** The Root Widget is selectable and will show it's name as the display text */
+	Enabled,
+
+	/** The Root Widget is selectable and it will show Self as the display text */
+	Self
+};
+
 class UMGEDITOR_API SReadOnlyHierarchyView : public SCompoundWidget
 {
 private:
@@ -39,6 +52,7 @@ public:
 
 	SLATE_BEGIN_ARGS(SReadOnlyHierarchyView) {}
 		SLATE_ARGUMENT_DEFAULT(bool, ShowSearch) = true;
+		SLATE_ARGUMENT_DEFAULT(ERootSelectionMode, RootSelectionMode) = ERootSelectionMode::Enabled;
 		SLATE_ARGUMENT_DEFAULT(ESelectionMode::Type, SelectionMode) = ESelectionMode::Single;
 		SLATE_EVENT(FOnSelectionChanged, OnSelectionChanged)
 		SLATE_ARGUMENT(TArray<FName>, ShowOnly)
@@ -84,6 +98,8 @@ private:
 	TArray<TSharedPtr<FItem>> RootWidgets;
 	TArray<TSharedPtr<FItem>> FilteredRootWidgets;
 
+	ERootSelectionMode RootSelectionMode;
+
 	FOnSelectionChanged OnSelectionChangedDelegate;
 
 	TSharedPtr<SSearchBox> SearchBox;
