UMG: Add option to hide the BindingWidget, NavigationSImulation, Pallets and hierarchy windows.

diff --git a/Engine/Source/Editor/UMGEditor/Private/Animation/UMGDetailKeyframeHandler.cpp b/Engine/Source/Editor/UMGEditor/Private/Animation/UMGDetailKeyframeHandler.cpp
--- a/Engine/Source/Editor/UMGEditor/Private/Animation/UMGDetailKeyframeHandler.cpp
+++ b/Engine/Source/Editor/UMGEditor/Private/Animation/UMGDetailKeyframeHandler.cpp
@@ -13,8 +13,7 @@ FUMGDetailKeyframeHandler::FUMGDetailKeyframeHandler(TSharedPtr<FWidgetBlueprint
 
 bool FUMGDetailKeyframeHandler::IsPropertyKeyable(const UClass* InObjectClass, const IPropertyHandle& InPropertyHandle) const
 {
-	return !GetDefault<UUMGEditorProjectSettings>()->bHideWidgetAnimationEditor 
-			&& BlueprintEditor.Pin()->GetSequencer()->CanKeyProperty(FCanKeyPropertyParams(InObjectClass, InPropertyHandle));
+	return BlueprintEditor.Pin()->GetSequencer()->CanKeyProperty(FCanKeyPropertyParams(InObjectClass, InPropertyHandle));
 }
 
 bool FUMGDetailKeyframeHandler::IsPropertyKeyingEnabled() const
diff --git a/Engine/Source/Editor/UMGEditor/Private/BlueprintModes/WidgetDesignerApplicationMode.cpp b/Engine/Source/Editor/UMGEditor/Private/BlueprintModes/WidgetDesignerApplicationMode.cpp
--- a/Engine/Source/Editor/UMGEditor/Private/BlueprintModes/WidgetDesignerApplicationMode.cpp
+++ b/Engine/Source/Editor/UMGEditor/Private/BlueprintModes/WidgetDesignerApplicationMode.cpp
@@ -48,73 +48,104 @@ FWidgetDesignerApplicationMode::FWidgetDesignerApplicationMode(TSharedPtr<FWidge
 		DisplayMetrics.PrimaryDisplayWorkAreaRect.Bottom - DisplayMetrics.PrimaryDisplayWorkAreaRect.Top);
 	const FVector2D WindowSize = (CenterScale * DisplaySize) / DPIScale;
 
-	TabLayout = FTabManager::NewLayout("WidgetBlueprintEditor_Designer_Layout_v4_8")
-	->AddArea
+	TSharedPtr<FTabManager::FStack> PaletteTab;
+	if (GetDefault<UUMGEditorProjectSettings>()->bEnablePaletteWindow || GetDefault<UUMGEditorProjectSettings>()->bEnableLibraryWindow)
+	{
+		PaletteTab = FTabManager::NewStack()
+			->SetSizeCoefficient(0.5f);
+
+		if (GetDefault<UUMGEditorProjectSettings>()->bEnablePaletteWindow)
+		{
+			PaletteTab->SetForegroundTab(FPaletteTabSummoner::TabID)
+				->AddTab(FPaletteTabSummoner::TabID, ETabState::OpenedTab);
+		}
+		if (GetDefault<UUMGEditorProjectSettings>()->bEnableLibraryWindow)
+		{
+			PaletteTab->AddTab(FLibraryTabSummoner::TabID, ETabState::OpenedTab);
+		}
+	}
+
+	TSharedPtr<FTabManager::FStack> EnableHierarchyTab;
+	if (GetDefault<UUMGEditorProjectSettings>()->bEnableHierarchyWindow || GetDefault<UUMGEditorProjectSettings>()->bEnableBindWidgetWindow)
+	{
+		EnableHierarchyTab = FTabManager::NewStack()
+			->SetSizeCoefficient(0.5f);
+
+		if (GetDefault<UUMGEditorProjectSettings>()->bEnableHierarchyWindow)
+		{
+			EnableHierarchyTab->SetForegroundTab(FHierarchyTabSummoner::TabID)
+				->AddTab(FHierarchyTabSummoner::TabID, ETabState::OpenedTab);
+		}
+
+		if (GetDefault<UUMGEditorProjectSettings>()->bEnableBindWidgetWindow)
+		{
+			EnableHierarchyTab->AddTab(FBindWidgetTabSummoner::TabID, ETabState::OpenedTab);
+		}
+	}
+
+	TSharedRef<FTabManager::FArea> MainArea = FTabManager::NewPrimaryArea()
+		->SetOrientation(Orient_Horizontal);
+	if (PaletteTab || EnableHierarchyTab)
+	{
+		TSharedRef<FTabManager::FSplitter> LeftSplit = FTabManager::NewSplitter()
+			->SetSizeCoefficient(0.15f)
+			->SetOrientation(Orient_Vertical);
+		if (PaletteTab)
+		{
+			LeftSplit->Split(PaletteTab.ToSharedRef());
+		}
+		if (EnableHierarchyTab)
+		{
+			LeftSplit->Split(EnableHierarchyTab.ToSharedRef());
+		}
+		MainArea->Split(LeftSplit);
+	}
+
+	MainArea->Split
 	(
-		FTabManager::NewPrimaryArea()
-		->SetOrientation(Orient_Horizontal)
+		FTabManager::NewSplitter()
+		->SetSizeCoefficient(0.85f)
+		->SetOrientation(Orient_Vertical)
 		->Split
 		(
 			FTabManager::NewSplitter()
-			->SetSizeCoefficient( 0.15f )
-			->SetOrientation(Orient_Vertical)
+			->SetSizeCoefficient(0.7f)
+			->SetOrientation(Orient_Horizontal)
 			->Split
 			(
 				FTabManager::NewStack()
-				->SetSizeCoefficient(0.5f)
-				->SetForegroundTab(FPaletteTabSummoner::TabID)
-				->AddTab(FPaletteTabSummoner::TabID, ETabState::OpenedTab)
-				->AddTab(FLibraryTabSummoner::TabID, ETabState::OpenedTab)
+				->SetHideTabWell(true)
+				->SetSizeCoefficient(0.05f)
+				->AddTab(FDesignerTabSummoner::ToolPaletteTabID, ETabState::ClosedTab)
 			)
 			->Split
 			(
 				FTabManager::NewStack()
-				->SetSizeCoefficient(0.5f)
-				->SetForegroundTab(FHierarchyTabSummoner::TabID)
-				->AddTab(FHierarchyTabSummoner::TabID, ETabState::OpenedTab)
-				->AddTab(FBindWidgetTabSummoner::TabID, ETabState::OpenedTab)
-			)
-		)
-		->Split
-		(
-			FTabManager::NewSplitter()
-			->SetSizeCoefficient(0.85f)
-			->SetOrientation(Orient_Vertical)
-			->Split
-			(
-				FTabManager::NewSplitter()
-				->SetSizeCoefficient(0.7f)
-				->SetOrientation(Orient_Horizontal)
-				->Split
-				(
-					FTabManager::NewStack()
-					->SetHideTabWell(true)
-					->SetSizeCoefficient(0.05f)
-					->AddTab(FDesignerTabSummoner::ToolPaletteTabID, ETabState::ClosedTab)
-				)
-				->Split
-				(
-					FTabManager::NewStack()
-					->SetHideTabWell(true)
-					->SetSizeCoefficient(0.80f)
-					->AddTab( FDesignerTabSummoner::TabID, ETabState::OpenedTab )
-				)
-				->Split
-				(
-					FTabManager::NewStack()
-					->SetSizeCoefficient(0.35f)
-					->AddTab(FDetailsTabSummoner::TabID, ETabState::OpenedTab)
-				)
+				->SetHideTabWell(true)
+				->SetSizeCoefficient(0.80f)
+				->AddTab(FDesignerTabSummoner::TabID, ETabState::OpenedTab)
 			)
 			->Split
 			(
 				FTabManager::NewStack()
-				->SetSizeCoefficient( 0.3f )
-				->AddTab(FAnimationTabSummoner::TabID, ETabState::ClosedTab)
-				->AddTab(FBlueprintEditorTabs::CompilerResultsID, ETabState::ClosedTab)
-				->SetForegroundTab(FAnimationTabSummoner::TabID)
+				->SetSizeCoefficient(0.35f)
+				->AddTab(FDetailsTabSummoner::TabID, ETabState::OpenedTab)
 			)
 		)
+		->Split
+		(
+			FTabManager::NewStack()
+			->SetSizeCoefficient(0.3f)
+			->SetForegroundTab(FAnimationTabSummoner::TabID)
+			->AddTab(FAnimationTabSummoner::TabID, ETabState::ClosedTab)
+			->AddTab(FBlueprintEditorTabs::CompilerResultsID, ETabState::ClosedTab)
+		)
+	);
+
+	TabLayout = FTabManager::NewLayout("WidgetBlueprintEditor_Designer_Layout_v4_8")
+	->AddArea
+	(
+		MainArea
 	)
 	->AddArea
 	(
@@ -133,13 +164,28 @@ FWidgetDesignerApplicationMode::FWidgetDesignerApplicationMode(TSharedPtr<FWidge
 	//TabFactories.RegisterFactory(MakeShareable(new FSelectionDetailsSummoner(InWidgetEditor)));
 	TabFactories.RegisterFactory(MakeShareable(new FDetailsTabSummoner(InWidgetEditor)));
 	TabFactories.RegisterFactory(MakeShareable(new FDesignerTabSummoner(InWidgetEditor)));
-	TabFactories.RegisterFactory(MakeShareable(new FHierarchyTabSummoner(InWidgetEditor)));
-	TabFactories.RegisterFactory(MakeShareable(new FBindWidgetTabSummoner(InWidgetEditor)));
-	TabFactories.RegisterFactory(MakeShareable(new FPaletteTabSummoner(InWidgetEditor)));
-	TabFactories.RegisterFactory(MakeShareable(new FLibraryTabSummoner(InWidgetEditor)));
-	TabFactories.RegisterFactory(MakeShareable(new FAnimationTabSummoner(InWidgetEditor)));
 	TabFactories.RegisterFactory(MakeShareable(new FCompilerResultsSummoner(InWidgetEditor)));
-	TabFactories.RegisterFactory(MakeShareable(new FNavigationTabSummoner(InWidgetEditor)));
+	TabFactories.RegisterFactory(MakeShareable(new FAnimationTabSummoner(InWidgetEditor)));
+	if (GetDefault<UUMGEditorProjectSettings>()->bEnableHierarchyWindow)
+	{
+		TabFactories.RegisterFactory(MakeShareable(new FHierarchyTabSummoner(InWidgetEditor)));
+	}
+	if (GetDefault<UUMGEditorProjectSettings>()->bEnableBindWidgetWindow)
+	{
+		TabFactories.RegisterFactory(MakeShareable(new FBindWidgetTabSummoner(InWidgetEditor)));
+	}
+	if (GetDefault<UUMGEditorProjectSettings>()->bEnablePaletteWindow)
+	{
+		TabFactories.RegisterFactory(MakeShareable(new FPaletteTabSummoner(InWidgetEditor)));
+	}
+	if (GetDefault<UUMGEditorProjectSettings>()->bEnableLibraryWindow)
+	{
+		TabFactories.RegisterFactory(MakeShareable(new FLibraryTabSummoner(InWidgetEditor)));
+	}
+	if (GetDefault<UUMGEditorProjectSettings>()->bEnableNavigationSimulationWindow)
+	{
+		TabFactories.RegisterFactory(MakeShareable(new FNavigationTabSummoner(InWidgetEditor)));
+	}
 
 	IUMGEditorModule& UMGEditorModule = FModuleManager::LoadModuleChecked<IUMGEditorModule>("UMGEditor");
 	UMGEditorModule.OnRegisterTabsForEditor().Broadcast(*this, TabFactories);
@@ -200,7 +246,7 @@ void FWidgetDesignerApplicationMode::PostActivateMode()
 
 	TSharedPtr<FWidgetBlueprintEditor> BP = GetBlueprintEditor();
 
-	if (!GetDefault<UUMGEditorProjectSettings>()->bHideWidgetAnimationEditor)
+	if (GetDefault<UUMGEditorProjectSettings>()->bEnableWidgetAnimationEditor)
 	{
 		FWidgetDrawerConfig WidgetAnimSequencerDrawer(FAnimationTabSummoner::WidgetAnimSequencerDrawerID);
 		WidgetAnimSequencerDrawer.GetDrawerContentDelegate.BindSP(BP.Get(), &FWidgetBlueprintEditor::OnGetWidgetAnimSequencer);
diff --git a/Engine/Source/Editor/UMGEditor/Private/Details/SWidgetDetailsView.cpp b/Engine/Source/Editor/UMGEditor/Private/Details/SWidgetDetailsView.cpp
--- a/Engine/Source/Editor/UMGEditor/Private/Details/SWidgetDetailsView.cpp
+++ b/Engine/Source/Editor/UMGEditor/Private/Details/SWidgetDetailsView.cpp
@@ -30,6 +30,7 @@
 #include "WidgetBlueprintEditorUtils.h"
 #include "ScopedTransaction.h"
 #include "Styling/SlateIconFinder.h"
+#include "UMGEditorModule.h"
 #include "UMGEditorProjectSettings.h"
 
 #define LOCTEXT_NAMESPACE "UMG"
@@ -49,8 +50,11 @@ void SWidgetDetailsView::Construct(const FArguments& InArgs, TSharedPtr<FWidgetB
 	PropertyView = EditModule.CreateDetailView(DetailsViewArgs);
 
 	// Create a handler for keyframing via the details panel
-	TSharedRef<IDetailKeyframeHandler> KeyframeHandler = MakeShareable( new FUMGDetailKeyframeHandler( InBlueprintEditor ) );
-	PropertyView->SetKeyframeHandler( KeyframeHandler );
+	if (GetDefault<UUMGEditorProjectSettings>()->bEnableWidgetAnimationEditor)
+	{
+		TSharedRef<IDetailKeyframeHandler> KeyframeHandler = MakeShareable(new FUMGDetailKeyframeHandler(InBlueprintEditor));
+		PropertyView->SetKeyframeHandler(KeyframeHandler);
+	}
 
 	// Create a handler for property binding via the details panel
 	TSharedRef<FDetailWidgetExtensionHandler> BindingHandler = MakeShareable( new FDetailWidgetExtensionHandler( InBlueprintEditor ) );
@@ -66,8 +70,7 @@ void SWidgetDetailsView::Construct(const FArguments& InArgs, TSharedPtr<FWidgetB
 		SNew(SCheckBox)
 		.IsChecked(this, &SWidgetDetailsView::GetIsVariable)
 		.OnCheckStateChanged(this, &SWidgetDetailsView::HandleIsVariableChanged)
-		.Padding(FMargin(3, 1, 18, 1))
-		.Visibility(GetDefault<UUMGEditorProjectSettings>()->bGraphEditorHidden ? EVisibility::Hidden : EVisibility::Visible)
+		.Padding(FMargin(3.0f, 1.0f, 18.0f, 1.0f))
 		[
 			SNew(STextBlock)
 			.Text(LOCTEXT("IsVariable", "Is Variable"))
@@ -79,7 +82,7 @@ void SWidgetDetailsView::Construct(const FArguments& InArgs, TSharedPtr<FWidgetB
 
 		+ SVerticalBox::Slot()
 		.AutoHeight()
-		.Padding(0, 2)
+		.Padding(0.0f, 2.0f)
 		[
 			SNew(SBorder)
 			.BorderImage(FAppStyle::GetBrush(TEXT("ToolPanel.GroupBorder")))
@@ -89,14 +92,14 @@ void SWidgetDetailsView::Construct(const FArguments& InArgs, TSharedPtr<FWidgetB
 
 				+ SVerticalBox::Slot()
 				.AutoHeight()
-				.Padding(4)
+				.Padding(4.0f)
 				[
 					SNew(SHorizontalBox)
 					.Visibility(this, &SWidgetDetailsView::GetCategoryAreaVisibility)
 
 					+ SHorizontalBox::Slot()
 					.AutoWidth()
-					.Padding(0, 0, 3, 0)
+					.Padding(0.0f, 0.0f, 3.0f, 0.0f)
 					.VAlign(VAlign_Center)
 					[
 						SNew(SImage)
@@ -105,7 +108,7 @@ void SWidgetDetailsView::Construct(const FArguments& InArgs, TSharedPtr<FWidgetB
 
 					+ SHorizontalBox::Slot()
 					.AutoWidth()
-					.Padding(0, 0, 6, 0)
+					.Padding(0.0f, 0.0f, 6.0f, 0.0f)
 					[
 						SNew(SBox)
 						.WidthOverride(200.0f)
@@ -123,14 +126,14 @@ void SWidgetDetailsView::Construct(const FArguments& InArgs, TSharedPtr<FWidgetB
 
 				+ SVerticalBox::Slot()
 				.AutoHeight()
-				.Padding(4)
+				.Padding(4.0f)
 				[
 					SNew(SHorizontalBox)
 					.Visibility(this, &SWidgetDetailsView::GetNameAreaVisibility)
 
 					+ SHorizontalBox::Slot()
 					.AutoWidth()
-					.Padding(0, 0, 3, 0)
+					.Padding(0.0f, 0.0f, 3.0f, 0.0f)
 					.VAlign(VAlign_Center)
 					[
 						SNew(SImage)
@@ -140,7 +143,7 @@ void SWidgetDetailsView::Construct(const FArguments& InArgs, TSharedPtr<FWidgetB
 
 					+ SHorizontalBox::Slot()
 					.AutoWidth()
-					.Padding(0, 0, 6, 0)
+					.Padding(0.0f, 0.0f, 6.0f, 0.0f)
 					[
 						SNew(SBox)
 						.WidthOverride(200.0f)
@@ -159,7 +162,7 @@ void SWidgetDetailsView::Construct(const FArguments& InArgs, TSharedPtr<FWidgetB
 					+ SHorizontalBox::Slot()
 					.AutoWidth()
 					[
-						GetDefault<UUMGEditorProjectSettings>()->bGraphEditorHidden ? SNullWidget::NullWidget : IsVariableCheckbox
+						GetDefault<UUMGEditorProjectSettings>()->bEnableMakeVariable ? IsVariableCheckbox : SNullWidget::NullWidget
 					]
 
 					+ SHorizontalBox::Slot()
diff --git a/Engine/Source/Editor/UMGEditor/Private/Settings/UMGEditorProjectSettings.cpp b/Engine/Source/Editor/UMGEditor/Private/Settings/UMGEditorProjectSettings.cpp
--- a/Engine/Source/Editor/UMGEditor/Private/Settings/UMGEditorProjectSettings.cpp
+++ b/Engine/Source/Editor/UMGEditor/Private/Settings/UMGEditorProjectSettings.cpp
@@ -18,8 +18,15 @@ UUMGEditorProjectSettings::UUMGEditorProjectSettings()
 	CurrentVersion = 1;
 	bShowWidgetsFromEngineContent = false;
 	bShowWidgetsFromDeveloperContent = true;
-	bGraphEditorHidden = false;
-	bHideWidgetAnimationEditor = false;
+
+	bEnableMakeVariable = true;
+	bEnableWidgetAnimationEditor = true;
+	bEnablePaletteWindow = true;
+	bEnableLibraryWindow = true;
+	bEnableHierarchyWindow = true;
+	bEnableBindWidgetWindow = true;
+	bEnableNavigationSimulationWindow = true;
+
 	bUseEditorConfigPaletteFiltering = false;
 	bUseUserWidgetParentClassViewerSelector = true;
 	bUseUserWidgetParentDefaultClassViewerSelector = true;
diff --git a/Engine/Source/Editor/UMGEditor/Public/UMGEditorProjectSettings.h b/Engine/Source/Editor/UMGEditor/Public/UMGEditorProjectSettings.h
--- a/Engine/Source/Editor/UMGEditor/Public/UMGEditorProjectSettings.h
+++ b/Engine/Source/Editor/UMGEditor/Public/UMGEditorProjectSettings.h
@@ -176,15 +176,35 @@ public:
 	UPROPERTY(EditAnywhere, config, Category = Designer)
 	TSubclassOf<UPanelWidget> DefaultRootWidget;
 
-	/** Set true to hide UMG-graph related elements when the graph editor is hidden  */
-	UPROPERTY(EditAnywhere, Category = Designer)
-	bool bGraphEditorHidden;
-
-	/** Set true to hide widget animation related elements in the UMG editor  */
-	UPROPERTY(EditAnywhere, Category = Designer)
-	bool bHideWidgetAnimationEditor;
-
-	/** Set true to filter all categories and widgets out in the palette, selectively enabling them later via permission lists  */
+	/** Set true to enable the Is Variable checkbox in the UMG editor DetailView. */
+	UPROPERTY(EditAnywhere, Category = "Designer | Window")
+	bool bEnableMakeVariable;
+
+	/** Set true to hide widget animation related elements in the UMG editor. */
+	UPROPERTY(EditAnywhere, Category = "Designer | Window")
+	bool bEnableWidgetAnimationEditor;
+
+	/** Set true to enabled the Palette window in the UMG editor. */
+	UPROPERTY(EditAnywhere, Category = "Designer | Window")
+	bool bEnablePaletteWindow;
+
+	/** Set true to enabled the LIbrary window in the UMG editor. */
+	UPROPERTY(EditAnywhere, Category = "Designer | Window")
+	bool bEnableLibraryWindow;
+
+	/** Set true to enabled the Widget Hierarchy window in the UMG editor. */
+	UPROPERTY(EditAnywhere, Category = "Designer | Window")
+	bool bEnableHierarchyWindow;
+	
+	/** Set true to enabled the Bind Widget window in the UMG editor. */
+	UPROPERTY(EditAnywhere, Category = "Designer | Window")
+	bool bEnableBindWidgetWindow;
+
+	/** Set true to enabled the Navigation Simulation window in the UMG editor. */
+	UPROPERTY(EditAnywhere, Category = "Designer | Window")
+	bool bEnableNavigationSimulationWindow;
+
+	/** Set true to filter all categories and widgets out in the palette, selectively enabling them later via permission lists. */
 	UPROPERTY(EditAnywhere, config, Category = Designer)
 	bool bUseEditorConfigPaletteFiltering;
 
