UMG: Add extension to add Customization to the UMG detail view from an extension.

diff --git a/Engine/Source/Editor/UMGEditor/Private/Details/SWidgetDetailsView.cpp b/Engine/Source/Editor/UMGEditor/Private/Details/SWidgetDetailsView.cpp
--- a/Engine/Source/Editor/UMGEditor/Private/Details/SWidgetDetailsView.cpp
+++ b/Engine/Source/Editor/UMGEditor/Private/Details/SWidgetDetailsView.cpp
@@ -84,9 +84,8 @@ void SWidgetDetailsView::Construct(const FArguments& InArgs, TSharedPtr<FWidgetB
 		.AutoHeight()
 		.Padding(0.0f, 2.0f)
 		[
-			SNew(SBorder)
+			SAssignNew(BorderArea, SBorder)
 			.BorderImage(FAppStyle::GetBrush(TEXT("ToolPanel.GroupBorder")))
-			.Visibility(this, &SWidgetDetailsView::GetBorderAreaVisibility)
 			[
 				SNew(SVerticalBox)
 
@@ -213,16 +212,32 @@ SWidgetDetailsView::~SWidgetDetailsView()
 
 void SWidgetDetailsView::RegisterCustomizations()
 {
-	PropertyView->RegisterInstancedCustomPropertyLayout(UWidget::StaticClass(), FOnGetDetailCustomizationInstance::CreateStatic(&FBlueprintWidgetCustomization::MakeInstance, BlueprintEditor.Pin().ToSharedRef(), BlueprintEditor.Pin()->GetBlueprintObj()));
-	PropertyView->RegisterInstancedCustomPropertyTypeLayout(TEXT("Widget"), FOnGetPropertyTypeCustomizationInstance::CreateStatic(&FWidgetTypeCustomization::MakeInstance, BlueprintEditor.Pin().ToSharedRef()), nullptr);
-	PropertyView->RegisterInstancedCustomPropertyTypeLayout(TEXT("WidgetNavigation"), FOnGetPropertyTypeCustomizationInstance::CreateStatic(&FWidgetNavigationCustomization::MakeInstance, BlueprintEditor.Pin().ToSharedRef()));
-	PropertyView->RegisterInstancedCustomPropertyTypeLayout(TEXT("PanelSlot"), FOnGetPropertyTypeCustomizationInstance::CreateStatic(&FCanvasSlotCustomization::MakeInstance, BlueprintEditor.Pin()->GetBlueprintObj()));
+	check(BlueprintEditor.Pin());
+	TSharedRef<FWidgetBlueprintEditor> BlueprintEditorRef = BlueprintEditor.Pin().ToSharedRef();
+
+	PropertyView->RegisterInstancedCustomPropertyLayout(UWidget::StaticClass(), FOnGetDetailCustomizationInstance::CreateStatic(&FBlueprintWidgetCustomization::MakeInstance, BlueprintEditorRef, BlueprintEditorRef->GetBlueprintObj()));
+	PropertyView->RegisterInstancedCustomPropertyTypeLayout(TEXT("Widget"), FOnGetPropertyTypeCustomizationInstance::CreateStatic(&FWidgetTypeCustomization::MakeInstance, BlueprintEditorRef), nullptr);
+	PropertyView->RegisterInstancedCustomPropertyTypeLayout(TEXT("WidgetNavigation"), FOnGetPropertyTypeCustomizationInstance::CreateStatic(&FWidgetNavigationCustomization::MakeInstance, BlueprintEditorRef));
+	PropertyView->RegisterInstancedCustomPropertyTypeLayout(TEXT("PanelSlot"), FOnGetPropertyTypeCustomizationInstance::CreateStatic(&FCanvasSlotCustomization::MakeInstance, BlueprintEditorRef->GetBlueprintObj()));
 	PropertyView->RegisterInstancedCustomPropertyTypeLayout(TEXT("EHorizontalAlignment"), FOnGetPropertyTypeCustomizationInstance::CreateStatic(&FHorizontalAlignmentCustomization::MakeInstance));
 	PropertyView->RegisterInstancedCustomPropertyTypeLayout(TEXT("EVerticalAlignment"), FOnGetPropertyTypeCustomizationInstance::CreateStatic(&FVerticalAlignmentCustomization::MakeInstance));
 	PropertyView->RegisterInstancedCustomPropertyTypeLayout(TEXT("SlateChildSize"), FOnGetPropertyTypeCustomizationInstance::CreateStatic(&FSlateChildSizeCustomization::MakeInstance));
 	PropertyView->RegisterInstancedCustomPropertyTypeLayout(TEXT("SlateBrush"), FOnGetPropertyTypeCustomizationInstance::CreateStatic(&FSlateBrushStructCustomization::MakeInstance, false));
 	PropertyView->RegisterInstancedCustomPropertyTypeLayout(TEXT("SlateFontInfo"), FOnGetPropertyTypeCustomizationInstance::CreateStatic(&FSlateFontInfoStructCustomization::MakeInstance));
 	PropertyView->RegisterInstancedCustomPropertyTypeLayout(TEXT("ETextJustify"), FOnGetPropertyTypeCustomizationInstance::CreateStatic(&FTextJustifyCustomization::MakeInstance));
+
+	TWeakPtr<FWidgetBlueprintEditor> WeakBlueprintEditor = BlueprintEditor;
+	IUMGEditorModule& UMGEditorModule = FModuleManager::LoadModuleChecked<IUMGEditorModule>("UMGEditor");
+	for (const IUMGEditorModule::FCustomPropertyTypeLayout& Layout : UMGEditorModule.GetAllInstancedCustomPropertyTypeLayout())
+	{
+		if (Layout.Type.GetAssetName().IsValid() && Layout.Delegate.IsBound())
+		{
+			PropertyView->RegisterInstancedCustomPropertyTypeLayout(Layout.Type.GetAssetName(), FOnGetPropertyTypeCustomizationInstance::CreateLambda([LocalDelegate = Layout.Delegate, WeakBlueprintEditor]
+				{
+					return LocalDelegate.Execute(WeakBlueprintEditor);
+				}));
+		}
+	}
 }
 
 void SWidgetDetailsView::OnEditorSelectionChanging()
@@ -241,28 +256,31 @@ void SWidgetDetailsView::OnEditorSelectionChanged()
 	SelectedObjects.Empty();
 	PropertyView->SetObjects(SelectedObjects);
 
+	TOptional<bool> bIsWidgetSelection;
+
 	// Add any selected widgets to the list of pending selected objects.
-	TSet< FWidgetReference > SelectedWidgets = BlueprintEditor.Pin()->GetSelectedWidgets();
-	if ( SelectedWidgets.Num() > 0 )
+	const TSet<FWidgetReference>& SelectedWidgets = BlueprintEditor.Pin()->GetSelectedWidgets();
+	for ( const FWidgetReference& WidgetRef : SelectedWidgets )
 	{
-		for ( FWidgetReference& WidgetRef : SelectedWidgets )
-		{
-			// Edit actions will go directly to the preview widget, changes will be
-			// propagated to the template via SWidgetDetailsView::NotifyPostChange
-			SelectedObjects.Add(WidgetRef.GetPreview());
-		}
+		// Edit actions will go directly to the preview widget, changes will be
+		// propagated to the template via SWidgetDetailsView::NotifyPostChange
+		SelectedObjects.Add(WidgetRef.GetPreview());
+		bIsWidgetSelection = true;
 	}
 
 	// Add any selected objects (non-widgets) to the pending selected objects.
-	TSet< TWeakObjectPtr<UObject> > Selection = BlueprintEditor.Pin()->GetSelectedObjects();
-	for ( TWeakObjectPtr<UObject> Selected : Selection )
+	const TSet<TWeakObjectPtr<UObject>>& Selection = BlueprintEditor.Pin()->GetSelectedObjects();
+	for ( const TWeakObjectPtr<UObject> Selected : Selection )
 	{
 		if ( UObject* S = Selected.Get() )
 		{
 			SelectedObjects.Add(S);
+			bIsWidgetSelection = bIsWidgetSelection.Get(true) && Cast<UWidget>(S) != nullptr;
 		}
 	}
 
+	BorderArea->SetVisibility(bIsWidgetSelection.Get(false) ? EVisibility::Visible : EVisibility::Collapsed);
+
 	// If only 1 valid selected object exists, update the class link to point to the right class.
 	if ( SelectedObjects.Num() == 1 && SelectedObjects[0].IsValid() )
 	{
@@ -348,20 +366,11 @@ bool SWidgetDetailsView::IsWidgetCDOSelected() const
 		{
 			return true;
 		}
-		else if (Widget && !Widget->HasAnyFlags(RF_ClassDefaultObject))
-		{
-			return false;
-		}
 	}
 
 	return false;
 }
 
-EVisibility SWidgetDetailsView::GetBorderAreaVisibility() const
-{
-	return (SelectedObjects.Num() == 0) ? EVisibility::Collapsed : EVisibility::Visible;
-}
-
 EVisibility SWidgetDetailsView::GetNameAreaVisibility() const
 {
 	return IsWidgetCDOSelected() ? EVisibility::Collapsed : EVisibility::Visible;
@@ -378,6 +387,8 @@ void SWidgetDetailsView::HandleCategoryTextCommitted(const FText& Text, ETextCom
 	{
 		if ( UUserWidget* Widget = Cast<UUserWidget>(SelectedObjects[0].Get()) )
 		{
+			ensureMsgf(IsWidgetCDOSelected(), TEXT("You can ony change the category if you are the preview widget."));
+
 			UUserWidget* WidgetCDO = Widget->GetClass()->GetDefaultObject<UUserWidget>();
 			WidgetCDO->PaletteCategory = Text;
 
@@ -387,6 +398,17 @@ void SWidgetDetailsView::HandleCategoryTextCommitted(const FText& Text, ETextCom
 
 			// Immediately force a rebuild so that all palettes update to show it in a new category.
 			FBlueprintEditorUtils::MarkBlueprintAsStructurallyModified(Blueprint);
+
+			// MarkBlueprintAsStructurallyModified will invalidate the selection. Reselect it.
+			if (TSharedPtr<FWidgetBlueprintEditor> Editor = BlueprintEditor.Pin())
+			{
+				if (Editor->GetPreview())
+				{
+					TSet<UObject*> NewSelectedObjects;
+					NewSelectedObjects.Add(Editor->GetPreview());
+					Editor->SelectObjects(NewSelectedObjects);
+				}
+			}
 		}
 	}
 }
diff --git a/Engine/Source/Editor/UMGEditor/Private/Details/SWidgetDetailsView.h b/Engine/Source/Editor/UMGEditor/Private/Details/SWidgetDetailsView.h
--- a/Engine/Source/Editor/UMGEditor/Private/Details/SWidgetDetailsView.h
+++ b/Engine/Source/Editor/UMGEditor/Private/Details/SWidgetDetailsView.h
@@ -52,8 +52,6 @@ private:
 	
 	bool IsWidgetCDOSelected() const;
 
-	EVisibility GetBorderAreaVisibility() const;
-
 	EVisibility GetNameAreaVisibility() const;
 
 	const FSlateBrush* GetNameIcon() const;
@@ -79,6 +77,9 @@ private:
 	/** The editor that owns this details view */
 	TWeakPtr<class FWidgetBlueprintEditor> BlueprintEditor;
 
+	/** The container widget for the class link users can click to open another asset */
+	TSharedPtr<SWidget> BorderArea;
+
 	/** The name text box that users can use to rename their widgets */
 	TSharedPtr<SEditableTextBox> NameTextBox;
 
diff --git a/Engine/Source/Editor/UMGEditor/Private/UMGEditorModule.cpp b/Engine/Source/Editor/UMGEditor/Private/UMGEditorModule.cpp
--- a/Engine/Source/Editor/UMGEditor/Private/UMGEditorModule.cpp
+++ b/Engine/Source/Editor/UMGEditor/Private/UMGEditorModule.cpp
@@ -297,6 +297,32 @@ public:
 		return RegisterLayoutExtensions; 
 	}
 
+	virtual void RegisterInstancedCustomPropertyTypeLayout(FTopLevelAssetPath Type, FOnGetInstancePropertyTypeCustomizationInstance Delegate) override
+	{
+		int32 IndexOf = CustomPropertyTypeLayout.IndexOfByPredicate([Type](const FCustomPropertyTypeLayout& Element) { return Element.Type == Type; });
+		if (ensure(!CustomPropertyTypeLayout.IsValidIndex(IndexOf)))
+		{
+			FCustomPropertyTypeLayout Layout;
+			Layout.Type = Type;
+			Layout.Delegate = MoveTemp(Delegate);
+			CustomPropertyTypeLayout.Add(MoveTemp(Layout));
+		}
+	}
+
+	virtual void UnregisterInstancedCustomPropertyTypeLayout(FTopLevelAssetPath Type) override
+	{
+		int32 IndexOf = CustomPropertyTypeLayout.IndexOfByPredicate([Type](const FCustomPropertyTypeLayout& Element){ return Element.Type == Type; });
+		if (CustomPropertyTypeLayout.IsValidIndex(IndexOf))
+		{
+			CustomPropertyTypeLayout.RemoveAtSwap(IndexOf);
+		}
+	}
+
+	virtual TArrayView<const FCustomPropertyTypeLayout> GetAllInstancedCustomPropertyTypeLayout() const override
+	{
+		return CustomPropertyTypeLayout;
+	}
+
 private:
 	void RegisterAssetTypeAction(IAssetTools& AssetTools, TSharedRef<IAssetTypeActions> Action)
 	{
@@ -381,6 +407,9 @@ private:
 	/** Support layout extensions */
 	FOnRegisterLayoutExtensions	RegisterLayoutExtensions;
 
+	/** */
+	TArray<FCustomPropertyTypeLayout> CustomPropertyTypeLayout;
+
 	/** Handle for the FModuleManager::OnModulesChanged */
 	FDelegateHandle ModuleChangedHandle;
 	/** Handle for FBlueprintEditorModule::RegisterVariableCustomization */
diff --git a/Engine/Source/Editor/UMGEditor/Public/UMGEditorModule.h b/Engine/Source/Editor/UMGEditor/Public/UMGEditorModule.h
--- a/Engine/Source/Editor/UMGEditor/Public/UMGEditorModule.h
+++ b/Engine/Source/Editor/UMGEditor/Public/UMGEditorModule.h
@@ -6,13 +6,15 @@
 #include "Toolkits/AssetEditorToolkit.h"
 #include "IHasDesignerExtensibility.h"
 #include "IHasPropertyBindingExtensibility.h"
+#include "UObject/TopLevelAssetPath.h"
 
 extern const FName UMGEditorAppIdentifier;
 
-class FUMGEditor;
 class FWidgetBlueprintApplicationMode;
+class FWidgetBlueprintEditor;
 class FWorkflowAllowedTabSet;
 class IBlueprintWidgetCustomizationExtender;
+class IPropertyTypeCustomization;
 
 /** The public interface of the UMG editor module. */
 class IUMGEditorModule : 
@@ -46,6 +48,22 @@ public:
 	virtual TArrayView<TSharedRef<IBlueprintWidgetCustomizationExtender>> GetAllWidgetCustomizationExtenders() = 0;
 
 	/** Support for general layout extenders */
-	DECLARE_EVENT_OneParam(IStaticMeshEditor, FOnRegisterLayoutExtensions, FLayoutExtender&);
+	DECLARE_EVENT_OneParam(IUMGEditorModule, FOnRegisterLayoutExtensions, FLayoutExtender&);
 	virtual FOnRegisterLayoutExtensions& OnRegisterLayoutExtensions() = 0;
+
+	DECLARE_DELEGATE_RetVal_OneParam(TSharedRef<IPropertyTypeCustomization>, FOnGetInstancePropertyTypeCustomizationInstance, TWeakPtr<FWidgetBlueprintEditor> Editor);
+	/** Add an instance customization to the widget property view. */
+	virtual void RegisterInstancedCustomPropertyTypeLayout(FTopLevelAssetPath Type, FOnGetInstancePropertyTypeCustomizationInstance) = 0;
+
+	/** Remove an instance customization to the widget property view. */
+	virtual void UnregisterInstancedCustomPropertyTypeLayout(FTopLevelAssetPath Type) = 0;
+
+	struct FCustomPropertyTypeLayout
+	{
+		FTopLevelAssetPath Type;
+		FOnGetInstancePropertyTypeCustomizationInstance Delegate;
+	};
+
+	/** Remove an instance customization to the widget property view. */
+	virtual TArrayView<const FCustomPropertyTypeLayout> GetAllInstancedCustomPropertyTypeLayout() const = 0;
 };
