Updated UHT to throw errors if the underlying type isn't set or if the generated header isn't included for Epic code.

diff --git a/Engine/Source/Editor/UMGEditor/Public/Hierarchy/SReadOnlyHierarchyView.h b/Engine/Source/Editor/UMGEditor/Public/Hierarchy/SReadOnlyHierarchyView.h
--- a/Engine/Source/Editor/UMGEditor/Public/Hierarchy/SReadOnlyHierarchyView.h
+++ b/Engine/Source/Editor/UMGEditor/Public/Hierarchy/SReadOnlyHierarchyView.h
@@ -10,6 +10,8 @@
 #include "WidgetBlueprint.h"
 #include "Widgets/SCompoundWidget.h"
 
+#include "SReadOnlyHierarchyView.generated.h"
+
 class ITableRow;
 class SSearchBox;
 class STableViewBase;
