Add WidgetChild type that can be used in UserWidgets to be able to select the Root Widget or direct Children.
Add the detail panel customization to edit WidgetChild

diff --git a/Engine/Source/Editor/UMGEditor/Private/Customizations/WidgetChildTypeCustomization.cpp b/Engine/Source/Editor/UMGEditor/Private/Customizations/WidgetChildTypeCustomization.cpp
new file mode 100644
--- /dev/null
+++ b/Engine/Source/Editor/UMGEditor/Private/Customizations/WidgetChildTypeCustomization.cpp
@@ -0,0 +1,188 @@
+// Copyright Epic Games, Inc. All Rights Reserved.
+
+#include "Customizations/WidgetChildTypeCustomization.h"
+#include "Widgets/Text/STextBlock.h"
+#include "Framework/MultiBox/MultiBoxBuilder.h"
+#include "Widgets/Input/SEditableTextBox.h"
+#include "Widgets/Input/SComboButton.h"
+#include "Widgets/Input/SComboBox.h"
+#include "Hierarchy/SReadOnlyHierarchyView.h"
+
+#include "Blueprint/WidgetNavigation.h"
+#include "Blueprint/WidgetChild.h"
+#include "Blueprint/WidgetTree.h"
+
+#include "WidgetBlueprint.h"
+
+#include "IDetailChildrenBuilder.h"
+#include "DetailWidgetRow.h"
+#include "DetailLayoutBuilder.h"
+#include "ScopedTransaction.h"
+
+
+#define LOCTEXT_NAMESPACE "UMG"
+
+// FWidgetChildTypeCustomization
+////////////////////////////////////////////////////////////////////////////////
+
+void FWidgetChildTypeCustomization::CustomizeHeader(TSharedRef<IPropertyHandle> PropertyHandle, FDetailWidgetRow& HeaderRow, IPropertyTypeCustomizationUtils& CustomizationUtils)
+{
+	PropertyHandlePtr = PropertyHandle;
+
+	HeaderRow
+	.NameContent()
+	[
+		PropertyHandle->CreatePropertyNameWidget()
+	]
+	.ValueContent()
+	.MaxDesiredWidth(FDetailWidgetRow::DefaultValueMaxWidth * 2)
+	[
+		SAssignNew(WidgetListComboButton, SComboButton)
+		.ButtonStyle(FAppStyle::Get(), "PropertyEditor.AssetComboStyle")
+		.ForegroundColor(FAppStyle::GetColor("PropertyEditor.AssetName.ColorAndOpacity"))
+		.OnGetMenuContent(this, &FWidgetChildTypeCustomization::GetPopupContent)
+		.ContentPadding(2.0f)
+		.ButtonContent()
+		[
+			SNew(STextBlock)
+			.Text(this, &FWidgetChildTypeCustomization::GetCurrentValueText)
+			.Font(IDetailLayoutBuilder::GetDetailFont())
+		]
+	];
+}
+
+void FWidgetChildTypeCustomization::SetDesiredFocusWidgetChild(UUserWidget* OwnerUserWidget, const FWidgetChild& WidgetChild)
+{
+	ensure(OwnerUserWidget);
+
+	// We Get the value from the Preview and add it to the CDO.
+	if (TSharedPtr<IPropertyHandle> PropertyHandlePinnedPtr = PropertyHandlePtr.Pin())
+	{
+		if (FStructProperty* StructProperty = CastField<FStructProperty>(PropertyHandlePinnedPtr->GetProperty()))
+		{
+			TArray<void*> RawData;
+			PropertyHandlePinnedPtr->AccessRawData(RawData);
+			FWidgetChild* PreviousWidgetChild = reinterpret_cast<FWidgetChild*>(RawData[0]);
+
+			OwnerUserWidget->Modify();
+
+			FString TextValue;
+			StructProperty->Struct->ExportText(TextValue, &WidgetChild, PreviousWidgetChild, OwnerUserWidget, EPropertyPortFlags::PPF_None, nullptr);
+			ensure(PropertyHandlePinnedPtr->SetValueFromFormattedString(TextValue, EPropertyValueSetFlags::DefaultFlags) == FPropertyAccess::Result::Success);
+		}
+	}
+}
+
+void FWidgetChildTypeCustomization::OnWidgetSelectionChanged(FName SelectedName, ESelectInfo::Type SelectionType)
+{
+	if (TSharedPtr<IPropertyHandle> PropertyHandlePinnedPtr = PropertyHandlePtr.Pin())
+	{
+		TArray<UObject*> OuterObjects;
+		PropertyHandlePinnedPtr->GetOuterObjects(OuterObjects);
+
+		for (UObject* OuterObject : OuterObjects)
+		{
+			if (UUserWidget* UserWidget = Cast<UUserWidget>(OuterObject))
+			{
+				const FScopedTransaction Transaction(LOCTEXT("SetDesiredFocus", "Set Desired Focus"));
+				if (UWidgetBlueprint* WidgetBlueprint = Editor.Pin()->GetWidgetBlueprintObj())
+				{
+					if (SelectedName == WidgetBlueprint->GetFName())
+					{
+						SetDesiredFocusWidgetChild(UserWidget, FWidgetChild());
+						break;
+					}
+				}
+				SetDesiredFocusWidgetChild(UserWidget, FWidgetChild(UserWidget, SelectedName));
+				break;
+			}
+		}
+	}
+	WidgetListComboButton->SetIsOpen(false);
+}
+
+TSharedRef<SWidget> FWidgetChildTypeCustomization::GetPopupContent()
+{
+	TSharedPtr<FWidgetBlueprintEditor> BlueprintEditor = Editor.Pin();
+	UWidgetBlueprint* WidgetBlueprint = BlueprintEditor->GetWidgetBlueprintObj();
+
+	const float MinPopupWidth = 250.0f;
+	const float MinPopupHeight = 200.0f;
+
+	return SNew(SBorder)
+		.BorderImage(FAppStyle::Get().GetBrush("Brushes.Recessed"))
+		.Padding(4)
+		[
+			SNew(SBox)
+			.MinDesiredWidth(MinPopupWidth)
+			.MinDesiredHeight(MinPopupHeight)
+			[
+				SNew(SReadOnlyHierarchyView, WidgetBlueprint)
+				.OnSelectionChanged(this, &FWidgetChildTypeCustomization::OnWidgetSelectionChanged)
+				.ShowSearch(true)
+				.RootSelectionMode(ERootSelectionMode::Self)
+			]
+		];
+}
+
+void FWidgetChildTypeCustomization::CustomizeChildren(TSharedRef<IPropertyHandle> PropertyHandle, IDetailChildrenBuilder& ChildBuilder, IPropertyTypeCustomizationUtils& CustomizationUtils)
+{
+
+}
+
+UWidget* FWidgetChildTypeCustomization::GetCurrentValue() const
+{
+	if (TSharedPtr<IPropertyHandle> PropertyHandle = PropertyHandlePtr.Pin())
+	{
+		UObject* Object;
+		switch (PropertyHandle->GetValue(Object))
+		{
+		case FPropertyAccess::Success:
+			if (UWidget* Widget = Cast<UWidget>(Object))
+			{
+				return Widget;
+			}
+			break;
+		}
+	}
+
+	return nullptr;
+}
+
+FText FWidgetChildTypeCustomization::GetCurrentValueText() const
+{
+	if (TSharedPtr<IPropertyHandle> PropertyHandle = PropertyHandlePtr.Pin())
+	{
+		UObject* Object;
+		switch (PropertyHandle->GetValue(Object))
+		{
+		case FPropertyAccess::MultipleValues:
+			return LOCTEXT("MultipleValues", "Multiple Values");
+		case FPropertyAccess::Success:
+			if (UWidget* Widget = Cast<UWidget>(Object))
+			{
+				return Widget->GetLabelText();
+			}
+			return FText::GetEmpty();
+		case FPropertyAccess::Fail:
+			break;
+
+		}
+		
+		if (FStructProperty* StructProperty = CastField<FStructProperty>(PropertyHandle->GetProperty()))
+		{
+			TArray<void*> RawData;
+			PropertyHandle->AccessRawData(RawData);
+			if (FWidgetChild* ChildWidget = reinterpret_cast<FWidgetChild*>(RawData[0]))
+			{
+				FName WidgetName = ChildWidget->GetChildName();
+				return WidgetName.IsNone()? LOCTEXT("SelfText", "Self") : FText::FromName(WidgetName);
+			}
+		}		
+	}
+
+	return FText::GetEmpty();
+}
+
+#undef LOCTEXT_NAMESPACE
+
diff --git a/Engine/Source/Editor/UMGEditor/Private/Customizations/WidgetChildTypeCustomization.h b/Engine/Source/Editor/UMGEditor/Private/Customizations/WidgetChildTypeCustomization.h
new file mode 100644
--- /dev/null
+++ b/Engine/Source/Editor/UMGEditor/Private/Customizations/WidgetChildTypeCustomization.h
@@ -0,0 +1,57 @@
+// Copyright Epic Games, Inc. All Rights Reserved.
+
+#pragma once
+
+#include "CoreMinimal.h"
+#include "Layout/Visibility.h"
+#include "WidgetBlueprintEditor.h"
+#include "IPropertyTypeCustomization.h"
+#include "UObject/WeakObjectPtrTemplates.h"
+#include "Widgets/Views/SListView.h"
+
+class IDetailChildrenBuilder;
+class IPropertyHandle;
+class UWidget;
+struct FWidgetChild;
+class SSearchBox;
+
+class FWidgetChildTypeCustomization : public IPropertyTypeCustomization
+{
+public:
+	/** Makes a new instance of this detail layout class for a specific detail view requesting it */
+	static TSharedRef<class IPropertyTypeCustomization> MakeInstance(TSharedRef<FWidgetBlueprintEditor> InEditor)
+	{
+		return MakeShareable(new FWidgetChildTypeCustomization(InEditor));
+	}
+
+	FWidgetChildTypeCustomization(TSharedRef<FWidgetBlueprintEditor> InEditor)
+		: Editor(InEditor)
+	{
+	}
+	
+	virtual void CustomizeHeader(TSharedRef<IPropertyHandle> PropertyHandle, FDetailWidgetRow& HeaderRow, IPropertyTypeCustomizationUtils& CustomizationUtils) override;
+	virtual void CustomizeChildren(TSharedRef<IPropertyHandle> PropertyHandle, IDetailChildrenBuilder& ChildBuilder, IPropertyTypeCustomizationUtils& CustomizationUtils) override;
+
+private:
+
+	void SetDesiredFocusWidgetChild(UUserWidget* OwnerUserWidget, const FWidgetChild& WidgetChild);
+
+	void OnWidgetSelectionChanged(FName SelectedName, ESelectInfo::Type SelectionType);
+
+	TSharedRef<SWidget> GetPopupContent();
+	void HandleMenuOpen();
+
+	UWidget* GetCurrentValue() const;
+	FText GetCurrentValueText() const;
+
+private:
+
+	TWeakPtr<FWidgetBlueprintEditor> Editor;
+	TSharedPtr<SComboButton> WidgetListComboButton;
+	TSharedPtr<SListView<TWeakObjectPtr<UWidget>>> WidgetListView;
+	TSharedPtr<SSearchBox> SearchBox;
+
+	TWeakPtr<IPropertyHandle> PropertyHandlePtr;
+
+	TArray<TWeakObjectPtr<UWidget>> ReferenceableWidgets;
+};
diff --git a/Engine/Source/Editor/UMGEditor/Private/Details/SWidgetDetailsView.cpp b/Engine/Source/Editor/UMGEditor/Private/Details/SWidgetDetailsView.cpp
--- a/Engine/Source/Editor/UMGEditor/Private/Details/SWidgetDetailsView.cpp
+++ b/Engine/Source/Editor/UMGEditor/Private/Details/SWidgetDetailsView.cpp
@@ -19,6 +19,7 @@
 #include "Customizations/SlateBrushCustomization.h"
 #include "Customizations/SlateFontInfoCustomization.h"
 #include "Customizations/WidgetTypeCustomization.h"
+#include "Customizations/WidgetChildTypeCustomization.h"
 #include "Customizations/WidgetNavigationCustomization.h"
 #include "Customizations/CanvasSlotCustomization.h"
 #include "Customizations/HorizontalAlignmentCustomization.h"
@@ -200,6 +201,7 @@ SWidgetDetailsView::~SWidgetDetailsView()
 
 	// Unregister the property type layouts
 	PropertyView->UnregisterInstancedCustomPropertyTypeLayout(TEXT("Widget"));
+	PropertyView->UnregisterInstancedCustomPropertyTypeLayout(TEXT("WidgetChild"));
 	PropertyView->UnregisterInstancedCustomPropertyTypeLayout(TEXT("WidgetNavigation"));
 	PropertyView->UnregisterInstancedCustomPropertyTypeLayout(TEXT("PanelSlot"));
 	PropertyView->UnregisterInstancedCustomPropertyTypeLayout(TEXT("EHorizontalAlignment"));
@@ -217,6 +219,7 @@ void SWidgetDetailsView::RegisterCustomizations()
 
 	PropertyView->RegisterInstancedCustomPropertyLayout(UWidget::StaticClass(), FOnGetDetailCustomizationInstance::CreateStatic(&FBlueprintWidgetCustomization::MakeInstance, BlueprintEditorRef, BlueprintEditorRef->GetBlueprintObj()));
 	PropertyView->RegisterInstancedCustomPropertyTypeLayout(TEXT("Widget"), FOnGetPropertyTypeCustomizationInstance::CreateStatic(&FWidgetTypeCustomization::MakeInstance, BlueprintEditorRef), nullptr);
+	PropertyView->RegisterInstancedCustomPropertyTypeLayout(TEXT("WidgetChild"), FOnGetPropertyTypeCustomizationInstance::CreateStatic(&FWidgetChildTypeCustomization::MakeInstance, BlueprintEditor.Pin().ToSharedRef()), nullptr);
 	PropertyView->RegisterInstancedCustomPropertyTypeLayout(TEXT("WidgetNavigation"), FOnGetPropertyTypeCustomizationInstance::CreateStatic(&FWidgetNavigationCustomization::MakeInstance, BlueprintEditorRef));
 	PropertyView->RegisterInstancedCustomPropertyTypeLayout(TEXT("PanelSlot"), FOnGetPropertyTypeCustomizationInstance::CreateStatic(&FCanvasSlotCustomization::MakeInstance, BlueprintEditorRef->GetBlueprintObj()));
 	PropertyView->RegisterInstancedCustomPropertyTypeLayout(TEXT("EHorizontalAlignment"), FOnGetPropertyTypeCustomizationInstance::CreateStatic(&FHorizontalAlignmentCustomization::MakeInstance));
diff --git a/Engine/Source/Runtime/UMG/Private/WidgetChild.cpp b/Engine/Source/Runtime/UMG/Private/WidgetChild.cpp
new file mode 100644
--- /dev/null
+++ b/Engine/Source/Runtime/UMG/Private/WidgetChild.cpp
@@ -0,0 +1,36 @@
+// Copyright Epic Games, Inc. All Rights Reserved.
+
+#include "Blueprint/WidgetChild.h"
+
+#include "CoreMinimal.h"
+#include "Blueprint/WidgetTree.h"
+
+
+FWidgetChild::FWidgetChild()
+{}
+
+FWidgetChild::FWidgetChild(const class UUserWidget* Outer, FName InChildName)
+	:ChildName(InChildName)
+{
+	if (!ChildName.IsNone() && Outer)
+	{
+		ChildWidgetPtr = Outer->WidgetTree->FindWidget(ChildName);
+	}
+}
+
+bool FWidgetChild::IsValid() const 
+{
+	return ChildName.IsNone() || ChildWidgetPtr != nullptr; 
+}
+
+bool FWidgetChild::Resolve(const UWidgetTree* WidgetTree)
+{
+	if (!ChildName.IsNone() && WidgetTree)
+	{
+		ChildWidgetPtr = WidgetTree->FindWidget(ChildName);
+		return ChildWidgetPtr != nullptr;
+	}
+	
+	ChildWidgetPtr = nullptr;
+	return false;
+}
diff --git a/Engine/Source/Runtime/UMG/Public/Blueprint/WidgetChild.h b/Engine/Source/Runtime/UMG/Public/Blueprint/WidgetChild.h
new file mode 100644
--- /dev/null
+++ b/Engine/Source/Runtime/UMG/Public/Blueprint/WidgetChild.h
@@ -0,0 +1,39 @@
+// Copyright Epic Games, Inc. All Rights Reserved.
+
+#pragma once
+
+#include "CoreMinimal.h"
+#include "UObject/ObjectMacros.h"
+#include "UObject/Object.h"
+#include "UObject/WeakObjectPtr.h"
+
+#include "WidgetChild.generated.h"
+
+class UWidget;
+
+/**
+ * Represent a Widget present in the Tree Widget of the UserWidget
+ */
+USTRUCT()
+struct UMG_API FWidgetChild
+{
+	GENERATED_BODY();
+public:
+	FWidgetChild();
+	FWidgetChild(const class UUserWidget* Outer, FName InChildName);
+
+	FName GetChildName() const { return ChildName; };
+	UWidget* GetChildWidget() const { return ChildWidgetPtr.Get(); };
+
+	bool IsValid() const;
+
+	bool Resolve(const class UWidgetTree* WidgetTree);
+
+private:
+	/** This either the widget to focus, OR the name of the function to call. */
+	UPROPERTY(EditAnywhere, Category = "Interaction")
+	FName ChildName;
+
+	UPROPERTY(Transient)
+	TWeakObjectPtr<UWidget> ChildWidgetPtr;	
+};
