Build fix. Missing Foward Declaration

diff --git a/Engine/Source/Editor/UMGEditor/Private/Customizations/WidgetChildTypeCustomization.h b/Engine/Source/Editor/UMGEditor/Private/Customizations/WidgetChildTypeCustomization.h
--- a/Engine/Source/Editor/UMGEditor/Private/Customizations/WidgetChildTypeCustomization.h
+++ b/Engine/Source/Editor/UMGEditor/Private/Customizations/WidgetChildTypeCustomization.h
@@ -13,6 +13,7 @@ class IDetailChildrenBuilder;
 class IPropertyHandle;
 class UWidget;
 struct FWidgetChild;
+class SComboButton;
 class SSearchBox;
 
 class FWidgetChildTypeCustomization : public IPropertyTypeCustomization
