UMG: Add setting to disable UMG binding.

diff --git a/Engine/Plugins/Runtime/ModelViewViewModel/Source/ModelViewViewModelBlueprint/Private/MVVMDeveloperProjectSettings.cpp b/Engine/Plugins/Runtime/ModelViewViewModel/Source/ModelViewViewModelBlueprint/Private/MVVMDeveloperProjectSettings.cpp
--- a/Engine/Plugins/Runtime/ModelViewViewModel/Source/ModelViewViewModelBlueprint/Private/MVVMDeveloperProjectSettings.cpp
+++ b/Engine/Plugins/Runtime/ModelViewViewModel/Source/ModelViewViewModelBlueprint/Private/MVVMDeveloperProjectSettings.cpp
@@ -3,6 +3,9 @@
 #include "MVVMDeveloperProjectSettings.h"
 #include "Engine/Blueprint.h"
 
+#include "MVVMBlueprintViewModelContext.h"
+#include "Types/MVVMExecutionMode.h"
+
 #define LOCTEXT_NAMESPACE "MVVMDeveloperProjectSettings"
 
 UMVVMDeveloperProjectSettings::UMVVMDeveloperProjectSettings()
@@ -10,6 +13,13 @@ UMVVMDeveloperProjectSettings::UMVVMDeveloperProjectSettings()
 	AllowedExecutionMode.Add(EMVVMExecutionMode::Immediate);
 	AllowedExecutionMode.Add(EMVVMExecutionMode::Delayed);
 	AllowedExecutionMode.Add(EMVVMExecutionMode::Tick);
+	
+	AllowedContextCreationType.Add(EMVVMBlueprintViewModelContextCreationType::Manual);
+	AllowedContextCreationType.Add(EMVVMBlueprintViewModelContextCreationType::CreateInstance);
+	AllowedContextCreationType.Add(EMVVMBlueprintViewModelContextCreationType::GlobalViewModelCollection);
+	AllowedContextCreationType.Add(EMVVMBlueprintViewModelContextCreationType::PropertyPath);
+	AllowedContextCreationType.Add(EMVVMBlueprintViewModelContextCreationType::PropertyPath);
+	AllowedContextCreationType.Add(EMVVMBlueprintViewModelContextCreationType::Resolver);
 }
 
 FName UMVVMDeveloperProjectSettings::GetCategoryName() const
diff --git a/Engine/Plugins/Runtime/ModelViewViewModel/Source/ModelViewViewModelBlueprint/Public/MVVMDeveloperProjectSettings.h b/Engine/Plugins/Runtime/ModelViewViewModel/Source/ModelViewViewModelBlueprint/Public/MVVMDeveloperProjectSettings.h
--- a/Engine/Plugins/Runtime/ModelViewViewModel/Source/ModelViewViewModelBlueprint/Public/MVVMDeveloperProjectSettings.h
+++ b/Engine/Plugins/Runtime/ModelViewViewModel/Source/ModelViewViewModelBlueprint/Public/MVVMDeveloperProjectSettings.h
@@ -3,10 +3,12 @@
 #pragma once
 
 #include "Engine/DeveloperSettings.h"
-#include "Types/MVVMExecutionMode.h"
 
 #include "MVVMDeveloperProjectSettings.generated.h"
 
+enum class EMVVMBlueprintViewModelContextCreationType : uint8;
+enum class EMVVMExecutionMode : uint8;
+
 /**
  * 
  */
@@ -57,6 +59,11 @@ public:
 		return AllowedExecutionMode.Contains(ExecutionMode);
 	}
 
+	bool IsContextCreationTypeAllowed(EMVVMBlueprintViewModelContextCreationType ContextCreationType) const
+	{
+		return AllowedContextCreationType.Contains(ContextCreationType);
+	}
+
 	EMVVMDeveloperConversionFunctionFilterType GetConversionFunctionFilter() const
 	{
 		return ConversionFunctionFilter;
@@ -72,8 +79,16 @@ private:
 	/** Permission list for filtering which execution mode is allowed. */
 	UPROPERTY(EditAnywhere, config, Category = "Viewmodel")
 	TSet<EMVVMExecutionMode> AllowedExecutionMode;
+	
+	/** Permission list for filtering which context creation type is allowed. */
+	UPROPERTY(EditAnywhere, config, Category = "Viewmodel")
+	TSet<EMVVMBlueprintViewModelContextCreationType> AllowedContextCreationType;
 
 public:
+	/** Binding can be made from the DetailView Bind option. */
+	UPROPERTY(EditAnywhere, config, Category = "Viewmodel")
+	bool bAllowBindingFromDetailView = true;
+
 	/** Permission list for filtering which execution mode is allowed. */
 	UPROPERTY(EditAnywhere, config, Category = "Viewmodel")
 	EMVVMDeveloperConversionFunctionFilterType ConversionFunctionFilter = EMVVMDeveloperConversionFunctionFilterType::BlueprintActionRegistry;
diff --git a/Engine/Plugins/Runtime/ModelViewViewModel/Source/ModelViewViewModelEditor/Private/Customizations/MVVMPropertyBindingExtension.cpp b/Engine/Plugins/Runtime/ModelViewViewModel/Source/ModelViewViewModelEditor/Private/Customizations/MVVMPropertyBindingExtension.cpp
--- a/Engine/Plugins/Runtime/ModelViewViewModel/Source/ModelViewViewModelEditor/Private/Customizations/MVVMPropertyBindingExtension.cpp
+++ b/Engine/Plugins/Runtime/ModelViewViewModel/Source/ModelViewViewModelEditor/Private/Customizations/MVVMPropertyBindingExtension.cpp
@@ -7,6 +7,7 @@
 #include "Framework/MultiBox/MultiBoxBuilder.h"
 #include "Kismet2/BlueprintEditorUtils.h"
 #include "MVVMBlueprintView.h"
+#include "MVVMDeveloperProjectSettings.h"
 #include "MVVMWidgetBlueprintExtension_View.h"
 #include "Styling/StyleColors.h"
 #include "Widgets/Images/SImage.h"
@@ -225,6 +226,11 @@ TSharedPtr<FExtender> FMVVMPropertyBindingExtension::CreateMenuExtender(const UW
 
 bool FMVVMPropertyBindingExtension::CanExtend(const UWidgetBlueprint* WidgetBlueprint, const UWidget* Widget, const FProperty* Property) const
 {
+	if (!GetDefault<UMVVMDeveloperProjectSettings>()->bAllowBindingFromDetailView)
+	{
+		return false;
+	}
+
 	const UMVVMWidgetBlueprintExtension_View* MVVMExtensionPtr = UMVVMWidgetBlueprintExtension_View::GetExtension<UMVVMWidgetBlueprintExtension_View>(WidgetBlueprint);
 	if (MVVMExtensionPtr == nullptr)
 	{
diff --git a/Engine/Source/Editor/UMGEditor/Private/Customizations/UMGDetailCustomizations.cpp b/Engine/Source/Editor/UMGEditor/Private/Customizations/UMGDetailCustomizations.cpp
--- a/Engine/Source/Editor/UMGEditor/Private/Customizations/UMGDetailCustomizations.cpp
+++ b/Engine/Source/Editor/UMGEditor/Private/Customizations/UMGDetailCustomizations.cpp
@@ -101,7 +101,7 @@ private:
 	TSharedPtr<FEdGraphSchemaAction> Action;
 };
 
-TSharedRef<SWidget> FBlueprintWidgetCustomization::MakePropertyBindingWidget(TWeakPtr<FWidgetBlueprintEditor> InEditor, UFunction* SignatureFunction, TSharedRef<IPropertyHandle> InPropertyHandle, bool bInGeneratePureBindings)
+TSharedRef<SWidget> FBlueprintWidgetCustomization::MakePropertyBindingWidget(TWeakPtr<FWidgetBlueprintEditor> InEditor, UFunction* SignatureFunction, TSharedRef<IPropertyHandle> InPropertyHandle, bool bInGeneratePureBindings, bool bAllowOldBinding)
 {
 	if (!IModularFeatures::Get().IsModularFeatureAvailable("PropertyAccessEditor"))
 	{
@@ -582,6 +582,7 @@ TSharedRef<SWidget> FBlueprintWidgetCustomization::MakePropertyBindingWidget(TWe
 	});
 
 	Args.bGeneratePureBindings = bInGeneratePureBindings;
+	Args.bAllowNewBindings = bAllowOldBinding;
 
 	IPropertyAccessEditor& PropertyAccessEditor = IModularFeatures::Get().GetModularFeature<IPropertyAccessEditor>("PropertyAccessEditor");
 	return PropertyAccessEditor.MakePropertyBindingWidget(InEditor.Pin()->GetBlueprintObj(), Args);
@@ -653,6 +654,11 @@ bool FBlueprintWidgetCustomization::HasPropertyBindings(TWeakPtr<FWidgetBlueprin
 void FBlueprintWidgetCustomization::CreateEventCustomization( IDetailLayoutBuilder& DetailLayout, FDelegateProperty* Property, UWidget* Widget )
 {
 	TSharedRef<IPropertyHandle> DelegatePropertyHandle = DetailLayout.GetProperty(Property->GetFName(), Property->GetOwnerChecked<UClass>());
+	UWidgetBlueprint* BlueprintObj = Blueprint.Get();
+	if (!BlueprintObj)
+	{
+		return;
+	}
 
 	const bool bHasValidHandle = DelegatePropertyHandle->IsValidHandle();
 	if(!bHasValidHandle)
@@ -696,7 +702,7 @@ void FBlueprintWidgetCustomization::CreateEventCustomization( IDetailLayoutBuild
 		.MinDesiredWidth(200)
 		.MaxDesiredWidth(250)
 		[
-			MakePropertyBindingWidget(Editor.Pin(), Property->SignatureFunction, DelegatePropertyHandle, false)
+			MakePropertyBindingWidget(Editor.Pin(), Property->SignatureFunction, DelegatePropertyHandle, false, BlueprintObj->ArePropertyBindingsAllowed())
 		];
 }
 
@@ -704,42 +710,49 @@ void FBlueprintWidgetCustomization::ResetToDefault_RemoveBinding(TSharedPtr<IPro
 {
 	const FScopedTransaction Transaction(LOCTEXT("UnbindDelegate", "Remove Binding"));
 
-	Blueprint->Modify();
-
-	TArray<UObject*> OuterObjects;
-	PropertyHandle->GetOuterObjects(OuterObjects);
-	for ( UObject* SelectedObject : OuterObjects )
+	UWidgetBlueprint* BlueprintObj = Blueprint.Get();
+	if (BlueprintObj)
 	{
-		FDelegateEditorBinding Binding;
-		Binding.ObjectName = SelectedObject->GetName();
-		Binding.PropertyName = PropertyHandle->GetProperty()->GetFName();
+		BlueprintObj->Modify();
+		
+		TArray<UObject*> OuterObjects;
+		PropertyHandle->GetOuterObjects(OuterObjects);
+		for (UObject* SelectedObject : OuterObjects)
+		{
+			FDelegateEditorBinding Binding;
+			Binding.ObjectName = SelectedObject->GetName();
+			Binding.PropertyName = PropertyHandle->GetProperty()->GetFName();
+
+			BlueprintObj->Bindings.Remove(Binding);
+		}
 
-		Blueprint->Bindings.Remove(Binding);
+		FBlueprintEditorUtils::MarkBlueprintAsStructurallyModified(BlueprintObj);
 	}
 
-	FBlueprintEditorUtils::MarkBlueprintAsStructurallyModified(Blueprint);
+
 }
 
 
 FReply FBlueprintWidgetCustomization::HandleAddOrViewEventForVariable(const FName EventName, FName PropertyName, TWeakObjectPtr<UClass> PropertyClass)
 {
-	UBlueprint* BlueprintObj = Blueprint;
-
-	// Find the corresponding variable property in the Blueprint
-	FObjectProperty* VariableProperty = FindFProperty<FObjectProperty>(BlueprintObj->SkeletonGeneratedClass, PropertyName);
-
-	if (VariableProperty)
+	if (UBlueprint* BlueprintObj = Blueprint.Get())
 	{
-		if (!FKismetEditorUtilities::FindBoundEventForComponent(BlueprintObj, EventName, VariableProperty->GetFName()))
-		{
-			FKismetEditorUtilities::CreateNewBoundEventForClass(PropertyClass.Get(), EventName, BlueprintObj, VariableProperty);
-		}
-		else
+		// Find the corresponding variable property in the Blueprint
+		FObjectProperty* VariableProperty = FindFProperty<FObjectProperty>(BlueprintObj->SkeletonGeneratedClass, PropertyName);
+
+		if (VariableProperty)
 		{
-			const UK2Node_ComponentBoundEvent* ExistingNode = FKismetEditorUtilities::FindBoundEventForComponent(BlueprintObj, EventName, VariableProperty->GetFName());
-			if (ExistingNode)
+			if (!FKismetEditorUtilities::FindBoundEventForComponent(BlueprintObj, EventName, VariableProperty->GetFName()))
 			{
-				FKismetEditorUtilities::BringKismetToFocusAttentionOnObject(ExistingNode);
+				FKismetEditorUtilities::CreateNewBoundEventForClass(PropertyClass.Get(), EventName, BlueprintObj, VariableProperty);
+			}
+			else
+			{
+				const UK2Node_ComponentBoundEvent* ExistingNode = FKismetEditorUtilities::FindBoundEventForComponent(BlueprintObj, EventName, VariableProperty->GetFName());
+				if (ExistingNode)
+				{
+					FKismetEditorUtilities::BringKismetToFocusAttentionOnObject(ExistingNode);
+				}
 			}
 		}
 	}
@@ -749,18 +762,26 @@ FReply FBlueprintWidgetCustomization::HandleAddOrViewEventForVariable(const FNam
 
 int32 FBlueprintWidgetCustomization::HandleAddOrViewIndexForButton(const FName EventName, FName PropertyName) const
 {
-	UBlueprint* BlueprintObj = Blueprint;
-
-	if (FKismetEditorUtilities::FindBoundEventForComponent(BlueprintObj, EventName, PropertyName))
+	if (UBlueprint* BlueprintObj = Blueprint.Get())
 	{
-		return 0; // View
-	}
+		if (FKismetEditorUtilities::FindBoundEventForComponent(BlueprintObj, EventName, PropertyName))
+		{
+			return 0; // View
+		}
 
-	return 1; // Add
+		return 1; // Add
+	}
+	return 0;
 }
 
 void FBlueprintWidgetCustomization::CreateMulticastEventCustomization(IDetailLayoutBuilder& DetailLayout, FName ThisComponentName, UClass* PropertyClass, FMulticastDelegateProperty* DelegateProperty)
 {
+	UBlueprint* BlueprintObj = Blueprint.Get();
+	if (BlueprintObj == nullptr)
+	{
+		return;
+	}
+
 	const FString AddString = FString(TEXT("Add "));
 	const FString ViewString = FString(TEXT("View "));
 
@@ -777,7 +798,7 @@ void FBlueprintWidgetCustomization::CreateMulticastEventCustomization(IDetailLay
 		PropertyTooltip = FText::FromString(DelegateProperty->GetName());
 	}
 
-	FObjectProperty* ComponentProperty = FindFProperty<FObjectProperty>(Blueprint->SkeletonGeneratedClass, ThisComponentName);
+	FObjectProperty* ComponentProperty = FindFProperty<FObjectProperty>(BlueprintObj->SkeletonGeneratedClass, ThisComponentName);
 
 	if ( !ComponentProperty )
 	{
@@ -937,6 +958,12 @@ void FBlueprintWidgetCustomization::PerformAccessibilityCustomization(IDetailLay
 
 void FBlueprintWidgetCustomization::CustomizeAccessibilityProperty(IDetailLayoutBuilder& DetailLayout, const FName& BehaviorPropertyName, const FName& TextPropertyName)
 {
+	UWidgetBlueprint* BlueprintObj = Blueprint.Get();
+	if (!BlueprintObj)
+	{
+		return;
+	}
+
 	// Treat AccessibleBehavior as the "base" property for the row, and then add the AccessibleText binding to the end of it.
 	TSharedRef<IPropertyHandle> AccessibleBehaviorPropertyHandle = DetailLayout.GetProperty(BehaviorPropertyName);
 	IDetailPropertyRow& AccessibilityRow = DetailLayout.EditCategory("Accessibility").AddProperty(AccessibleBehaviorPropertyHandle);
@@ -954,7 +981,7 @@ void FBlueprintWidgetCustomization::CustomizeAccessibilityProperty(IDetailLayout
 			return !HasPropertyBindings(ThisEditor, AccessibleTextPropertyHandle);
 		}));
 
-	TSharedRef<SWidget> BindingWidget = MakePropertyBindingWidget(Editor, AccessibleTextDelegateProperty->SignatureFunction, AccessibleTextPropertyHandle, false);
+	TSharedRef<SWidget> BindingWidget = MakePropertyBindingWidget(Editor, AccessibleTextDelegateProperty->SignatureFunction, AccessibleTextPropertyHandle, false, BlueprintObj->ArePropertyBindingsAllowed());
 	TSharedRef<SHorizontalBox> CustomTextLayout = SNew(SHorizontalBox)
 	.Visibility(TAttribute<EVisibility>::Create([AccessibleBehaviorPropertyHandle]() -> EVisibility
 	{
diff --git a/Engine/Source/Editor/UMGEditor/Private/Customizations/UMGDetailCustomizations.h b/Engine/Source/Editor/UMGEditor/Private/Customizations/UMGDetailCustomizations.h
--- a/Engine/Source/Editor/UMGEditor/Private/Customizations/UMGDetailCustomizations.h
+++ b/Engine/Source/Editor/UMGEditor/Private/Customizations/UMGDetailCustomizations.h
@@ -33,7 +33,7 @@ public:
 	virtual void CustomizeDetails( IDetailLayoutBuilder& DetailLayout ) override;
 
 	/** Make a property binding widget */
-	static TSharedRef<SWidget> MakePropertyBindingWidget(TWeakPtr<FWidgetBlueprintEditor> InEditor, UFunction* SignatureFunction, TSharedRef<IPropertyHandle> InDelegatePropertyHandle, bool bInGeneratePureBindings);
+	static TSharedRef<SWidget> MakePropertyBindingWidget(TWeakPtr<FWidgetBlueprintEditor> InEditor, UFunction* SignatureFunction, TSharedRef<IPropertyHandle> InDelegatePropertyHandle, bool bInGeneratePureBindings, bool bAllowOldBinding);
 
 	/** Whether the property (or its parent property, if this is an array/struct element) currently has bindings */
 	static bool HasPropertyBindings(TWeakPtr<FWidgetBlueprintEditor> InEditor, const TSharedRef<IPropertyHandle>& InPropertyHandle);
@@ -57,6 +57,5 @@ private:
 private:
 
 	TWeakPtr<FWidgetBlueprintEditor> Editor;
-
-	UWidgetBlueprint* Blueprint;
+	TWeakObjectPtr<UWidgetBlueprint> Blueprint;
 };
diff --git a/Engine/Source/Editor/UMGEditor/Private/Details/DetailWidgetExtensionHandler.cpp b/Engine/Source/Editor/UMGEditor/Private/Details/DetailWidgetExtensionHandler.cpp
--- a/Engine/Source/Editor/UMGEditor/Private/Details/DetailWidgetExtensionHandler.cpp
+++ b/Engine/Source/Editor/UMGEditor/Private/Details/DetailWidgetExtensionHandler.cpp
@@ -140,8 +140,10 @@ void FDetailWidgetExtensionHandler::ExtendWidgetRow(
 
 	const UWidgetBlueprint* WidgetBlueprint = BlueprintEditor.Pin()->GetWidgetBlueprintObj();
 
+	bool bShouldShowOldBindingWidget = Private::ShouldShowOldBindingWidget(WidgetBlueprint, PropertyHandleRef.Get());
+
 	bool bShouldShowWidget = false;
-	bShouldShowWidget |= Private::ShouldShowOldBindingWidget(WidgetBlueprint, PropertyHandleRef.Get());
+	bShouldShowWidget |= bShouldShowOldBindingWidget;
 	
 	TArray<UObject*> Objects;
 	InPropertyHandle->GetOuterObjects(Objects);
@@ -174,6 +176,6 @@ void FDetailWidgetExtensionHandler::ExtendWidgetRow(
 
 	InWidgetRow.ExtensionContent()
 	[
-		FBlueprintWidgetCustomization::MakePropertyBindingWidget(BlueprintEditor, SignatureFunction, InPropertyHandle.ToSharedRef(), true)
+		FBlueprintWidgetCustomization::MakePropertyBindingWidget(BlueprintEditor, SignatureFunction, InPropertyHandle.ToSharedRef(), true, bShouldShowOldBindingWidget)
 	];
 }
diff --git a/Engine/Source/Editor/UMGEditor/Public/UMGEditorProjectSettings.h b/Engine/Source/Editor/UMGEditor/Public/UMGEditorProjectSettings.h
--- a/Engine/Source/Editor/UMGEditor/Public/UMGEditorProjectSettings.h
+++ b/Engine/Source/Editor/UMGEditor/Public/UMGEditorProjectSettings.h
@@ -140,11 +140,11 @@ public:
 	// End UObject Interface
 #endif
 
-protected:
-
+public:
 	UPROPERTY(EditAnywhere, config, Category = Compiler)
 	FWidgetCompilerOptions DefaultCompilerOptions;
 
+private:
 	UPROPERTY(EditAnywhere, config, Category = Compiler)
 	TArray<FDirectoryWidgetCompilerOptions> DirectoryCompilerOptions;
 
