UMG: Detect if the WidgetTemplate is a generated class in the constructor instead of using an argument.

diff --git a/Engine/Source/Editor/UMGEditor/Private/Palette/SPaletteViewModel.cpp b/Engine/Source/Editor/UMGEditor/Private/Palette/SPaletteViewModel.cpp
--- a/Engine/Source/Editor/UMGEditor/Private/Palette/SPaletteViewModel.cpp
+++ b/Engine/Source/Editor/UMGEditor/Private/Palette/SPaletteViewModel.cpp
@@ -238,8 +238,7 @@ void FWidgetCatalogViewModel::BuildClassWidgetList()
 
 		if (WidgetClass->IsChildOf(UUserWidget::StaticClass()))
 		{
-			bool bIsBlueprintGeneratedClass = Cast<UBlueprintGeneratedClass>(WidgetClass) != nullptr;
-			AddWidgetTemplate(MakeShared<FWidgetTemplateBlueprintClass>(FAssetData(WidgetClass), WidgetClass, bIsBlueprintGeneratedClass));
+			AddWidgetTemplate(MakeShared<FWidgetTemplateBlueprintClass>(FAssetData(WidgetClass), WidgetClass));
 		}
 		else
 		{
@@ -283,8 +282,7 @@ void FWidgetCatalogViewModel::BuildClassWidgetList()
 
 		if (Usable.GetValue().NativeParentClass->IsChildOf(UUserWidget::StaticClass()))
 		{
-			bool bIsBlueprintGeneratedClass = false;
-			AddWidgetTemplate(MakeShared<FWidgetTemplateBlueprintClass>(BPAssetData, nullptr, bIsBlueprintGeneratedClass));
+			AddWidgetTemplate(MakeShared<FWidgetTemplateBlueprintClass>(BPAssetData, nullptr));
 		}
 		else
 		{
@@ -308,6 +306,12 @@ void FWidgetCatalogViewModel::BuildClassWidgetList()
 			continue;
 		}
 
+		// Was already parsed by the TObjectIterator<UClass>
+		if (BPAssetData.IsAssetLoaded())
+		{
+			continue;
+		}
+
 		TValueOrError<FWidgetBlueprintEditorUtils::FUsableWidgetClassResult, void> Usable = FWidgetBlueprintEditorUtils::IsUsableWidgetClass(BPAssetData);
 		if (Usable.HasError())
 		{
@@ -319,8 +323,7 @@ void FWidgetCatalogViewModel::BuildClassWidgetList()
 			continue;
 		}
 
-		bool bIsBlueprintGeneratedClass = true;
-		AddWidgetTemplate(MakeShared<FWidgetTemplateBlueprintClass>(BPAssetData, nullptr, bIsBlueprintGeneratedClass));
+		AddWidgetTemplate(MakeShared<FWidgetTemplateBlueprintClass>(BPAssetData, nullptr));
 	}
 }
 
diff --git a/Engine/Source/Editor/UMGEditor/Private/Templates/WidgetTemplateBlueprintClass.cpp b/Engine/Source/Editor/UMGEditor/Private/Templates/WidgetTemplateBlueprintClass.cpp
--- a/Engine/Source/Editor/UMGEditor/Private/Templates/WidgetTemplateBlueprintClass.cpp
+++ b/Engine/Source/Editor/UMGEditor/Private/Templates/WidgetTemplateBlueprintClass.cpp
@@ -16,11 +16,17 @@
 
 #define LOCTEXT_NAMESPACE "UMGEditor"
 
-FWidgetTemplateBlueprintClass::FWidgetTemplateBlueprintClass(const FAssetData& InWidgetAssetData, TSubclassOf<UUserWidget> InUserWidgetClass, bool bInIsBlueprintGeneratedClass)
+//FWidgetTemplateBlueprintClass::FWidgetTemplateBlueprintClass(const FAssetData& InWidgetAssetData, TSubclassOf<UUserWidget> InUserWidgetClass, bool bInIsBlueprintGeneratedClass)
+//	: FWidgetTemplateBlueprintClass(InWidgetAssetData, InUserWidgetClass)
+//{}
+
+FWidgetTemplateBlueprintClass::FWidgetTemplateBlueprintClass(const FAssetData& InWidgetAssetData, TSubclassOf<UUserWidget> InUserWidgetClass)
 	: FWidgetTemplateClass(InWidgetAssetData, InUserWidgetClass)
-	, bIsBlueprintGeneratedClass(bInIsBlueprintGeneratedClass)
 {
 	// Blueprints get the class type actions for their parent native class - this avoids us having to load the blueprint
+	static const FTopLevelAssetPath BlueprintGeneratedClassAssetPath = UWidgetBlueprintGeneratedClass::StaticClass()->GetClassPathName();
+	bool bClassIsUWidgetBlueprintGeneratedClass = Cast<UBlueprintGeneratedClass>(InUserWidgetClass.Get()) != nullptr;
+	bIsBlueprintGeneratedClass = BlueprintGeneratedClassAssetPath == InWidgetAssetData.AssetClassPath || bClassIsUWidgetBlueprintGeneratedClass;
 	if (bIsBlueprintGeneratedClass && !InUserWidgetClass && InWidgetAssetData.IsValid())
 	{
 		FString ParentClassName;
diff --git a/Engine/Source/Editor/UMGEditor/Public/Templates/WidgetTemplateBlueprintClass.h b/Engine/Source/Editor/UMGEditor/Public/Templates/WidgetTemplateBlueprintClass.h
--- a/Engine/Source/Editor/UMGEditor/Public/Templates/WidgetTemplateBlueprintClass.h
+++ b/Engine/Source/Editor/UMGEditor/Public/Templates/WidgetTemplateBlueprintClass.h
@@ -25,7 +25,10 @@ public:
 	* @param InUserWidgetClass The user-created widget class that will be created by this template
 	* @param bCacheParentClass True means we will search for parent class and cache, this is a slow operation, but useful for BPGCs.
 	*/
-	FWidgetTemplateBlueprintClass(const FAssetData& InWidgetAssetData, TSubclassOf<UUserWidget> InUserWidgetClass = nullptr, bool bInIsBlueprintGeneratedClass = false);
+	FWidgetTemplateBlueprintClass(const FAssetData& InWidgetAssetData, TSubclassOf<UUserWidget> InUserWidgetClass = nullptr);
+
+	//UE_DEPRECATED(5.3, "FWidgetTemplateBlueprintClass with bInIsBlueprintGeneratedClass is deprecated.")
+	//FWidgetTemplateBlueprintClass(const FAssetData& InWidgetAssetData, TSubclassOf<UUserWidget> InUserWidgetClass, bool bInIsBlueprintGeneratedClass);
 
 	/** Destructor */
 	virtual ~FWidgetTemplateBlueprintClass();
