Re-Enable Abstract Classes when creating new user widgets

#prelfight 63f82cc12b1a2765f9349297

diff --git a/Engine/Source/Editor/UMGEditor/Private/WidgetBlueprintFactory.cpp b/Engine/Source/Editor/UMGEditor/Private/WidgetBlueprintFactory.cpp
--- a/Engine/Source/Editor/UMGEditor/Private/WidgetBlueprintFactory.cpp
+++ b/Engine/Source/Editor/UMGEditor/Private/WidgetBlueprintFactory.cpp
@@ -90,7 +90,7 @@ bool UWidgetBlueprintFactory::ConfigureProperties()
 			Options.ExtraPickerCommonClasses.Add(UUserWidget::StaticClass());
 		}
 
-		Filter->DisallowedClassFlags = CLASS_Deprecated | CLASS_Abstract | CLASS_NewerVersionExists | CLASS_Hidden | CLASS_HideDropDown;
+		Filter->DisallowedClassFlags = CLASS_Deprecated | CLASS_NewerVersionExists | CLASS_Hidden | CLASS_HideDropDown;
 		Filter->AllowedChildrenOfClasses.Add(UUserWidget::StaticClass());
 
 		const FText TitleText = LOCTEXT("CreateWidgetBlueprint", "Pick Parent Class for New Widget Blueprint");
