Fix BP GeneratedClass Crashing for UWidgets

#tests repo'd issue and in separate cook did some thing with fixes and got no crash see video in jira
#localization none

diff --git a/Engine/Source/Editor/UMGEditor/Private/Templates/WidgetTemplateBlueprintClass.cpp b/Engine/Source/Editor/UMGEditor/Private/Templates/WidgetTemplateBlueprintClass.cpp
--- a/Engine/Source/Editor/UMGEditor/Private/Templates/WidgetTemplateBlueprintClass.cpp
+++ b/Engine/Source/Editor/UMGEditor/Private/Templates/WidgetTemplateBlueprintClass.cpp
@@ -24,9 +24,12 @@ FWidgetTemplateBlueprintClass::FWidgetTemplateBlueprintClass(const FAssetData& I
 	: FWidgetTemplateClass(InWidgetAssetData, InUserWidgetClass)
 {
 	// Blueprints get the class type actions for their parent native class - this avoids us having to load the blueprint
-	static const FTopLevelAssetPath BlueprintGeneratedClassAssetPath = UWidgetBlueprintGeneratedClass::StaticClass()->GetClassPathName();
+	static const FTopLevelAssetPath WidgetBlueprintGeneratedClassAssetPath = UWidgetBlueprintGeneratedClass::StaticClass()->GetClassPathName();
+	static const FTopLevelAssetPath BlueprintGeneratedClassAssetPath = UBlueprintGeneratedClass::StaticClass()->GetClassPathName();
 	bool bClassIsUWidgetBlueprintGeneratedClass = Cast<UBlueprintGeneratedClass>(InUserWidgetClass.Get()) != nullptr;
-	bIsBlueprintGeneratedClass = BlueprintGeneratedClassAssetPath == InWidgetAssetData.AssetClassPath || bClassIsUWidgetBlueprintGeneratedClass;
+	bIsBlueprintGeneratedClass = WidgetBlueprintGeneratedClassAssetPath == InWidgetAssetData.AssetClassPath 
+								|| BlueprintGeneratedClassAssetPath == InWidgetAssetData.AssetClassPath 
+								|| bClassIsUWidgetBlueprintGeneratedClass;
 	if (bIsBlueprintGeneratedClass && !InUserWidgetClass && InWidgetAssetData.IsValid())
 	{
 		FString ParentClassName;
diff --git a/Engine/Source/Editor/UMGEditor/Private/Templates/WidgetTemplateClass.cpp b/Engine/Source/Editor/UMGEditor/Private/Templates/WidgetTemplateClass.cpp
--- a/Engine/Source/Editor/UMGEditor/Private/Templates/WidgetTemplateClass.cpp
+++ b/Engine/Source/Editor/UMGEditor/Private/Templates/WidgetTemplateClass.cpp
@@ -160,6 +160,11 @@ void FWidgetTemplateClass::OnObjectsReplaced(const TMap<UObject*, UObject*>& Rep
 
 UWidget* FWidgetTemplateClass::CreateNamed(class UWidgetTree* Tree, FName NameOverride)
 {
+	if (Tree == nullptr || !WidgetClass.IsValid() || WidgetClass.Get()->HasAnyClassFlags(CLASS_Abstract | CLASS_Deprecated))
+	{
+		return nullptr;
+	}
+
 	if (NameOverride != NAME_None)
 	{
 		UObject* ExistingObject = StaticFindObject(UObject::StaticClass(), Tree, *NameOverride.ToString());
@@ -169,11 +174,6 @@ UWidget* FWidgetTemplateClass::CreateNamed(class UWidgetTree* Tree, FName NameOv
 		}
 	}
 
-	if (WidgetClass.Get()->HasAnyClassFlags(CLASS_Abstract | CLASS_Deprecated))
-	{
-		return nullptr;
-	}
-
 	UWidget* NewWidget = Tree->ConstructWidget<UWidget>(WidgetClass.Get(), NameOverride);
 	NewWidget->CreatedFromPalette();
 
