Fix PaletteView not toggling on left click.

diff --git a/Engine/Source/Editor/UMGEditor/Private/Library/SLibraryView.cpp b/Engine/Source/Editor/UMGEditor/Private/Library/SLibraryView.cpp
--- a/Engine/Source/Editor/UMGEditor/Private/Library/SLibraryView.cpp
+++ b/Engine/Source/Editor/UMGEditor/Private/Library/SLibraryView.cpp
@@ -204,9 +204,10 @@ void SLibraryView::Construct(const FArguments& InArgs, TSharedPtr<FWidgetBluepri
 
 	SAssignNew(WidgetTemplatesView, STreeView< TSharedPtr<FWidgetViewModel> >)
 		.ItemHeight(1.0f)
-		.SelectionMode(ESelectionMode::Single)
+		.SelectionMode(ESelectionMode::SingleToggle)
 		.OnGenerateRow(this, &SLibraryView::OnGenerateWidgetTemplateLibrary)
 		.OnGetChildren(FilterHandler.ToSharedRef(), &LibraryFilterHandler::OnGetFilteredChildren)
+		.OnMouseButtonClick(this, &SLibraryView::WidgetLibrary_OnClick)
 		.TreeItemsSource(&TreeWidgetViewModels);
 
 	FilterHandler->SetTreeView(WidgetTemplatesView.Get());
@@ -407,4 +408,25 @@ void SLibraryView::Tick(const FGeometry& AllottedGeometry, const double InCurren
 	}
 }
 
+void SLibraryView::WidgetLibrary_OnClick(TSharedPtr<FWidgetViewModel> SelectedItem)
+{
+	if (!SelectedItem.IsValid())
+	{
+		return;
+	}
+
+	// If it's a category, toggle it
+	if (SelectedItem->IsCategory())
+	{
+		if (TSharedPtr<FWidgetHeaderViewModel> CategoryHeader = StaticCastSharedPtr<FWidgetHeaderViewModel>(SelectedItem))
+		{
+			if (TSharedPtr<ITableRow> TableRow = WidgetTemplatesView->WidgetFromItem(CategoryHeader))
+			{
+				TableRow->ToggleExpansion();
+			}
+		}
+		return;
+	}
+}
+
 #undef LOCTEXT_NAMESPACE
diff --git a/Engine/Source/Editor/UMGEditor/Private/Library/SLibraryView.h b/Engine/Source/Editor/UMGEditor/Private/Library/SLibraryView.h
--- a/Engine/Source/Editor/UMGEditor/Private/Library/SLibraryView.h
+++ b/Engine/Source/Editor/UMGEditor/Private/Library/SLibraryView.h
@@ -82,6 +82,9 @@ public:
 	
 	virtual void Tick( const FGeometry& AllottedGeometry, const double InCurrentTime, const float InDeltaTime ) override;
 
+	/** On item clicked in palette view */
+	void WidgetLibrary_OnClick(TSharedPtr<FWidgetViewModel> SelectedItem);
+
 private:
 
 	/** Builds settings menu. */
diff --git a/Engine/Source/Editor/UMGEditor/Private/Palette/SPaletteView.cpp b/Engine/Source/Editor/UMGEditor/Private/Palette/SPaletteView.cpp
--- a/Engine/Source/Editor/UMGEditor/Private/Palette/SPaletteView.cpp
+++ b/Engine/Source/Editor/UMGEditor/Private/Palette/SPaletteView.cpp
@@ -140,10 +140,11 @@ void SPaletteView::Construct(const FArguments& InArgs, TSharedPtr<FWidgetBluepri
 
 	SAssignNew(WidgetTemplatesView, STreeView< TSharedPtr<FWidgetViewModel> >)
 		.ItemHeight(1.0f)
-		.SelectionMode(ESelectionMode::Single)
+		.SelectionMode(ESelectionMode::SingleToggle)
 		.OnGenerateRow(this, &SPaletteView::OnGenerateWidgetTemplateItem)
 		.OnGetChildren(FilterHandler.ToSharedRef(), &PaletteFilterHandler::OnGetFilteredChildren)
 		.OnSelectionChanged(this, &SPaletteView::WidgetPalette_OnSelectionChanged)
+		.OnMouseButtonClick(this, &SPaletteView::WidgetPalette_OnClick)
 		.TreeItemsSource(&TreeWidgetViewModels);
 		
 
@@ -209,6 +210,27 @@ void SPaletteView::OnSearchChanged(const FText& InFilterText)
 	PaletteViewModel->SetSearchText(InFilterText);
 }
 
+void SPaletteView::WidgetPalette_OnClick(TSharedPtr<FWidgetViewModel> SelectedItem)
+{
+	if (!SelectedItem.IsValid())
+	{
+		return;
+	}
+
+	// If it's a category, toggle it
+	if (SelectedItem->IsCategory())
+	{
+		if (TSharedPtr<FWidgetHeaderViewModel> CategoryHeader = StaticCastSharedPtr<FWidgetHeaderViewModel>(SelectedItem))
+		{
+			if (TSharedPtr<ITableRow> TableRow = WidgetTemplatesView->WidgetFromItem(CategoryHeader))
+			{
+				TableRow->ToggleExpansion();
+			}
+		}
+		return;
+	}
+}
+
 void SPaletteView::WidgetPalette_OnSelectionChanged(TSharedPtr<FWidgetViewModel> SelectedItem, ESelectInfo::Type SelectInfo)
 {
 	if (!SelectedItem.IsValid()) 
diff --git a/Engine/Source/Editor/UMGEditor/Private/Palette/SPaletteView.h b/Engine/Source/Editor/UMGEditor/Private/Palette/SPaletteView.h
--- a/Engine/Source/Editor/UMGEditor/Private/Palette/SPaletteView.h
+++ b/Engine/Source/Editor/UMGEditor/Private/Palette/SPaletteView.h
@@ -86,6 +86,9 @@ public:
 	/** Gets the text currently displayed in the search box. */
 	FText GetSearchText() const;
 
+	/** On item clicked in palette view */
+	void WidgetPalette_OnClick(TSharedPtr<FWidgetViewModel> SelectedItem);
+
 	/** On Selected template widget changed in palette view */
 	void WidgetPalette_OnSelectionChanged(TSharedPtr<FWidgetViewModel> SelectedItem, ESelectInfo::Type SelectInfo);
 
diff --git a/Engine/Source/Editor/UMGEditor/Private/Palette/SPaletteViewModel.h b/Engine/Source/Editor/UMGEditor/Private/Palette/SPaletteViewModel.h
--- a/Engine/Source/Editor/UMGEditor/Private/Palette/SPaletteViewModel.h
+++ b/Engine/Source/Editor/UMGEditor/Private/Palette/SPaletteViewModel.h
@@ -25,6 +25,11 @@ public:
 
 	virtual bool IsTemplate() const = 0;
 
+	virtual bool IsCategory() const
+	{
+		return false;
+	}
+
 	/** @param OutStrings - Returns an array of strings used for filtering/searching this item. */
 	virtual void GetFilterStrings(TArray<FString>& OutStrings) const = 0;
 
@@ -101,6 +106,11 @@ public:
 		return false;
 	}
 
+	virtual bool IsCategory() const override
+	{
+		return true;
+	}
+
 	virtual void GetFilterStrings(TArray<FString>& OutStrings) const override
 	{
 		// Headers should never be included in filtering to avoid showing a header with all of
