UMG: Detects circular reference in the widgettree at compilation. If a circular exists, the widget will be removed from the compiled widgettree (to prevent callstack overflow). You could repro by creating a circular reference without compiling the UMG assets (new UMG A and B, add B in A, add A in B, compile A or B.

diff --git a/Engine/Source/Editor/UMGEditor/Private/WidgetBlueprint.cpp b/Engine/Source/Editor/UMGEditor/Private/WidgetBlueprint.cpp
--- a/Engine/Source/Editor/UMGEditor/Private/WidgetBlueprint.cpp
+++ b/Engine/Source/Editor/UMGEditor/Private/WidgetBlueprint.cpp
@@ -1154,7 +1154,7 @@ bool UWidgetBlueprint::IsWidgetFreeFromCircularReferences(UUserWidget* UserWidge
 			if (GeneratedByBlueprint->WidgetTree && GeneratedByBlueprint->WidgetTree->RootWidget)
 			{
 				TArray<UWidget*> ChildWidgets;
-				GeneratedByBlueprint->WidgetTree->GetChildWidgets(GeneratedByBlueprint->WidgetTree->RootWidget, ChildWidgets);
+				GeneratedByBlueprint->WidgetTree->GetAllWidgets(ChildWidgets);
 				for (UWidget* ChildWidget : ChildWidgets)
 				{
 					if (UWidgetBlueprint* ChildGeneratedBlueprint = Cast<UWidgetBlueprint>(ChildWidget->WidgetGeneratedBy))
@@ -1189,6 +1189,63 @@ bool UWidgetBlueprint::IsWidgetFreeFromCircularReferences(UUserWidget* UserWidge
 	return true;
 }
 
+namespace UE::UMG::Private
+{
+bool HasCircularReferences(const UClass* CurrentClass, TArray<const UClass*, TInlineAllocator<32>> DiscoveredBlueprint, UWidget*& OutResult)
+{
+	if (DiscoveredBlueprint.ContainsByPredicate([CurrentClass](const UClass* Other) { return CurrentClass->IsChildOf(Other); }))
+	{
+		return true;
+	}
+	DiscoveredBlueprint.Add(CurrentClass);
+
+	if (const UWidgetBlueprintGeneratedClass* CurrentWidgetClass = Cast<const UWidgetBlueprintGeneratedClass>(CurrentClass))
+	{
+		TArray<UWidget*> AllWidgets;
+		if (const UWidgetBlueprint* WidgetBP = Cast<const UWidgetBlueprint>(CurrentWidgetClass->ClassGeneratedBy))
+		{
+			if (WidgetBP->WidgetTree)
+			{
+				WidgetBP->WidgetTree->GetAllWidgets(AllWidgets);
+			}
+		}
+		else if (UWidgetTree* CurrentWidgetTree = CurrentWidgetClass->GetWidgetTreeArchetype())
+		{
+			CurrentWidgetTree->GetAllWidgets(AllWidgets);
+
+		}
+
+		for (UWidget* Widget : AllWidgets)
+		{
+			if (UUserWidget* UserWidget = Cast<UUserWidget>(Widget))
+			{
+				if (HasCircularReferences(UserWidget->GetClass(), DiscoveredBlueprint, OutResult))
+				{
+					OutResult = Widget;
+					return true;
+				}
+			}
+		}
+	}
+
+	return false;
+}
+}
+
+TValueOrError<void, UWidget*> UWidgetBlueprint::HasCircularReferences() const
+{
+	if (GeneratedClass)
+	{
+		TArray<const UClass*, TInlineAllocator<32>> DiscoveredBlueprint;
+		UWidget* Result = nullptr;
+		if (UE::UMG::Private::HasCircularReferences(GeneratedClass, DiscoveredBlueprint, Result))
+		{
+			return MakeError(Result);
+		}
+	}
+	return MakeValue();
+}
+
 UPackage* UWidgetBlueprint::GetWidgetTemplatePackage() const
 {
 	return GetOutermost();
diff --git a/Engine/Source/Editor/UMGEditor/Private/WidgetBlueprintCompiler.cpp b/Engine/Source/Editor/UMGEditor/Private/WidgetBlueprintCompiler.cpp
--- a/Engine/Source/Editor/UMGEditor/Private/WidgetBlueprintCompiler.cpp
+++ b/Engine/Source/Editor/UMGEditor/Private/WidgetBlueprintCompiler.cpp
@@ -804,6 +804,8 @@ void FWidgetBlueprintCompilerContext::FinishCompilingClass(UClass* Class)
 				LinkerLoadingContext->AddUniqueLoadedObjects(DupObjects);
 			}
 
+			//WidgetBP->IsWidgetFreeFromCircularReferences();
+
 			BPGClass->SetWidgetTreeArchetype(NewWidgetTree);
 			if (OldWidgetTree)
 			{
@@ -814,6 +816,21 @@ void FWidgetBlueprintCompilerContext::FinishCompilingClass(UClass* Class)
 			WidgetBP->WidgetTree->SetFlags(PreviousFlags);
 		}
 
+		{
+			TValueOrError<void, UWidget*> HasReference = WidgetBP->HasCircularReferences();
+			if (HasReference.HasError())
+			{
+				if (UWidget* FoundCircularWidget = BPGClass->GetWidgetTreeArchetype()->FindWidget(HasReference.GetError()->GetFName()))
+				{
+					BPGClass->GetWidgetTreeArchetype()->RemoveWidget(FoundCircularWidget);
+				}
+				MessageLog.Error(*FText::Format(LOCTEXT("CircularReference", "The WidgetTree '{0}' Contains circular references. See widget '{1}'"),
+					FText::FromString(WidgetBP->WidgetTree->GetPathName()),
+					FText::FromString(HasReference.GetError()->GetName())
+					).ToString());
+			}
+		}
+
 		int32 AnimIndex = 0;
 		for ( const UWidgetAnimation* Animation : WidgetBP->Animations )
 		{
diff --git a/Engine/Source/Editor/UMGEditor/Public/WidgetBlueprint.h b/Engine/Source/Editor/UMGEditor/Public/WidgetBlueprint.h
--- a/Engine/Source/Editor/UMGEditor/Public/WidgetBlueprint.h
+++ b/Engine/Source/Editor/UMGEditor/Public/WidgetBlueprint.h
@@ -10,6 +10,7 @@
 #include "Binding/DynamicPropertyPath.h"
 #include "Blueprint/WidgetBlueprintGeneratedClass.h"
 #include "Animation/WidgetAnimationBinding.h"
+#include "Templates/ValueOrError.h"
 
 #include "WidgetBlueprint.generated.h"
 
@@ -313,6 +314,9 @@ public:
 
 	/** Returns true if the supplied user widget will not create a circular reference when added to this blueprint */
 	bool IsWidgetFreeFromCircularReferences(UUserWidget* UserWidget) const;
+	
+	/**  */
+	TValueOrError<void, UWidget*> HasCircularReferences() const;
 
 	static bool ValidateGeneratedClass(const UClass* InClass);
 	
