Fix duplicated Loctext warning introduced with CL 24746024.
#rnx
#issue 279297

diff --git a/Engine/Source/Editor/UMGEditor/Private/WidgetBlueprintCompiler.cpp b/Engine/Source/Editor/UMGEditor/Private/WidgetBlueprintCompiler.cpp
--- a/Engine/Source/Editor/UMGEditor/Private/WidgetBlueprintCompiler.cpp
+++ b/Engine/Source/Editor/UMGEditor/Private/WidgetBlueprintCompiler.cpp
@@ -824,7 +824,7 @@ void FWidgetBlueprintCompilerContext::FinishCompilingClass(UClass* Class)
 				{
 					BPGClass->GetWidgetTreeArchetype()->RemoveWidget(FoundCircularWidget);
 				}
-				MessageLog.Error(*FText::Format(LOCTEXT("CircularReference", "The WidgetTree '{0}' Contains circular references. See widget '{1}'"),
+				MessageLog.Error(*FText::Format(LOCTEXT("WidgetTreeCircularReference", "The WidgetTree '{0}' Contains circular references. See widget '{1}'"),
 					FText::FromString(WidgetBP->WidgetTree->GetPathName()),
 					FText::FromString(HasReference.GetError()->GetName())
 					).ToString());
