Expose SAssetView to editor widgets

diff --git a/Engine/Source/Editor/UMGEditor/Public/Components/AssetThumbnailWidget.cpp b/Engine/Source/Editor/UMGEditor/Public/Components/AssetThumbnailWidget.cpp
new file mode 100644
--- /dev/null
+++ b/Engine/Source/Editor/UMGEditor/Public/Components/AssetThumbnailWidget.cpp
@@ -0,0 +1,102 @@
+// Copyright Epic Games, Inc. All Rights Reserved.
+
+#include "AssetThumbnailWidget.h"
+
+#include "AssetThumbnail.h"
+#include "AssetRegistry/IAssetRegistry.h"
+#include "Styling/AppStyle.h"
+#include "ThumbnailRendering/ThumbnailManager.h"
+#include "Widgets/Layout/SBorder.h"
+#include "Widgets/Text/STextBlock.h"
+
+#define LOCTEXT_NAMESPACE "UMG"
+
+FAssetThumbnailConfig FAssetThumbnailWidgetSettings::ToThumbnailConfig() const
+{
+	FAssetThumbnailConfig Result;
+	Result.bForceGenericThumbnail				= bForceGenericThumbnail;
+	Result.bAllowHintText						= bAllowHintText;
+	Result.bAllowRealTimeOnHovered				= bAllowRealTimeOnHovered;
+	Result.bAllowAssetSpecificThumbnailOverlay	= bAllowAssetSpecificThumbnailOverlay;
+	Result.ThumbnailLabel						= static_cast<EThumbnailLabel::Type>(ThumbnailLabel);
+	Result.HighlightedText						= TAttribute<FText>::CreateLambda([Delegate = HighlightedTextDelegate]()
+	{
+		return Delegate.IsBound() ? Delegate.Execute() : FText::GetEmpty();
+	});
+	Result.HintColorAndOpacity					= HintColorAndOpacity;
+	Result.AssetTypeColorOverride				= bOverrideAssetTypeColor ? AssetTypeColorOverride : TOptional<FLinearColor>();
+	Result.Padding								= Padding;
+	Result.GenericThumbnailSize					= GenericThumbnailSize;
+	Result.ColorStripOrientation				= static_cast<EThumbnailColorStripOrientation>(ColorStripOrientation);
+	return Result;
+}
+void UAssetThumbnailWidget::SetThumbnailSettings(const FAssetThumbnailWidgetSettings& InThumbnailSettings)
+{
+	ThumbnailSettings = InThumbnailSettings;
+	UpdateNativeAssetThumbnailWidget();
+}
+
+void UAssetThumbnailWidget::SetAsset(const FAssetData& AssetData)
+{
+	AssetToShow = AssetData;
+	UpdateNativeAssetThumbnailWidget();
+}
+
+void UAssetThumbnailWidget::SetAssetByObject(UObject* Object)
+{
+	AssetToShow = Object;
+	UpdateNativeAssetThumbnailWidget();
+}
+
+void UAssetThumbnailWidget::SynchronizeProperties()
+{
+	Super::SynchronizeProperties();
+	UpdateNativeAssetThumbnailWidget();
+}
+
+void UAssetThumbnailWidget::ReleaseSlateResources(bool bReleaseChildren)
+{
+	Super::ReleaseSlateResources(bReleaseChildren);
+	DisplayedWidget.Reset();
+	ThumbnailRenderer.Reset();
+}
+
+TSharedRef<SWidget> UAssetThumbnailWidget::RebuildWidget()
+{
+	DisplayedWidget = SNew(SBorder)
+		.Padding(0.0f)
+		.HAlign(HAlign_Fill)
+		.VAlign(VAlign_Fill)
+		.BorderImage(FAppStyle::GetBrush("NoBorder"));
+	UpdateNativeAssetThumbnailWidget();
+	return DisplayedWidget.ToSharedRef();
+}
+
+const FText UAssetThumbnailWidget::GetPaletteCategory()
+{
+	return LOCTEXT("Editor", "Editor");
+}
+
+void UAssetThumbnailWidget::UpdateNativeAssetThumbnailWidget()
+{
+	if (ThumbnailRenderer && (AssetToShow.IsValid() && ThumbnailRenderer->GetAssetData() != AssetToShow))
+	{
+		ThumbnailRenderer->SetAsset(AssetToShow);
+	}
+	else if (AssetToShow.IsValid() &&
+		(!ThumbnailRenderer || ThumbnailRenderer->GetSize() != Resolution))
+	{
+		ThumbnailRenderer = MakeShared<FAssetThumbnail>(AssetToShow, FMath::Clamp(Resolution.X, 1, 1024), FMath::Clamp(Resolution.Y, 1, 1024), UThumbnailManager::Get().GetSharedThumbnailPool());
+	}
+	
+	if (ThumbnailRenderer && ensure(DisplayedWidget))
+	{
+		DisplayedWidget->SetContent(ThumbnailRenderer->MakeThumbnailWidget(ThumbnailSettings.ToThumbnailConfig()));
+	}
+	else
+	{
+		DisplayedWidget->SetContent(SNew(STextBlock).Text(LOCTEXT("NoAssetData", "Call SetThumbnailAsset")));
+	}
+}
+
+#undef LOCTEXT_NAMESPACE
diff --git a/Engine/Source/Editor/UMGEditor/Public/Components/AssetThumbnailWidget.h b/Engine/Source/Editor/UMGEditor/Public/Components/AssetThumbnailWidget.h
new file mode 100644
--- /dev/null
+++ b/Engine/Source/Editor/UMGEditor/Public/Components/AssetThumbnailWidget.h
@@ -0,0 +1,126 @@
+// Copyright Epic Games, Inc. All Rights Reserved.
+
+#pragma once
+
+#include "Components/Widget.h"
+#include "AssetRegistry/AssetData.h"
+#include "AssetThumbnailWidget.generated.h"
+
+class FAssetThumbnail;
+class SBorder;
+class SWidget;
+struct FAssetThumbnailConfig;
+
+UENUM(BlueprintType, meta = (DisplayName = "EThumbnailLabelType"))
+enum class EThumbnailLabelType_BlueprintType : uint8
+{
+	ClassName,
+	AssetName,
+	NoLabel
+};
+
+UENUM(BlueprintType, meta = (DisplayName = "EThumbnailColorStripOrientation"))
+enum class EThumbnailColorStripOrientation_BlueprintType : uint8
+{
+	/** Display the color strip as a horizontal line along the bottom edge */
+	HorizontalBottomEdge,
+	/** Display the color strip as a vertical line along the right edge */
+	VerticalRightEdge,
+};
+
+DECLARE_DYNAMIC_DELEGATE_RetVal(FText, FGetHighlightTextDelegate);
+
+/** Copied mostly from FAssetThumbnailConfig */
+USTRUCT(BlueprintType)
+struct UMGEDITOR_API FAssetThumbnailWidgetSettings
+{
+	GENERATED_BODY()
+	
+	UPROPERTY(EditAnywhere, BlueprintReadWrite, Category = "Appearance")
+	bool bForceGenericThumbnail = false;
+	
+	UPROPERTY(EditAnywhere, BlueprintReadWrite, Category = "Appearance")
+	bool bAllowHintText = true;
+	
+	UPROPERTY(EditAnywhere, BlueprintReadWrite, Category = "Appearance")
+	bool bAllowRealTimeOnHovered = true;
+	
+	UPROPERTY(EditAnywhere, BlueprintReadWrite, Category = "Appearance")
+	bool bAllowAssetSpecificThumbnailOverlay = false;
+	
+	UPROPERTY(EditAnywhere, BlueprintReadWrite, Category = "Appearance")
+	EThumbnailLabelType_BlueprintType ThumbnailLabel = EThumbnailLabelType_BlueprintType::ClassName;
+	
+	UPROPERTY(EditAnywhere, BlueprintReadWrite, Category = "Appearance")
+	FGetHighlightTextDelegate HighlightedTextDelegate;
+	
+	UPROPERTY(EditAnywhere, BlueprintReadWrite, Category = "Appearance", meta = (EditCondition = "bAllowHintText", EditConditionHides))
+	FLinearColor HintColorAndOpacity;
+
+	/** Whether to override the asset type's colour */
+	UPROPERTY(EditAnywhere, BlueprintReadWrite, Category = "Appearance")
+	bool bOverrideAssetTypeColor = false;
+	
+	UPROPERTY(EditAnywhere, BlueprintReadWrite, Category = "Appearance", meta = (EditCondition = "bOverrideAssetTypeColor", EditConditionHides))
+	FLinearColor AssetTypeColorOverride;
+	
+	UPROPERTY(EditAnywhere, BlueprintReadWrite, Category = "Appearance")
+	FMargin Padding;
+	
+	UPROPERTY(EditAnywhere, BlueprintReadWrite, Category = "Appearance")
+	int32 GenericThumbnailSize = 64;
+	
+	UPROPERTY(EditAnywhere, BlueprintReadWrite, Category = "Appearance")
+	EThumbnailColorStripOrientation_BlueprintType ColorStripOrientation = EThumbnailColorStripOrientation_BlueprintType::HorizontalBottomEdge;
+
+	FAssetThumbnailConfig ToThumbnailConfig() const;
+};
+
+/**
+ * 
+ */
+UCLASS()
+class UMGEDITOR_API UAssetThumbnailWidget : public UWidget
+{
+	GENERATED_BODY()
+public:
+
+	UFUNCTION(BlueprintCallable, Category = "Appearance")
+	void SetThumbnailSettings(const FAssetThumbnailWidgetSettings& InThumbnailSettings);
+	FAssetThumbnailWidgetSettings GetThumbnailSettings() const { return ThumbnailSettings; }
+
+	UFUNCTION(BlueprintCallable, Category = "Appearance")
+	void SetAsset(const FAssetData& AssetData);
+	UFUNCTION(BlueprintCallable, Category = "Appearance")
+	void SetAssetByObject(UObject* Object);
+
+	//~ Begin UWidget Interface
+	virtual void SynchronizeProperties() override;
+	virtual void ReleaseSlateResources(bool bReleaseChildren) override;
+	virtual TSharedRef<SWidget> RebuildWidget() override;
+	virtual const FText GetPaletteCategory() override;
+	//~ End UWidget Interface
+	
+private:
+
+	/** The asset of which to show the thumbnail. */
+	UPROPERTY(EditAnywhere, Category = "Appearance")
+	FAssetData AssetToShow;
+
+	/** Desired size of the thumbnail */
+	UPROPERTY(EditAnywhere, Category = "Appearance", meta = (ClampMin = "1", ClampMax = "1024"))
+	FIntPoint Resolution = { 64, 64};
+	
+	/** Behaviour and style of the widget */
+	UPROPERTY(EditAnywhere, Getter, Setter, BlueprintSetter = "SetThumbnailSettings", Category = "Appearance")
+	FAssetThumbnailWidgetSettings ThumbnailSettings;
+
+	/** Creates ThumbnailAssetWidget.  */
+	TSharedPtr<FAssetThumbnail> ThumbnailRenderer;
+	
+	/** Wraps either an error widget or the contents of FAssetThumbnail::MakeThumbnailWidget. */
+	TSharedPtr<SBorder> DisplayedWidget;
+
+	/** Updates or creates ThumbnailRenderer and sets the contents of DisplayedWidget. */
+	void UpdateNativeAssetThumbnailWidget();
+};
