Fix struct initialization missing in FAssetThumbnailWidgetSettings.

diff --git a/Engine/Source/Editor/UMGEditor/Public/Components/AssetThumbnailWidget.h b/Engine/Source/Editor/UMGEditor/Public/Components/AssetThumbnailWidget.h
--- a/Engine/Source/Editor/UMGEditor/Public/Components/AssetThumbnailWidget.h
+++ b/Engine/Source/Editor/UMGEditor/Public/Components/AssetThumbnailWidget.h
@@ -55,14 +55,14 @@ struct UMGEDITOR_API FAssetThumbnailWidgetSettings
 	FGetHighlightTextDelegate HighlightedTextDelegate;
 	
 	UPROPERTY(EditAnywhere, BlueprintReadWrite, Category = "Appearance", meta = (EditCondition = "bAllowHintText", EditConditionHides))
-	FLinearColor HintColorAndOpacity;
+	FLinearColor HintColorAndOpacity = FLinearColor(0.0f, 0.0f, 0.0f, 0.0f);
 
 	/** Whether to override the asset type's colour */
 	UPROPERTY(EditAnywhere, BlueprintReadWrite, Category = "Appearance")
 	bool bOverrideAssetTypeColor = false;
 	
 	UPROPERTY(EditAnywhere, BlueprintReadWrite, Category = "Appearance", meta = (EditCondition = "bOverrideAssetTypeColor", EditConditionHides))
-	FLinearColor AssetTypeColorOverride;
+	FLinearColor AssetTypeColorOverride = FLinearColor(0.0f, 0.0f, 0.0f, 0.0f);
 	
 	UPROPERTY(EditAnywhere, BlueprintReadWrite, Category = "Appearance")
 	FMargin Padding;
