Improve the performance of dragging a widget around in UMG Designer: Prevent the addition/removal of the dragged widget in the hierarchy when the widget is not changing its parent. (reduce lag and invalidation count)

diff --git a/Engine/Source/Editor/UMGEditor/Private/Designer/SDesignerView.cpp b/Engine/Source/Editor/UMGEditor/Private/Designer/SDesignerView.cpp
--- a/Engine/Source/Editor/UMGEditor/Private/Designer/SDesignerView.cpp
+++ b/Engine/Source/Editor/UMGEditor/Private/Designer/SDesignerView.cpp
@@ -2536,10 +2536,35 @@ FReply SDesignerView::OnDragOver(const FGeometry& MyGeometry, const FDragDropEve
 {
 	SDesignSurface::OnDragOver(MyGeometry, DragDropEvent);
 
-	ClearDropPreviews();
-	
 	const bool bIsPreview = true;
-	ProcessDropAndAddWidget(MyGeometry, DragDropEvent, bIsPreview);
+	bool bFoundChangingParent = false;
+	TSharedPtr<FSelectedWidgetDragDropOp> SelectedDragDropOp = DragDropEvent.GetOperationAs<FSelectedWidgetDragDropOp>();
+	if (SelectedDragDropOp.IsValid())
+	{
+		for (const auto& DraggedWidget : SelectedDragDropOp->DraggedWidgets)
+		{
+			if (!DraggedWidget.bStayingInParent)
+			{
+				bFoundChangingParent = true;
+				break;
+			}
+		}
+
+		if (bFoundChangingParent)
+		{
+			ClearDropPreviews();
+			ProcessDropAndAddWidget(MyGeometry, DragDropEvent, bIsPreview);
+		}
+		else
+		{
+			MoveWidgets(MyGeometry, DragDropEvent, bIsPreview, nullptr, false);
+		}
+	}
+	else
+	{
+		ClearDropPreviews();
+		ProcessDropAndAddWidget(MyGeometry, DragDropEvent, bIsPreview);
+	}
 
 	if ( DropPreviews.Num() > 0 )
 	{
@@ -2748,6 +2773,15 @@ void SDesignerView::ProcessDropAndAddWidget(const FGeometry& MyGeometry, const F
 	}
 
 	// Attempt to deal with moving widgets from a drag operation.
+	MoveWidgets(MyGeometry, DragDropEvent, bIsPreview, Target, true);
+}
+
+void SDesignerView::MoveWidgets(const FGeometry& MyGeometry, const FDragDropEvent& DragDropEvent, const bool bIsPreview, UWidget* Target, const bool bAnyWidgetChangingParent)
+{
+	TSharedPtr<FSelectedWidgetDragDropOp> SelectedDragDropOp = DragDropEvent.GetOperationAs<FSelectedWidgetDragDropOp>();
+	TSharedPtr<FDragDropOperation> DragOperation = DragDropEvent.GetOperation();
+	FScopedTransaction DragAndDropTransaction(LOCTEXT("Designer_DragAddDrop", "Drag and Drop Widget"));
+
 	if (SelectedDragDropOp.IsValid() && SelectedDragDropOp->DraggedWidgets.Num() > 0)
 	{
 		SelectedDragDropOp->SetCursorOverride(TOptional<EMouseCursor::Type>());
@@ -2772,23 +2806,46 @@ void SDesignerView::ProcessDropAndAddWidget(const FGeometry& MyGeometry, const F
 			if (DragDropEvent.IsAltDown() && DraggedWidget.bStayingInParent)
 			{
 				DraggedWidget.bStayingInParent = false;
-				if ( SelectedDragDropOp->bShowingMessage )
+				if (SelectedDragDropOp->bShowingMessage)
 				{
 					SelectedDragDropOp->bShowingMessage = false;
 					PopDesignerMessage();
 				}
 			}
 
+			FWidgetHitResult HitResult;
+			FGeometry WidgetUnderCursorGeometry;
+
+			bool bFoundWidgetUnderCursor = FindWidgetUnderCursor(MyGeometry, DragDropEvent, UPanelWidget::StaticClass(), HitResult);
+
 			// If we're staying in the parent we started in, replace the parent found under the cursor with
 			// the original one, also update the arranged widget data so that our layout calculations are accurate.
+			UWidgetBlueprint* BP = GetBlueprint();
+
 			if (DraggedWidget.bStayingInParent)
 			{
+				// If we are not changing parents, keep the widget in the hierarchy but clean the DropPreviews list.
+				if (!bAnyWidgetChangingParent)
+				{
+					DropPreviews.Empty();
+				}
 				WidgetUnderCursorGeometry = GetDesignerGeometry();
 				if (GetWidgetGeometry(DraggedWidget.ParentWidget, WidgetUnderCursorGeometry))
 				{
 					Target = bIsPreview ? DraggedWidget.ParentWidget.GetPreview() : DraggedWidget.ParentWidget.GetTemplate();
 				}
 			}
+			else if (bFoundWidgetUnderCursor)
+			{
+				WidgetUnderCursorGeometry = HitResult.WidgetArranged.Geometry;
+			}
+
+			// If we changed the value of bStayingInParent to false since the last check, remove this widget from the hierarchy and later determine a new parent for it.
+			if (!DraggedWidget.bStayingInParent && !bAnyWidgetChangingParent)
+			{
+				ClearDropPreviews();
+				Target = bIsPreview ? HitResult.Widget.GetPreview() : HitResult.Widget.GetTemplate();
+			}
 
 			FWidgetReference TargetReference = bIsPreview ? BlueprintEditor.Pin()->GetReferenceFromPreview(Target) : BlueprintEditor.Pin()->GetReferenceFromTemplate(Target);
 			BlueprintEditor.Pin()->SetHoveredWidget(TargetReference);
@@ -2830,10 +2887,10 @@ void SDesignerView::ProcessDropAndAddWidget(const FGeometry& MyGeometry, const F
 
 							Widget->SetFlags(RF_Transactional);
 							Widget->Modify();
-						
+
 							ParentWidget->SetFlags(RF_Transactional);
 							ParentWidget->Modify();
-						
+
 						}
 
 						// The Widget originated from a different blueprint, so mark it as modified.
@@ -2880,17 +2937,35 @@ void SDesignerView::ProcessDropAndAddWidget(const FGeometry& MyGeometry, const F
 							bGridSnapX = false;
 						}
 					}
-
 					FVector2D LocalPosition = WidgetUnderCursorGeometry.AbsoluteToLocal(ScreenSpacePosition);
-
 					UPanelSlot* Slot = nullptr;
-					if (bIsChangingParent)
+
+					// Determine if we need to create a new slot or fetch an existing one.
+					// Fetching is much faster, so we want to avoid creating whenever possible.
+					if (bAnyWidgetChangingParent || bIsChangingParent)
 					{
-						Slot = NewParent->AddChild(Widget);
+
+						if (bIsChangingParent)
+						{
+							Slot = NewParent->AddChild(Widget);
+						}
+						else if (UPanelWidget* ParentWidgetAsPanel = Cast<UPanelWidget>(ParentWidget))
+						{
+							Slot = ParentWidgetAsPanel->InsertChildAt(ParentWidgetAsPanel->GetChildIndex(Widget), Widget);
+						}
 					}
-					else if (UPanelWidget* ParentWidgetAsPanel = Cast<UPanelWidget>(ParentWidget))
+					else
 					{
-						Slot = ParentWidgetAsPanel->InsertChildAt(ParentWidgetAsPanel->GetChildIndex(Widget), Widget);
+						if (UPanelWidget* ParentWidgetAsPanel = Cast<UPanelWidget>(ParentWidget))
+						{
+							Slot = Widget->Slot;
+
+							// If we expected to find a slot but it's null, we have to create it.
+							if (Slot == nullptr)
+							{
+								Slot = ParentWidgetAsPanel->AddChild(Widget);
+							}
+						}
 					}
 
 					if (Slot != nullptr)
@@ -2912,9 +2987,23 @@ void SDesignerView::ProcessDropAndAddWidget(const FGeometry& MyGeometry, const F
 								FArrangedWidget ArrangedWidget(SNullWidget::NullWidget, FGeometry());
 								FDesignTimeUtils::GetArrangedWidget(Widget->GetCachedWidget().ToSharedRef(), ArrangedWidget);
 
-								FVector2D Offset = DraggedWidget.DraggedOffset * ArrangedWidget.Geometry.GetLocalSize();
-								FVector2D NewPosition = LocalPosition - Offset;
+								FVector2D Offset = FVector2D::ZeroVector;
+								if (TSharedPtr<SWidget> CachedWidget = Widget->GetCachedWidget())
+								{
+									FDesignTimeUtils::GetArrangedWidget(CachedWidget.ToSharedRef(), ArrangedWidget);
+									Offset = DraggedWidget.DraggedOffset * ArrangedWidget.Geometry.GetLocalSize();
+								}
 
+								FVector2D NewPosition;
+								if (bAnyWidgetChangingParent || bIsChangingParent)
+								{
+									NewPosition = LocalPosition - Offset;
+								}
+								else
+								{
+									NewPosition = LocalPosition - DragDropEvent.GetCursorDelta() - Offset;
+
+								}
 								// Perform grid snapping on X and Y if we need to.
 								if (bGridSnapX)
 								{
@@ -2981,12 +3070,37 @@ FReply SDesignerView::OnDrop(const FGeometry& MyGeometry, const FDragDropEvent&
 	SDesignSurface::OnDrop(MyGeometry, DragDropEvent);
 
 	bMovingExistingWidget = false;
-
-	ClearDropPreviews();
-	
 	const bool bIsPreview = false;
-	ProcessDropAndAddWidget(MyGeometry, DragDropEvent, bIsPreview);
+	bool bFoundChangingParent = false;
 	TSharedPtr<FSelectedWidgetDragDropOp> SelectedDragDropOp = DragDropEvent.GetOperationAs<FSelectedWidgetDragDropOp>();
+
+	if (SelectedDragDropOp.IsValid())
+	{
+		for (const auto& DraggedWidget : SelectedDragDropOp->DraggedWidgets)
+		{
+			if (!DraggedWidget.bStayingInParent)
+			{
+				bFoundChangingParent = true;
+				break;
+			}
+		}
+		if (bFoundChangingParent)
+		{
+			ClearDropPreviews();
+			ProcessDropAndAddWidget(MyGeometry, DragDropEvent, bIsPreview);
+		}
+		else
+		{
+			MoveWidgets(MyGeometry, DragDropEvent, bIsPreview, nullptr, false);
+		}
+	}
+	else
+	{
+		bFoundChangingParent = true;
+		ClearDropPreviews();
+		ProcessDropAndAddWidget(MyGeometry, DragDropEvent, bIsPreview);
+	}
+
 	if (DropPreviews.Num() > 0)
 	{
 		UWidgetBlueprint* BP = GetBlueprint();
diff --git a/Engine/Source/Editor/UMGEditor/Public/Designer/SDesignerView.h b/Engine/Source/Editor/UMGEditor/Public/Designer/SDesignerView.h
--- a/Engine/Source/Editor/UMGEditor/Public/Designer/SDesignerView.h
+++ b/Engine/Source/Editor/UMGEditor/Public/Designer/SDesignerView.h
@@ -284,6 +284,7 @@ private:
 	bool IsRespectingLocks() const;
 
 	void ProcessDropAndAddWidget(const FGeometry& MyGeometry, const FDragDropEvent& DragDropEvent, const bool bIsPreview);
+	void MoveWidgets(const FGeometry& MyGeometry, const FDragDropEvent& DragDropEvent, const bool bIsPreview, UWidget* Target, const bool bAnyWidgetChangingParent);
 
 	FVector2D GetExtensionPosition(TSharedRef<FDesignerSurfaceElement> ExtensionElement) const;
 
