UMG: Adding Desired Focus widget to UserWidget. This new property allow the user widget to forward the focus to one of it's children.
- This functionnality can be used at runtime and in UMG Designer as a default value.

diff --git a/Engine/Source/Editor/UMGEditor/Private/Customizations/WidgetChildTypeCustomization.cpp b/Engine/Source/Editor/UMGEditor/Private/Customizations/WidgetChildTypeCustomization.cpp
--- a/Engine/Source/Editor/UMGEditor/Private/Customizations/WidgetChildTypeCustomization.cpp
+++ b/Engine/Source/Editor/UMGEditor/Private/Customizations/WidgetChildTypeCustomization.cpp
@@ -51,7 +51,7 @@ void FWidgetChildTypeCustomization::CustomizeHeader(TSharedRef<IPropertyHandle>
 	];
 }
 
-void FWidgetChildTypeCustomization::SetDesiredFocusWidgetChild(UUserWidget* OwnerUserWidget, const FWidgetChild& WidgetChild)
+void FWidgetChildTypeCustomization::SetWidgetChild(UUserWidget* OwnerUserWidget, FName WidgetChildName)
 {
 	ensure(OwnerUserWidget);
 
@@ -64,9 +64,23 @@ void FWidgetChildTypeCustomization::SetDesiredFocusWidgetChild(UUserWidget* Owne
 			PropertyHandlePinnedPtr->AccessRawData(RawData);
 			FWidgetChild* PreviousWidgetChild = reinterpret_cast<FWidgetChild*>(RawData[0]);
 
-			OwnerUserWidget->Modify();
+			// When setting a WidgetChild we know it will affect the preview and the CDO. Mark them transactional.
+			{
+				OwnerUserWidget->SetFlags(RF_Transactional);
+				OwnerUserWidget->Modify();
+
+				if (UWidgetBlueprintGeneratedClass* BGClass = OwnerUserWidget->GetWidgetTreeOwningClass())
+				{
+					if (UUserWidget* UserWidgetCDO = BGClass->GetDefaultObject<UUserWidget>())
+					{
+						UserWidgetCDO->SetFlags(RF_Transactional);
+						UserWidgetCDO->Modify();
+					}
+				}
+			}
 
 			FString TextValue;
+			FWidgetChild WidgetChild(OwnerUserWidget, WidgetChildName);
 			StructProperty->Struct->ExportText(TextValue, &WidgetChild, PreviousWidgetChild, OwnerUserWidget, EPropertyPortFlags::PPF_None, nullptr);
 			ensure(PropertyHandlePinnedPtr->SetValueFromFormattedString(TextValue, EPropertyValueSetFlags::DefaultFlags) == FPropertyAccess::Result::Success);
 		}
@@ -84,16 +98,19 @@ void FWidgetChildTypeCustomization::OnWidgetSelectionChanged(FName SelectedName,
 		{
 			if (UUserWidget* UserWidget = Cast<UUserWidget>(OuterObject))
 			{
-				const FScopedTransaction Transaction(LOCTEXT("SetDesiredFocus", "Set Desired Focus"));
+				const FScopedTransaction Transaction(
+					FText::Format(
+						LOCTEXT("SetWidgetChildProperty", "Set {0}"),
+						PropertyHandlePinnedPtr->GetPropertyDisplayName()));
 				if (UWidgetBlueprint* WidgetBlueprint = Editor.Pin()->GetWidgetBlueprintObj())
 				{
 					if (SelectedName == WidgetBlueprint->GetFName())
 					{
-						SetDesiredFocusWidgetChild(UserWidget, FWidgetChild());
+						SetWidgetChild(UserWidget, FName());
 						break;
 					}
 				}
-				SetDesiredFocusWidgetChild(UserWidget, FWidgetChild(UserWidget, SelectedName));
+				SetWidgetChild(UserWidget, SelectedName);
 				break;
 			}
 		}
@@ -175,7 +192,7 @@ FText FWidgetChildTypeCustomization::GetCurrentValueText() const
 			PropertyHandle->AccessRawData(RawData);
 			if (FWidgetChild* ChildWidget = reinterpret_cast<FWidgetChild*>(RawData[0]))
 			{
-				FName WidgetName = ChildWidget->GetChildName();
+				FName WidgetName = ChildWidget->GetFName();
 				return WidgetName.IsNone()? LOCTEXT("SelfText", "Self") : FText::FromName(WidgetName);
 			}
 		}		
diff --git a/Engine/Source/Editor/UMGEditor/Private/Customizations/WidgetChildTypeCustomization.h b/Engine/Source/Editor/UMGEditor/Private/Customizations/WidgetChildTypeCustomization.h
--- a/Engine/Source/Editor/UMGEditor/Private/Customizations/WidgetChildTypeCustomization.h
+++ b/Engine/Source/Editor/UMGEditor/Private/Customizations/WidgetChildTypeCustomization.h
@@ -35,7 +35,7 @@ public:
 
 private:
 
-	void SetDesiredFocusWidgetChild(UUserWidget* OwnerUserWidget, const FWidgetChild& WidgetChild);
+	void SetWidgetChild(UUserWidget* OwnerUserWidget, FName WidgetChildName);
 
 	void OnWidgetSelectionChanged(FName SelectedName, ESelectInfo::Type SelectionType);
 
diff --git a/Engine/Source/Editor/UMGEditor/Private/Hierarchy/SHierarchyView.cpp b/Engine/Source/Editor/UMGEditor/Private/Hierarchy/SHierarchyView.cpp
--- a/Engine/Source/Editor/UMGEditor/Private/Hierarchy/SHierarchyView.cpp
+++ b/Engine/Source/Editor/UMGEditor/Private/Hierarchy/SHierarchyView.cpp
@@ -302,7 +302,7 @@ FReply SHierarchyView::HandleDeleteSelected()
 {
 	TSet<FWidgetReference> SelectedWidgets = BlueprintEditor.Pin()->GetSelectedWidgets();
 	
-	FWidgetBlueprintEditorUtils::DeleteWidgets(GetBlueprint(), SelectedWidgets);
+	FWidgetBlueprintEditorUtils::DeleteWidgets(BlueprintEditor.Pin().ToSharedRef(), GetBlueprint(), SelectedWidgets);
 
 	return FReply::Handled();
 }
diff --git a/Engine/Source/Editor/UMGEditor/Private/WidgetBlueprintCompiler.cpp b/Engine/Source/Editor/UMGEditor/Private/WidgetBlueprintCompiler.cpp
--- a/Engine/Source/Editor/UMGEditor/Private/WidgetBlueprintCompiler.cpp
+++ b/Engine/Source/Editor/UMGEditor/Private/WidgetBlueprintCompiler.cpp
@@ -1212,6 +1212,31 @@ void FWidgetBlueprintCompilerContext::OnPostCDOCompiled(const UObject::FPostCDOC
 		WidgetClass->GetDefaultObject<UUserWidget>()->ValidateBlueprint(*WidgetBP->WidgetTree, BlueprintLog);
 		ValidateWidgetAnimations();
 	}
+
+	ValidateDesiredFocusWidgetName();
+}
+
+
+void FWidgetBlueprintCompilerContext::ValidateDesiredFocusWidgetName()
+{
+	if (UWidgetBlueprintGeneratedClass* WidgetClass = NewWidgetBlueprintClass)
+	{
+		UWidgetBlueprint* WidgetBP = WidgetBlueprint();
+		UUserWidget* UserWidgetCDO = WidgetClass->GetDefaultObject<UUserWidget>();
+		if (WidgetBP && UserWidgetCDO)
+		{
+			UWidgetTree* LatestWidgetTree = FWidgetBlueprintEditorUtils::FindLatestWidgetTree(WidgetBP, UserWidgetCDO);
+			FName DesiredFocusWidgetName = UserWidgetCDO->GetDesiredFocusWidgetName();
+
+			if (!DesiredFocusWidgetName.IsNone() && !LatestWidgetTree->FindWidget(DesiredFocusWidgetName))
+			{
+				FBlueprintCompilerLog BlueprintLog(MessageLog, WidgetClass);
+				// Notify that the desired focus widget is not found in the Widget tree, so it's invalid.
+				const FText InvalidDesiredFocusWidgetNameMessage = LOCTEXT("InvalidDesiredFocusWidgetName", "User Widget '{0}' Desired Focus is set to a non-existent widget '{1}'. Select a valid desired focus for this User Widget.");
+				BlueprintLog.Warning(FText::Format(InvalidDesiredFocusWidgetNameMessage, FText::FromString(UserWidgetCDO->GetClass()->GetName()), FText::FromName(DesiredFocusWidgetName)));
+			}
+		}
+	}
 }
 
 void FWidgetBlueprintCompilerContext::EnsureProperGeneratedClass(UClass*& TargetUClass)
diff --git a/Engine/Source/Editor/UMGEditor/Private/WidgetBlueprintEditor.cpp b/Engine/Source/Editor/UMGEditor/Private/WidgetBlueprintEditor.cpp
--- a/Engine/Source/Editor/UMGEditor/Private/WidgetBlueprintEditor.cpp
+++ b/Engine/Source/Editor/UMGEditor/Private/WidgetBlueprintEditor.cpp
@@ -769,7 +769,7 @@ bool FWidgetBlueprintEditor::CanDeleteSelectedWidgets()
 void FWidgetBlueprintEditor::DeleteSelectedWidgets()
 {
 	TSet<FWidgetReference> Widgets = GetSelectedWidgets();
-	FWidgetBlueprintEditorUtils::DeleteWidgets(GetWidgetBlueprintObj(), Widgets);
+	FWidgetBlueprintEditorUtils::DeleteWidgets(SharedThis(this), GetWidgetBlueprintObj(), Widgets);
 
 	// Clear the selection now that the widget has been deleted.
 	TSet<FWidgetReference> Empty;
@@ -797,7 +797,7 @@ bool FWidgetBlueprintEditor::CanCutSelectedWidgets()
 void FWidgetBlueprintEditor::CutSelectedWidgets()
 {
 	TSet<FWidgetReference> Widgets = GetSelectedWidgets();
-	FWidgetBlueprintEditorUtils::CutWidgets(GetWidgetBlueprintObj(), Widgets);
+	FWidgetBlueprintEditorUtils::CutWidgets(SharedThis(this), GetWidgetBlueprintObj(), Widgets);
 }
 
 const UWidgetAnimation* FWidgetBlueprintEditor::RefreshCurrentAnimation()
diff --git a/Engine/Source/Editor/UMGEditor/Private/WidgetBlueprintEditorUtils.cpp b/Engine/Source/Editor/UMGEditor/Private/WidgetBlueprintEditorUtils.cpp
--- a/Engine/Source/Editor/UMGEditor/Private/WidgetBlueprintEditorUtils.cpp
+++ b/Engine/Source/Editor/UMGEditor/Private/WidgetBlueprintEditorUtils.cpp
@@ -224,6 +224,55 @@ bool FWidgetBlueprintEditorUtils::VerifyWidgetRename(TSharedRef<class FWidgetBlu
 	return true;
 }
 
+void FWidgetBlueprintEditorUtils::SetDesiredFocus(TSharedRef<FWidgetBlueprintEditor> BlueprintEditor, const FName DesiredFocusWidgetName)
+{
+	if (UUserWidget* PreviewWidget = BlueprintEditor->GetPreview())
+	{
+		if (UWidgetBlueprintGeneratedClass* BGClass = PreviewWidget->GetWidgetTreeOwningClass())
+		{
+			if (UUserWidget* WidgetCDO = BGClass->GetDefaultObject<UUserWidget>())
+			{
+				WidgetCDO->SetFlags(RF_Transactional);
+				WidgetCDO->Modify();
+				WidgetCDO->SetDesiredFocusWidget(DesiredFocusWidgetName);
+
+				// We need to change the PreviewWidget to make sure the DetailPanel show the right value.
+				PreviewWidget->SetFlags(RF_Transactional);
+				PreviewWidget->Modify();
+				PreviewWidget->SetDesiredFocusWidget(DesiredFocusWidgetName);
+			}
+		}
+	}
+}
+
+
+void FWidgetBlueprintEditorUtils::ReplaceDesiredFocus(TSharedRef<FWidgetBlueprintEditor> BlueprintEditor, const FName& OldName, const FName& NewName)
+{		
+	if (UUserWidget* PreviewWidget = BlueprintEditor->GetPreview())
+	{
+		if (UWidgetBlueprintGeneratedClass* BGClass = PreviewWidget->GetWidgetTreeOwningClass())
+		{
+			if (UUserWidget* WidgetCDO = BGClass->GetDefaultObject<UUserWidget>())
+			{
+				// Verify if the Name changed is the Desired focus Widget name.
+				if (WidgetCDO->GetDesiredFocusWidgetName() == OldName)
+				{
+					WidgetCDO->SetFlags(RF_Transactional);
+					WidgetCDO->Modify();
+					WidgetCDO->SetDesiredFocusWidget(NewName);
+
+					ensure(PreviewWidget->GetDesiredFocusWidgetName() == OldName);
+
+					// We need to change the PreviewWidget to make sure the DetailPanel show the right value.
+					PreviewWidget->SetFlags(RF_Transactional);
+					PreviewWidget->Modify();
+					PreviewWidget->SetDesiredFocusWidget(NewName);
+				}
+			}
+		}
+	}
+}
+
 bool FWidgetBlueprintEditorUtils::RenameWidget(TSharedRef<FWidgetBlueprintEditor> BlueprintEditor, const FName& OldObjectName, const FString& NewDisplayName)
 {
 	UWidgetBlueprint* Blueprint = BlueprintEditor->GetWidgetBlueprintObj();
@@ -321,6 +370,9 @@ bool FWidgetBlueprintEditorUtils::RenameWidget(TSharedRef<FWidgetBlueprintEditor
 		// Update Variable References and
 		// Update Event References to member variables
 		FBlueprintEditorUtils::ReplaceVariableReferences(Blueprint, OldObjectName, NewFName);
+
+		// Replace the Desired focus Widget name if it match the renamed widget
+		ReplaceDesiredFocus(BlueprintEditor, OldObjectName, NewFName);
 		
 		// Find and update all binding references in the widget blueprint
 		for ( FDelegateEditorBinding& Binding : Blueprint->Bindings )
@@ -459,7 +511,7 @@ bool FWidgetBlueprintEditorUtils::CanOpenSelectedWidgetsForEdit( TSet<FWidgetRef
 	return bCanOpenAllForEdit;
 }
 
-void FWidgetBlueprintEditorUtils::DeleteWidgets(UWidgetBlueprint* Blueprint, TSet<FWidgetReference> Widgets, bool bSilentDelete /*=false*/)
+void FWidgetBlueprintEditorUtils::DeleteWidgets(TSharedRef<FWidgetBlueprintEditor> BlueprintEditor, UWidgetBlueprint* Blueprint, TSet<FWidgetReference> Widgets, bool bSilentDelete /*=false*/)
 {
 	if ( Widgets.Num() > 0 )
 	{
@@ -520,6 +572,9 @@ void FWidgetBlueprintEditorUtils::DeleteWidgets(UWidgetBlueprint* Blueprint, TSe
 				FBlueprintEditorUtils::RemoveVariableNodes(Blueprint, WidgetTemplate->GetFName());
 			}
 
+			// Rename the Desired Focus that fit the Widget Deleted
+			ReplaceDesiredFocus(BlueprintEditor, WidgetTemplate->GetFName(), FName());
+
 			// Rename the removed widget to the transient package so that it doesn't conflict with future widgets sharing the same name.
 			WidgetTemplate->Rename(nullptr, GetTransientPackage());
 
@@ -981,6 +1036,24 @@ void FWidgetBlueprintEditorUtils::BuildReplaceWithMenu(FMenuBuilder& Menu, TShar
 	Menu.EndSection();
 }
 
+
+bool FWidgetBlueprintEditorUtils::IsDesiredFocusWiget(TSharedRef<FWidgetBlueprintEditor> BlueprintEditor, UWidget* Widget)
+{
+	// Verify if the Name changed is the Desired focus Widget name.
+	if (UUserWidget* PrevieWidget = BlueprintEditor->GetPreview())
+	{
+		if (UWidgetBlueprintGeneratedClass* BGClass = PrevieWidget->GetWidgetTreeOwningClass())
+		{
+			UUserWidget* WidgetCDO = BGClass->GetDefaultObject<UUserWidget>();
+			if (WidgetCDO && WidgetCDO->GetDesiredFocusWidgetName() == Widget->GetFName())
+			{
+				return true;
+			}
+		}
+	}
+	return false;
+}
+
 void FWidgetBlueprintEditorUtils::ReplaceWidgetWithSelectedTemplate(TSharedRef<FWidgetBlueprintEditor> BlueprintEditor, UWidgetBlueprint* BP, FWidgetReference Widget)
 {
 	// @Todo: Needs to deal with bound object in animation tracks
@@ -988,7 +1061,7 @@ void FWidgetBlueprintEditorUtils::ReplaceWidgetWithSelectedTemplate(TSharedRef<F
 	const FScopedTransaction Transaction(LOCTEXT("ReplaceWidgets", "Replace Widgets"));
 	bool bIsUserWidget = false;
 
-	UWidget* ThisWidget = Widget.GetTemplate();
+	UWidget* WidgetToReplace = Widget.GetTemplate();
 	UWidget* NewReplacementWidget;
 	if (BlueprintEditor->GetSelectedTemplate().IsValid())
 	{
@@ -1011,39 +1084,42 @@ void FWidgetBlueprintEditorUtils::ReplaceWidgetWithSelectedTemplate(TSharedRef<F
 	NewReplacementWidget->SetFlags(RF_Transactional);
 	NewReplacementWidget->Modify();
 
-	if (UPanelWidget* ExistingPanel = Cast<UPanelWidget>(ThisWidget))
+	if (UPanelWidget* ExistingPanel = Cast<UPanelWidget>(WidgetToReplace))
 	{
 		// if they are both panel widgets then call the existing replace function
 		UPanelWidget* ReplacementPanelWidget = Cast<UPanelWidget>(NewReplacementWidget);
 		if (ReplacementPanelWidget)
 		{
-			TSet<FWidgetReference> WidgetToReplace;
-			WidgetToReplace.Add(Widget);
-			ReplaceWidgets(BlueprintEditor, BP, WidgetToReplace, ReplacementPanelWidget->GetClass());
+			TSet<FWidgetReference> ChildWidgetToReplace;
+			ChildWidgetToReplace.Add(Widget);
+			ReplaceWidgets(BlueprintEditor, BP, ChildWidgetToReplace, ReplacementPanelWidget->GetClass());
 			return;
 		}
 	}
-	ThisWidget->SetFlags(RF_Transactional);
-	ThisWidget->Modify();
+	WidgetToReplace->SetFlags(RF_Transactional);
+	WidgetToReplace->Modify();
 
 	BP->WidgetTree->SetFlags(RF_Transactional);
 	BP->WidgetTree->Modify();
 
+	// We need to check before replacing because the Widget might be deleted, reseting the DesiredFocus
+	bool bReplacingDesiredFocus = IsDesiredFocusWiget(BlueprintEditor, WidgetToReplace);
+
 	// Look if the Widget to replace is a NamedSlot.
-	if (TScriptInterface<INamedSlotInterface> NamedSlotHost = FindNamedSlotHostForContent(ThisWidget, BP->WidgetTree))
+	if (TScriptInterface<INamedSlotInterface> NamedSlotHost = FindNamedSlotHostForContent(WidgetToReplace, BP->WidgetTree))
 	{
-		ReplaceNamedSlotHostContent(ThisWidget, NamedSlotHost, NewReplacementWidget);
+		ReplaceNamedSlotHostContent(WidgetToReplace, NamedSlotHost, NewReplacementWidget);
 	}
-	else if (UPanelWidget* CurrentParent = ThisWidget->GetParent())
+	else if (UPanelWidget* CurrentParent = WidgetToReplace->GetParent())
 	{
 		CurrentParent->SetFlags(RF_Transactional);
 		CurrentParent->Modify();
-		CurrentParent->ReplaceChild(ThisWidget, NewReplacementWidget);
+		CurrentParent->ReplaceChild(WidgetToReplace, NewReplacementWidget);
 
-		FString ReplaceName = ThisWidget->GetName();
-		bool bIsGeneratedName = ThisWidget->IsGeneratedName();
+		FString ReplaceName = WidgetToReplace->GetName();
+		bool bIsGeneratedName = WidgetToReplace->IsGeneratedName();
 		// Rename the removed widget to the transient package so that it doesn't conflict with future widgets sharing the same name.
-		ThisWidget->Rename(nullptr, GetTransientPackage());
+		WidgetToReplace->Rename(nullptr, GetTransientPackage());
 
 		// Rename the new Widget to maintain the current name if it's not a generic name
 		if (!bIsGeneratedName)
@@ -1052,7 +1128,7 @@ void FWidgetBlueprintEditorUtils::ReplaceWidgetWithSelectedTemplate(TSharedRef<F
 			NewReplacementWidget->Rename(*ReplaceName, BP->WidgetTree);
 		}
 	}
-	else if (ThisWidget == BP->WidgetTree->RootWidget)
+	else if (WidgetToReplace == BP->WidgetTree->RootWidget)
 	{
 		BP->WidgetTree->RootWidget = NewReplacementWidget;
 	}
@@ -1064,7 +1140,12 @@ void FWidgetBlueprintEditorUtils::ReplaceWidgetWithSelectedTemplate(TSharedRef<F
 	// Delete the widget that has been replaced
 	TSet<FWidgetReference> WidgetsToDelete;
 	WidgetsToDelete.Add(Widget);
-	DeleteWidgets(BP, WidgetsToDelete);
+	DeleteWidgets(BlueprintEditor, BP, WidgetsToDelete);
+
+	if (bReplacingDesiredFocus)
+	{
+		FWidgetBlueprintEditorUtils::SetDesiredFocus(BlueprintEditor, NewReplacementWidget->GetFName());
+	}
 
 	FBlueprintEditorUtils::MarkBlueprintAsStructurallyModified(BP);
 }
@@ -1182,7 +1263,7 @@ void FWidgetBlueprintEditorUtils::ReplaceWidgetWithChildren(TSharedRef<FWidgetBl
 
 		// Delete the widget that has been replaced
 		const bool bForceDelete = true;
-		DeleteWidgets(BP, WidgetsToDelete, bForceDelete);
+		DeleteWidgets(BlueprintEditor, BP, WidgetsToDelete, bForceDelete);
 
 		FBlueprintEditorUtils::MarkBlueprintAsStructurallyModified(BP);
 	}
@@ -1237,7 +1318,7 @@ void FWidgetBlueprintEditorUtils::ReplaceWidgetWithNamedSlot(TSharedRef<FWidgetB
 		}
 
 		// Remove the widget replaced
-		DeleteWidgets(BP, {Widget});
+		DeleteWidgets(BlueprintEditor, BP, {Widget});
 
 		FBlueprintEditorUtils::MarkBlueprintAsStructurallyModified(BP);
 	}
@@ -1267,22 +1348,22 @@ void FWidgetBlueprintEditorUtils::ReplaceWidgets(TSharedRef<FWidgetBlueprintEdit
 
 		UPanelWidget* NewReplacementWidget = CastChecked<UPanelWidget>(Template->Create(BP->WidgetTree));
 
-		UWidget* ThisWidget = Item.GetTemplate();
-		ThisWidget->SetFlags(RF_Transactional);
-		ThisWidget->Modify();
+		UWidget* WidgetToReplace = Item.GetTemplate();
+		WidgetToReplace->SetFlags(RF_Transactional);
+		WidgetToReplace->Modify();
 
 		// Look if the Widget to replace is a NamedSlot.
-		if (TScriptInterface<INamedSlotInterface> NamedSlotHost = FindNamedSlotHostForContent(ThisWidget, BP->WidgetTree))
+		if (TScriptInterface<INamedSlotInterface> NamedSlotHost = FindNamedSlotHostForContent(WidgetToReplace, BP->WidgetTree))
 		{
-			ReplaceNamedSlotHostContent(ThisWidget, NamedSlotHost, NewReplacementWidget);
+			ReplaceNamedSlotHostContent(WidgetToReplace, NamedSlotHost, NewReplacementWidget);
 		}
-		else if (UPanelWidget* CurrentParent = ThisWidget->GetParent())
+		else if (UPanelWidget* CurrentParent = WidgetToReplace->GetParent())
 		{
 			CurrentParent->SetFlags(RF_Transactional);
 			CurrentParent->Modify();
-			CurrentParent->ReplaceChild(ThisWidget, NewReplacementWidget);
+			CurrentParent->ReplaceChild(WidgetToReplace, NewReplacementWidget);
 		}
-		else if (ThisWidget == BP->WidgetTree->RootWidget)
+		else if (WidgetToReplace == BP->WidgetTree->RootWidget)
 		{
 			BP->WidgetTree->RootWidget = NewReplacementWidget;
 		}
@@ -1294,7 +1375,7 @@ void FWidgetBlueprintEditorUtils::ReplaceWidgets(TSharedRef<FWidgetBlueprintEdit
 		NewReplacementWidget->SetFlags(RF_Transactional);
 		NewReplacementWidget->Modify();
 
-		if (UPanelWidget* ExistingPanel = Cast<UPanelWidget>(ThisWidget))
+		if (UPanelWidget* ExistingPanel = Cast<UPanelWidget>(WidgetToReplace))
 		{
 			while (ExistingPanel->GetChildrenCount() > 0)
 			{
@@ -1306,12 +1387,15 @@ void FWidgetBlueprintEditorUtils::ReplaceWidgets(TSharedRef<FWidgetBlueprintEdit
 			}
 		}
 
-		FString ReplaceName = ThisWidget->GetName();
-		const bool bIsGeneratedName = ThisWidget->IsGeneratedName();
+		// We need to check before replacing because the Widget might be deleted, reseting the DesiredFocus
+		bool bReplacingDesiredFocus = IsDesiredFocusWiget(BlueprintEditor, WidgetToReplace);
+
+		FString ReplaceName = WidgetToReplace->GetName();
+		const bool bIsGeneratedName = WidgetToReplace->IsGeneratedName();
 
 		// Delete the widget that has been replaced
 		const bool bForceDelete = true;
-		DeleteWidgets(BP, {Item}, bForceDelete);
+		DeleteWidgets(BlueprintEditor, BP, {Item}, bForceDelete);
 
 		// Rename the new Widget to maintain the current name if it's not a generic name
 		if (!bIsGeneratedName)
@@ -1319,15 +1403,20 @@ void FWidgetBlueprintEditorUtils::ReplaceWidgets(TSharedRef<FWidgetBlueprintEdit
 			ReplaceName = FindNextValidName(BP->WidgetTree, ReplaceName);
 			NewReplacementWidget->Rename(*ReplaceName, BP->WidgetTree);
 		}
+
+		if (bReplacingDesiredFocus)
+		{
+			FWidgetBlueprintEditorUtils::SetDesiredFocus(BlueprintEditor, NewReplacementWidget->GetFName());
+		}
 	}
 
 	FBlueprintEditorUtils::MarkBlueprintAsStructurallyModified(BP);
 }
 
-void FWidgetBlueprintEditorUtils::CutWidgets(UWidgetBlueprint* BP, TSet<FWidgetReference> Widgets)
+void FWidgetBlueprintEditorUtils::CutWidgets(TSharedRef<FWidgetBlueprintEditor> BlueprintEditor, UWidgetBlueprint* BP, TSet<FWidgetReference> Widgets)
 {
 	CopyWidgets(BP, Widgets);
-	DeleteWidgets(BP, Widgets);
+	DeleteWidgets(BlueprintEditor, BP, Widgets);
 }
 
 void FWidgetBlueprintEditorUtils::CopyWidgets(UWidgetBlueprint* BP, TSet<FWidgetReference> Widgets)
@@ -1744,7 +1833,7 @@ TArray<UWidget*> FWidgetBlueprintEditorUtils::PasteWidgetsInternal(TSharedRef<FW
 
 			if (ChildWidgetToDelete)
 			{
-				DeleteWidgets(BP, { BlueprintEditor->GetReferenceFromTemplate(ChildWidgetToDelete) });
+				DeleteWidgets(BlueprintEditor, BP, { BlueprintEditor->GetReferenceFromTemplate(ChildWidgetToDelete) });
 			}
 
 			FBlueprintEditorUtils::MarkBlueprintAsStructurallyModified(BP);
diff --git a/Engine/Source/Editor/UMGEditor/Public/WidgetBlueprintCompiler.h b/Engine/Source/Editor/UMGEditor/Public/WidgetBlueprintCompiler.h
--- a/Engine/Source/Editor/UMGEditor/Public/WidgetBlueprintCompiler.h
+++ b/Engine/Source/Editor/UMGEditor/Public/WidgetBlueprintCompiler.h
@@ -66,6 +66,9 @@ protected:
 	*/
 	void ValidateWidgetAnimations();
 
+	/** Validates the Desired Focus name to make sure it's part of the Widget Tree. */
+	void ValidateDesiredFocusWidgetName();
+
 	// FKismetCompilerContext
 	virtual UEdGraphSchema_K2* CreateSchema() override;
 	virtual void CreateFunctionList() override;
diff --git a/Engine/Source/Editor/UMGEditor/Public/WidgetBlueprintEditorUtils.h b/Engine/Source/Editor/UMGEditor/Public/WidgetBlueprintEditorUtils.h
--- a/Engine/Source/Editor/UMGEditor/Public/WidgetBlueprintEditorUtils.h
+++ b/Engine/Source/Editor/UMGEditor/Public/WidgetBlueprintEditorUtils.h
@@ -32,15 +32,19 @@ public:
 
 	static bool RenameWidget(TSharedRef<class FWidgetBlueprintEditor> BlueprintEditor, const FName& OldObjectName, const FString& NewDisplayName);
 
+	static void ReplaceDesiredFocus(TSharedRef<FWidgetBlueprintEditor> BlueprintEditor, const FName& OldName, const FName& NewName);
+	
+	static void SetDesiredFocus(TSharedRef<FWidgetBlueprintEditor> BlueprintEditor, const FName DesiredFocusWidgetName);
+
 	static void CreateWidgetContextMenu(FMenuBuilder& MenuBuilder, TSharedRef<FWidgetBlueprintEditor> BlueprintEditor, FVector2D TargetLocation);
 
 	static void CopyWidgets(UWidgetBlueprint* BP, TSet<FWidgetReference> Widgets);
 
 	static TArray<UWidget*> PasteWidgets(TSharedRef<FWidgetBlueprintEditor> BlueprintEditor, UWidgetBlueprint* BP, FWidgetReference ParentWidget, FName SlotName, FVector2D PasteLocation);
 
-	static void DeleteWidgets(UWidgetBlueprint* BP, TSet<FWidgetReference> Widgets, bool bSilentDelete = false);
+	static void DeleteWidgets(TSharedRef<FWidgetBlueprintEditor> BlueprintEditor, UWidgetBlueprint* BP, TSet<FWidgetReference> Widgets, bool bSilentDelete = false);
 
-	static void CutWidgets(UWidgetBlueprint* BP, TSet<FWidgetReference> Widgets);
+	static void CutWidgets(TSharedRef<FWidgetBlueprintEditor> BlueprintEditor, UWidgetBlueprint* BP, TSet<FWidgetReference> Widgets);
 
 	static TArray<UWidget*> DuplicateWidgets(TSharedRef<FWidgetBlueprintEditor> BlueprintEditor, UWidgetBlueprint* BP, TSet<FWidgetReference> Widgets);
 
@@ -142,4 +146,6 @@ private:
 	static bool ShouldContinueReplaceOperation(UWidgetBlueprint* BP, const TArray<FText>& WidgetNames);	
 
 	static TOptional<FWidgetThumbnailProperties> DrawSWidgetInRenderTargetInternal(UUserWidget* WidgetInstance, FRenderTarget* RenderTarget2D, UTextureRenderTarget2D* TextureRenderTarget,FVector2D ThumbnailSize, bool bIsForThumbnail, TOptional<FVector2D> ThumbnailCustomSize, EThumbnailPreviewSizeMode ThumbnailSizeMode);
+	
+	static bool IsDesiredFocusWiget(TSharedRef<FWidgetBlueprintEditor> BlueprintEditor, UWidget* Widget);
 };
diff --git a/Engine/Source/Runtime/UMG/Private/UserWidget.cpp b/Engine/Source/Runtime/UMG/Private/UserWidget.cpp
--- a/Engine/Source/Runtime/UMG/Private/UserWidget.cpp
+++ b/Engine/Source/Runtime/UMG/Private/UserWidget.cpp
@@ -25,6 +25,7 @@
 #include "Interfaces/ITargetPlatform.h"
 #include "Blueprint/GameViewportSubsystem.h"
 #include "Blueprint/WidgetBlueprintLibrary.h"
+#include "Blueprint/WidgetChild.h"
 #include "UObject/EditorObjectVersion.h"
 #include "UMGPrivate.h"
 #include "UObject/ObjectSaveContext.h"
@@ -769,6 +770,39 @@ void UUserWidget::PlaySound(USoundBase* SoundToPlay)
 	}
 }
 
+bool UUserWidget::SetDesiredFocusWidget(FName WidgetName)
+{
+	DesiredFocusWidget = FWidgetChild(this, WidgetName);
+	return DesiredFocusWidget.GetWidget() != nullptr;
+}
+
+bool UUserWidget::SetDesiredFocusWidget(UWidget* Widget)
+{
+	if (Widget && WidgetTree)
+	{
+		TArray<UWidget*> AllWidgets;
+		WidgetTree->GetAllWidgets(AllWidgets);
+
+		if (AllWidgets.Contains(Widget))
+		{
+			DesiredFocusWidget = FWidgetChild(this, Widget->GetFName());
+			return DesiredFocusWidget.GetWidget() != nullptr;
+		}
+	}
+	return false;
+}
+
+
+FName UUserWidget::GetDesiredFocusWidgetName() const
+{
+	return DesiredFocusWidget.GetFName();
+}
+
+UWidget* UUserWidget::GetDesiredFocusWidget() const
+{
+	return DesiredFocusWidget.GetWidget();
+}
+
 UWidget* UUserWidget::GetWidgetHandle(TSharedRef<SWidget> InWidget)
 {
 	return WidgetTree->FindWidget(InWidget);
@@ -1253,6 +1287,21 @@ void UUserWidget::PostEditChangeProperty(FPropertyChangedEvent& PropertyChangedE
 {
 	Super::PostEditChangeProperty(PropertyChangedEvent);
 
+	const FName PropertyName = PropertyChangedEvent.GetPropertyName();
+
+	static FName DesiredFocusWidgetPropertyName(GET_MEMBER_NAME_CHECKED(UUserWidget, DesiredFocusWidget));
+	if (PropertyName == DesiredFocusWidgetPropertyName)
+	{
+		if (UWidgetBlueprintGeneratedClass* BGClass = GetWidgetTreeOwningClass())
+		{
+			if (UUserWidget* UserWidgetCDO = BGClass->GetDefaultObject<UUserWidget>())
+			{
+				// We cannot use the Widget Ptr as we need to find the widget with the same name in the CDO
+				UserWidgetCDO->SetDesiredFocusWidget(DesiredFocusWidget.GetFName());
+			}
+		}
+	}
+
 	if ( PropertyChangedEvent.ChangeType != EPropertyChangeType::Interactive )
 	{
 		TSharedPtr<SWidget> SafeWidget = GetCachedWidget();
@@ -1408,6 +1457,8 @@ void UUserWidget::NativePreConstruct()
 			});
 	}
 
+	DesiredFocusWidget.Resolve(WidgetTree);
+
 	PreConstruct(bIsDesignTime);
 }
 
@@ -1793,7 +1844,14 @@ PRAGMA_ENABLE_DEPRECATION_WARNINGS
 
 FReply UUserWidget::NativeOnFocusReceived( const FGeometry& InGeometry, const FFocusEvent& InFocusEvent )
 {
-	return OnFocusReceived( InGeometry, InFocusEvent ).NativeReply;
+	FReply Reply = OnFocusReceived( InGeometry, InFocusEvent ).NativeReply;
+
+	// Forward focus if Desired Focus is set
+	if (UWidget * WidgetToFocus = DesiredFocusWidget.Resolve(WidgetTree))
+	{
+		return FReply::Handled().SetUserFocus(WidgetToFocus->GetCachedWidget().ToSharedRef());
+	}
+	return Reply;
 }
 
 void UUserWidget::NativeOnFocusLost( const FFocusEvent& InFocusEvent )
diff --git a/Engine/Source/Runtime/UMG/Private/WidgetChild.cpp b/Engine/Source/Runtime/UMG/Private/WidgetChild.cpp
--- a/Engine/Source/Runtime/UMG/Private/WidgetChild.cpp
+++ b/Engine/Source/Runtime/UMG/Private/WidgetChild.cpp
@@ -9,28 +9,25 @@
 FWidgetChild::FWidgetChild()
 {}
 
-FWidgetChild::FWidgetChild(const class UUserWidget* Outer, FName InChildName)
-	:ChildName(InChildName)
+FWidgetChild::FWidgetChild(const class UUserWidget* Outer, FName InWidgetName)
+	:WidgetName(InWidgetName)
 {
-	if (!ChildName.IsNone() && Outer)
+	if (!WidgetName.IsNone() && Outer && Outer->WidgetTree)
 	{
-		ChildWidgetPtr = Outer->WidgetTree->FindWidget(ChildName);
+		WidgetPtr = Outer->WidgetTree->FindWidget(WidgetName);
 	}
 }
 
-bool FWidgetChild::IsValid() const 
+UWidget* FWidgetChild::Resolve(const UWidgetTree* WidgetTree)
 {
-	return ChildName.IsNone() || ChildWidgetPtr != nullptr; 
-}
-
-bool FWidgetChild::Resolve(const UWidgetTree* WidgetTree)
-{
-	if (!ChildName.IsNone() && WidgetTree)
+	if (!WidgetName.IsNone() && WidgetTree)
+	{
+		WidgetPtr = WidgetTree->FindWidget(WidgetName);		
+	}
+	else
 	{
-		ChildWidgetPtr = WidgetTree->FindWidget(ChildName);
-		return ChildWidgetPtr != nullptr;
+		WidgetPtr = nullptr;
 	}
 	
-	ChildWidgetPtr = nullptr;
-	return false;
+	return WidgetPtr.Get();
 }
diff --git a/Engine/Source/Runtime/UMG/Public/Blueprint/UserWidget.h b/Engine/Source/Runtime/UMG/Public/Blueprint/UserWidget.h
--- a/Engine/Source/Runtime/UMG/Public/Blueprint/UserWidget.h
+++ b/Engine/Source/Runtime/UMG/Public/Blueprint/UserWidget.h
@@ -3,6 +3,7 @@
 #pragma once
 
 #include "CoreMinimal.h"
+#include "Blueprint/WidgetChild.h"
 #include "UObject/ObjectMacros.h"
 #include "UObject/ObjectSaveContext.h"
 #include "Styling/SlateColor.h"
@@ -1195,6 +1196,28 @@ public:
 	UFUNCTION(BlueprintCallable, BlueprintCosmetic, Category="Sound", meta=( DeprecatedFunction, DeprecationMessage="Use the UGameplayStatics::PlaySound2D instead." ))
 	void PlaySound(class USoundBase* SoundToPlay);
 
+	/** 
+	 * Sets the child Widget that should receive focus when this UserWidget gets focus using it's name. 
+	 *
+	 * @param WidgetName Name of the Widget to forward the focus to when this widget receives focus.
+	 * @return True if the Widget is set properly. Will return false if we can't find a child widget with the specified name.
+	 */
+	bool SetDesiredFocusWidget(FName WidgetName);
+
+	/** 
+	 * Sets the child Widget that should receive focus when this UserWidget gets focus. 
+	 *
+	 * @param Widget Widget to forward the focus to when this widget receives focus
+	 * @return True if the Widget is set properly. Will return false if it's not a child of this UserWidget.
+	 */
+	bool SetDesiredFocusWidget(UWidget* Widget);
+
+	/** @returns The Name of the Widget that should receive focus when this UserWidget gets focus. */
+	FName GetDesiredFocusWidgetName() const;
+
+	/** @returns The Widget that should receive focus when this UserWidget gets focus. */
+	UWidget* GetDesiredFocusWidget() const;
+
 	/** @returns The UObject wrapper for a given SWidget */
 	UWidget* GetWidgetHandle(TSharedRef<SWidget> InWidget);
 
@@ -1472,6 +1495,9 @@ private:
 	UPROPERTY(EditDefaultsOnly, BlueprintReadOnly, Category = "Performance", meta=(AllowPrivateAccess="true"))
 	EWidgetTickFrequency TickFrequency;
 
+	UPROPERTY(EditDefaultsOnly, Category = "Interaction", meta = (AllowPrivateAccess = "true"))
+	FWidgetChild DesiredFocusWidget;
+
 protected:
 	UPROPERTY(Transient, DuplicateTransient)
 	TObjectPtr<class UInputComponent> InputComponent;
diff --git a/Engine/Source/Runtime/UMG/Public/Blueprint/WidgetChild.h b/Engine/Source/Runtime/UMG/Public/Blueprint/WidgetChild.h
--- a/Engine/Source/Runtime/UMG/Public/Blueprint/WidgetChild.h
+++ b/Engine/Source/Runtime/UMG/Public/Blueprint/WidgetChild.h
@@ -22,18 +22,24 @@ public:
 	FWidgetChild();
 	FWidgetChild(const class UUserWidget* Outer, FName InChildName);
 
-	FName GetChildName() const { return ChildName; };
-	UWidget* GetChildWidget() const { return ChildWidgetPtr.Get(); };
+	FName GetFName() const
+	{
+		return WidgetName;
+	};
 
-	bool IsValid() const;
+	UWidget* GetWidget() const
+	{ 
+		return WidgetPtr.Get(); 
+	};
 
-	bool Resolve(const class UWidgetTree* WidgetTree);
+	/** Resolves the Widget ptr using it's name. */
+	UWidget* Resolve(const class UWidgetTree* WidgetTree);
 
 private:
 	/** This either the widget to focus, OR the name of the function to call. */
 	UPROPERTY(EditAnywhere, Category = "Interaction")
-	FName ChildName;
+	FName WidgetName;
 
 	UPROPERTY(Transient)
-	TWeakObjectPtr<UWidget> ChildWidgetPtr;	
+	TWeakObjectPtr<UWidget> WidgetPtr;
 };
