Fix to static analysis error with potential access to nullptr

diff --git a/Engine/Source/Editor/UMGEditor/Private/Customizations/WidgetChildTypeCustomization.cpp b/Engine/Source/Editor/UMGEditor/Private/Customizations/WidgetChildTypeCustomization.cpp
--- a/Engine/Source/Editor/UMGEditor/Private/Customizations/WidgetChildTypeCustomization.cpp
+++ b/Engine/Source/Editor/UMGEditor/Private/Customizations/WidgetChildTypeCustomization.cpp
@@ -53,36 +53,37 @@ void FWidgetChildTypeCustomization::CustomizeHeader(TSharedRef<IPropertyHandle>
 
 void FWidgetChildTypeCustomization::SetWidgetChild(UUserWidget* OwnerUserWidget, FName WidgetChildName)
 {
-	ensure(OwnerUserWidget);
-
-	// We Get the value from the Preview and add it to the CDO.
-	if (TSharedPtr<IPropertyHandle> PropertyHandlePinnedPtr = PropertyHandlePtr.Pin())
+	if (ensure(OwnerUserWidget))
 	{
-		if (FStructProperty* StructProperty = CastField<FStructProperty>(PropertyHandlePinnedPtr->GetProperty()))
+		// We Get the value from the Preview and add it to the CDO.
+		if (TSharedPtr<IPropertyHandle> PropertyHandlePinnedPtr = PropertyHandlePtr.Pin())
 		{
-			TArray<void*> RawData;
-			PropertyHandlePinnedPtr->AccessRawData(RawData);
-			FWidgetChild* PreviousWidgetChild = reinterpret_cast<FWidgetChild*>(RawData[0]);
-
-			// When setting a WidgetChild we know it will affect the preview and the CDO. Mark them transactional.
+			if (FStructProperty* StructProperty = CastField<FStructProperty>(PropertyHandlePinnedPtr->GetProperty()))
 			{
-				OwnerUserWidget->SetFlags(RF_Transactional);
-				OwnerUserWidget->Modify();
+				TArray<void*> RawData;
+				PropertyHandlePinnedPtr->AccessRawData(RawData);
+				FWidgetChild* PreviousWidgetChild = reinterpret_cast<FWidgetChild*>(RawData[0]);
 
-				if (UWidgetBlueprintGeneratedClass* BGClass = OwnerUserWidget->GetWidgetTreeOwningClass())
+				// When setting a WidgetChild we know it will affect the preview and the CDO. Mark them transactional.
 				{
-					if (UUserWidget* UserWidgetCDO = BGClass->GetDefaultObject<UUserWidget>())
+					OwnerUserWidget->SetFlags(RF_Transactional);
+					OwnerUserWidget->Modify();
+
+					if (UWidgetBlueprintGeneratedClass* BGClass = OwnerUserWidget->GetWidgetTreeOwningClass())
 					{
-						UserWidgetCDO->SetFlags(RF_Transactional);
-						UserWidgetCDO->Modify();
+						if (UUserWidget* UserWidgetCDO = BGClass->GetDefaultObject<UUserWidget>())
+						{
+							UserWidgetCDO->SetFlags(RF_Transactional);
+							UserWidgetCDO->Modify();
+						}
 					}
 				}
-			}
 
-			FString TextValue;
-			FWidgetChild WidgetChild(OwnerUserWidget, WidgetChildName);
-			StructProperty->Struct->ExportText(TextValue, &WidgetChild, PreviousWidgetChild, OwnerUserWidget, EPropertyPortFlags::PPF_None, nullptr);
-			ensure(PropertyHandlePinnedPtr->SetValueFromFormattedString(TextValue, EPropertyValueSetFlags::DefaultFlags) == FPropertyAccess::Result::Success);
+				FString TextValue;
+				FWidgetChild WidgetChild(OwnerUserWidget, WidgetChildName);
+				StructProperty->Struct->ExportText(TextValue, &WidgetChild, PreviousWidgetChild, OwnerUserWidget, EPropertyPortFlags::PPF_None, nullptr);
+				ensure(PropertyHandlePinnedPtr->SetValueFromFormattedString(TextValue, EPropertyValueSetFlags::DefaultFlags) == FPropertyAccess::Result::Success);
+			}
 		}
 	}
 }
