MVVM: Remove the specialized viewmodel editor and use the Blueprint base editor instead.

diff --git a/Engine/Plugins/Runtime/ModelViewViewModel/Source/ModelViewViewModel/Private/ViewModel/MVVMViewModelBlueprintGeneratedClass.cpp b/Engine/Plugins/Runtime/ModelViewViewModel/Source/ModelViewViewModel/Private/ViewModel/MVVMViewModelBlueprintGeneratedClass.cpp
--- a/Engine/Plugins/Runtime/ModelViewViewModel/Source/ModelViewViewModel/Private/ViewModel/MVVMViewModelBlueprintGeneratedClass.cpp
+++ b/Engine/Plugins/Runtime/ModelViewViewModel/Source/ModelViewViewModel/Private/ViewModel/MVVMViewModelBlueprintGeneratedClass.cpp
@@ -6,51 +6,3 @@
 
 #include UE_INLINE_GENERATED_CPP_BY_NAME(MVVMViewModelBlueprintGeneratedClass)
 
-UMVVMViewModelBlueprintGeneratedClass::UMVVMViewModelBlueprintGeneratedClass()
-	: FieldNotifyStartBitNumber(INDEX_NONE)
-{}
-
-
-void UMVVMViewModelBlueprintGeneratedClass::PostLoadDefaultObject(UObject* Object)
-{
-	Super::PostLoadDefaultObject(Object);
-	InitializeFieldNotification(Cast<UMVVMViewModelBase>(Object));
-}
-
-
-void UMVVMViewModelBlueprintGeneratedClass::PurgeClass(bool bRecompilingOnLoad)
-{
-	Super::PurgeClass(bRecompilingOnLoad);
-	FieldNotifyNames.Empty();
-}
-
-
-void UMVVMViewModelBlueprintGeneratedClass::InitializeFieldNotification(const UMVVMViewModelBase* ViewModel)
-{
-	FieldNotifyStartBitNumber = 0;
-	if (ViewModel && FieldNotifyNames.Num())
-	{
-		int32 NumberOfField = 0;
-		ViewModel->GetFieldNotificationDescriptor().ForEachField(this, [&NumberOfField](::UE::FieldNotification::FFieldId FielId)
-			{
-				++NumberOfField;
-				return true;
-			});
-		FieldNotifyStartBitNumber = NumberOfField - FieldNotifyNames.Num();
-		ensureMsgf(FieldNotifyStartBitNumber >= 0, TEXT("The FieldNotifyStartIndex is negative. The number of field should be positive."));
-	}
-}
-
-
-void UMVVMViewModelBlueprintGeneratedClass::ForEachField(TFunctionRef<bool(::UE::FieldNotification::FFieldId FieldId)> Callback) const
-{
-	ensureMsgf(FieldNotifyStartBitNumber >= 0, TEXT("The FieldNotifyStartIndex is negative. The number of field should be positive."));
-	for (int32 Index = 0; Index < FieldNotifyNames.Num(); ++Index)
-	{
-		if (!Callback(UE::FieldNotification::FFieldId(FieldNotifyNames[Index].GetFieldName(), Index + FieldNotifyStartBitNumber)))
-		{
-			break;
-		}
-	}
-}
-
diff --git a/Engine/Plugins/Runtime/ModelViewViewModel/Source/ModelViewViewModel/Public/ViewModel/MVVMViewModelBlueprintGeneratedClass.h b/Engine/Plugins/Runtime/ModelViewViewModel/Source/ModelViewViewModel/Public/ViewModel/MVVMViewModelBlueprintGeneratedClass.h
--- a/Engine/Plugins/Runtime/ModelViewViewModel/Source/ModelViewViewModel/Public/ViewModel/MVVMViewModelBlueprintGeneratedClass.h
+++ b/Engine/Plugins/Runtime/ModelViewViewModel/Source/ModelViewViewModel/Public/ViewModel/MVVMViewModelBlueprintGeneratedClass.h
@@ -16,30 +16,12 @@ namespace UE::MVVM
 	class FViewModelBlueprintCompilerContext;
 }//namespace
 
+
+/** Will be deprecated in the next version. */
 UCLASS()
 class MODELVIEWVIEWMODEL_API UMVVMViewModelBlueprintGeneratedClass : public UBlueprintGeneratedClass
 {
 	GENERATED_BODY()
-
-	friend UE::MVVM::FViewModelBlueprintCompilerContext;
-
-public:
-	UMVVMViewModelBlueprintGeneratedClass();
-
-	//~ Begin UBlueprintGeneratedClass interface
-	virtual void PostLoadDefaultObject(UObject* Object) override;
-	virtual void PurgeClass(bool bRecompilingOnLoad) override;
-	//~ End UBlueprintGeneratedClass interface
-
-	void InitializeFieldNotification(const UMVVMViewModelBase* ViewModel);
-	void ForEachField(TFunctionRef<bool(::UE::FieldNotification::FFieldId FieldId)> Callback) const;
-
-private:
-	/** List Field Notifies. No index here on purpose to prevent saving them. */
-	UPROPERTY()
-	TArray<FFieldNotificationId> FieldNotifyNames;
-
-	int32 FieldNotifyStartBitNumber;
 };
 
 #if UE_ENABLE_INCLUDE_ORDER_DEPRECATED_IN_5_2
diff --git a/Engine/Plugins/Runtime/ModelViewViewModel/Source/ModelViewViewModelBlueprint/ModelViewViewModelBlueprint.Build.cs b/Engine/Plugins/Runtime/ModelViewViewModel/Source/ModelViewViewModelBlueprint/ModelViewViewModelBlueprint.Build.cs
--- a/Engine/Plugins/Runtime/ModelViewViewModel/Source/ModelViewViewModelBlueprint/ModelViewViewModelBlueprint.Build.cs
+++ b/Engine/Plugins/Runtime/ModelViewViewModel/Source/ModelViewViewModelBlueprint/ModelViewViewModelBlueprint.Build.cs
@@ -33,7 +33,5 @@ public class ModelViewViewModelBlueprint : ModuleRules
 				"UMGEditor",
 				"UnrealEd",
 			});
-
-		PublicDefinitions.Add("UE_MVVM_WITH_VIEWMODEL_EDITOR=0");
 	}
 }
diff --git a/Engine/Plugins/Runtime/ModelViewViewModel/Source/ModelViewViewModelBlueprint/Private/ModelViewViewModelBlueprintModule.cpp b/Engine/Plugins/Runtime/ModelViewViewModel/Source/ModelViewViewModelBlueprint/Private/ModelViewViewModelBlueprintModule.cpp
--- a/Engine/Plugins/Runtime/ModelViewViewModel/Source/ModelViewViewModelBlueprint/Private/ModelViewViewModelBlueprintModule.cpp
+++ b/Engine/Plugins/Runtime/ModelViewViewModel/Source/ModelViewViewModelBlueprint/Private/ModelViewViewModelBlueprintModule.cpp
@@ -3,36 +3,5 @@
 #include "ModelViewViewModelBlueprintModule.h"
 #include "Modules/ModuleManager.h"
 
-#include "ViewModel/MVVMViewModelBlueprint.h"
-#include "ViewModel/MVVMViewModelBlueprintCompiler.h"
 
-
-class FModelViewViewModelBlueprintModule : public IModuleInterface
-{
-public:
-	virtual void StartupModule() override
-	{
-#if UE_MVVM_WITH_VIEWMODEL_EDITOR
-		IKismetCompilerInterface& KismetCompilerModule = FModuleManager::GetModuleChecked<IKismetCompilerInterface>("KismetCompiler");
-		KismetCompilerModule.GetCompilers().Add(&ViewModelBlueprintCompiler);
-
-		FKismetCompilerContext::RegisterCompilerForBP(UMVVMViewModelBlueprint::StaticClass(), &UMVVMViewModelBlueprint::GetCompilerForViewModelBlueprint);
-#endif
-	}
-
-	virtual void ShutdownModule() override
-	{
-#if UE_MVVM_WITH_VIEWMODEL_EDITOR
-		if (IKismetCompilerInterface* KismetCompilerModule = FModuleManager::GetModulePtr<IKismetCompilerInterface>("KismetCompiler"))
-		{
-			KismetCompilerModule->GetCompilers().Remove(&ViewModelBlueprintCompiler);
-		}
-#endif
-	}
-
-private:
-	UE::MVVM::FViewModelBlueprintCompiler ViewModelBlueprintCompiler;
-};
-
-
-IMPLEMENT_MODULE(FModelViewViewModelBlueprintModule, ModelViewViewModelBlueprint);
+IMPLEMENT_MODULE(FDefaultModuleImpl, ModelViewViewModelBlueprint);
diff --git a/Engine/Plugins/Runtime/ModelViewViewModel/Source/ModelViewViewModelBlueprint/Private/ViewModel/MVVMViewModelBlueprint.cpp b/Engine/Plugins/Runtime/ModelViewViewModel/Source/ModelViewViewModelBlueprint/Private/ViewModel/MVVMViewModelBlueprint.cpp
--- a/Engine/Plugins/Runtime/ModelViewViewModel/Source/ModelViewViewModelBlueprint/Private/ViewModel/MVVMViewModelBlueprint.cpp
+++ b/Engine/Plugins/Runtime/ModelViewViewModel/Source/ModelViewViewModelBlueprint/Private/ViewModel/MVVMViewModelBlueprint.cpp
@@ -1,51 +1,6 @@
 // Copyright Epic Games, Inc. All Rights Reserved.
 
 #include "ViewModel/MVVMViewModelBlueprint.h"
-#include "ViewModel/MVVMViewModelBlueprintCompiler.h"
-#include "ViewModel/MVVMViewModelBlueprintGeneratedClass.h"
 
 
 #include UE_INLINE_GENERATED_CPP_BY_NAME(MVVMViewModelBlueprint)
-
-
-TSharedPtr<FKismetCompilerContext> UMVVMViewModelBlueprint::GetCompilerForViewModelBlueprint(UBlueprint* BP, FCompilerResultsLog& InMessageLog, const FKismetCompilerOptions& InCompileOptions)
-{
-	return TSharedPtr<FKismetCompilerContext>(new UE::MVVM::FViewModelBlueprintCompilerContext(CastChecked<UMVVMViewModelBlueprint>(BP), InMessageLog, InCompileOptions));
-}
-
-
-UMVVMViewModelBlueprintGeneratedClass* UMVVMViewModelBlueprint::GetViewModelBlueprintGeneratedClass() const
-{
-	return Cast<UMVVMViewModelBlueprintGeneratedClass>(*GeneratedClass);
-}
-
-
-UMVVMViewModelBlueprintGeneratedClass* UMVVMViewModelBlueprint::GetViewModelBlueprintSkeletonClass() const
-{
-	return Cast<UMVVMViewModelBlueprintGeneratedClass>(*SkeletonGeneratedClass);
-}
-
-
-bool UMVVMViewModelBlueprint::SupportsMacros() const
-{
-	return false;
-}
-
-
-bool UMVVMViewModelBlueprint::SupportsEventGraphs() const
-{
-	return false;
-}
-
-
-bool UMVVMViewModelBlueprint::SupportsDelegates() const
-{
-	return false;
-}
-
-
-bool UMVVMViewModelBlueprint::SupportsAnimLayers() const
-{
-	return false;
-}
-
diff --git a/Engine/Plugins/Runtime/ModelViewViewModel/Source/ModelViewViewModelBlueprint/Private/ViewModel/MVVMViewModelBlueprintCompiler.cpp b/Engine/Plugins/Runtime/ModelViewViewModel/Source/ModelViewViewModelBlueprint/Private/ViewModel/MVVMViewModelBlueprintCompiler.cpp
deleted file mode 100644
diff --git a/Engine/Plugins/Runtime/ModelViewViewModel/Source/ModelViewViewModelBlueprint/Private/ViewModel/MVVMViewModelBlueprintCompiler.h b/Engine/Plugins/Runtime/ModelViewViewModel/Source/ModelViewViewModelBlueprint/Private/ViewModel/MVVMViewModelBlueprintCompiler.h
deleted file mode 100644
diff --git a/Engine/Plugins/Runtime/ModelViewViewModel/Source/ModelViewViewModelBlueprint/Public/ViewModel/MVVMViewModelBlueprint.h b/Engine/Plugins/Runtime/ModelViewViewModel/Source/ModelViewViewModelBlueprint/Public/ViewModel/MVVMViewModelBlueprint.h
--- a/Engine/Plugins/Runtime/ModelViewViewModel/Source/ModelViewViewModelBlueprint/Public/ViewModel/MVVMViewModelBlueprint.h
+++ b/Engine/Plugins/Runtime/ModelViewViewModel/Source/ModelViewViewModelBlueprint/Public/ViewModel/MVVMViewModelBlueprint.h
@@ -14,33 +14,12 @@ class UMVVMViewModelBlueprintGeneratedClass;
 
 
 /**
- * 
+ * Will be deprecated in the next version.
  */
-UCLASS(BlueprintType, meta=(IgnoreClassThumbnail))
+UCLASS()
 class MODELVIEWVIEWMODELBLUEPRINT_API UMVVMViewModelBlueprint : public UBlueprint
 {
 	GENERATED_BODY()
-
-public:
-	/** */
-	static TSharedPtr<FKismetCompilerContext> GetCompilerForViewModelBlueprint(UBlueprint* BP, FCompilerResultsLog& InMessageLog, const FKismetCompilerOptions& InCompileOptions);
-
-	/** Get the (full) generated class for this viewmodel blueprint */
-	UMVVMViewModelBlueprintGeneratedClass* GetViewModelBlueprintGeneratedClass() const;
-
-	/** Get the (skeleton) generated class for this viewmodel blueprint */
-	UMVVMViewModelBlueprintGeneratedClass* GetViewModelBlueprintSkeletonClass() const;
-
-protected:
-	virtual bool SupportsMacros() const override;
-	virtual bool SupportsEventGraphs() const override;
-	virtual bool SupportsDelegates() const override;
-	virtual bool SupportsAnimLayers() const override;
-
-public:
-	// Use during compilation to clean the automatically generated graph.
-	UPROPERTY(Transient)
-	TArray<TObjectPtr<UEdGraph>> TemporaryGraph;
 };
 
 #if UE_ENABLE_INCLUDE_ORDER_DEPRECATED_IN_5_2
diff --git a/Engine/Plugins/Runtime/ModelViewViewModel/Source/ModelViewViewModelEditor/Private/ModelViewViewModelEditorModule.cpp b/Engine/Plugins/Runtime/ModelViewViewModel/Source/ModelViewViewModelEditor/Private/ModelViewViewModelEditorModule.cpp
--- a/Engine/Plugins/Runtime/ModelViewViewModel/Source/ModelViewViewModelEditor/Private/ModelViewViewModelEditorModule.cpp
+++ b/Engine/Plugins/Runtime/ModelViewViewModel/Source/ModelViewViewModelEditor/Private/ModelViewViewModelEditorModule.cpp
@@ -20,7 +20,6 @@
 #include "Tabs/MVVMBindingSummoner.h"
 #include "Tabs/MVVMViewModelSummoner.h"
 #include "UMGEditorModule.h"
-#include "ViewModel/AssetTypeActions_ViewModelBlueprint.h"
 #include "WidgetBlueprintEditor.h"
 #include "WidgetDrawerConfig.h"
 
@@ -36,16 +35,6 @@ void FModelViewViewModelEditorModule::StartupModule()
 	IUMGEditorModule& UMGEditorModule = FModuleManager::LoadModuleChecked<IUMGEditorModule>("UMGEditor");
 	UMGEditorModule.OnRegisterTabsForEditor().AddRaw(this, &FModelViewViewModelEditorModule::HandleRegisterBlueprintEditorTab);
 
-	// Register asset types
-	{
-#if UE_MVVM_WITH_VIEWMODEL_EDITOR
-		// Only remove what is related to the viewmodel editor, not the UMG extention for view.
-		IAssetTools& AssetTools = FModuleManager::LoadModuleChecked<FAssetToolsModule>("AssetTools").Get();
-		ViewModelBlueprintActions = MakeShared<UE::MVVM::FAssetTypeActions_ViewModelBlueprint>();
-		AssetTools.RegisterAssetTypeActions(ViewModelBlueprintActions.ToSharedRef());
-#endif
-	}
-
 	PropertyBindingExtension = MakeShared<FMVVMPropertyBindingExtension>();
 	UMGEditorModule.GetPropertyBindingExtensibilityManager()->AddExtension(PropertyBindingExtension.ToSharedRef());
 
@@ -91,17 +80,6 @@ void FModelViewViewModelEditorModule::ShutdownModule()
 	}
 	PropertyBindingExtension.Reset();
 
-#if UE_MVVM_WITH_VIEWMODEL_EDITOR
-	// Unregister all the asset types that we registered
-	if (FAssetToolsModule* AssetTools = FModuleManager::GetModulePtr<FAssetToolsModule>("AssetTools"))
-	{
-		if (ViewModelBlueprintActions)
-		{
-			AssetTools->Get().UnregisterAssetTypeActions(ViewModelBlueprintActions.ToSharedRef());
-		}
-	}
-#endif
-
 	FMVVMEditorStyle::DestroyInstance();
 
 	FMVVMEditorCommands::Unregister();
diff --git a/Engine/Plugins/Runtime/ModelViewViewModel/Source/ModelViewViewModelEditor/Private/ViewModel/AssetTypeActions_ViewModelBlueprint.cpp b/Engine/Plugins/Runtime/ModelViewViewModel/Source/ModelViewViewModelEditor/Private/ViewModel/AssetTypeActions_ViewModelBlueprint.cpp
deleted file mode 100644
diff --git a/Engine/Plugins/Runtime/ModelViewViewModel/Source/ModelViewViewModelEditor/Private/ViewModel/AssetTypeActions_ViewModelBlueprint.h b/Engine/Plugins/Runtime/ModelViewViewModel/Source/ModelViewViewModelEditor/Private/ViewModel/AssetTypeActions_ViewModelBlueprint.h
deleted file mode 100644
diff --git a/Engine/Plugins/Runtime/ModelViewViewModel/Source/ModelViewViewModelEditor/Private/ViewModel/MVVMViewModelBlueprintEditor.cpp b/Engine/Plugins/Runtime/ModelViewViewModel/Source/ModelViewViewModelEditor/Private/ViewModel/MVVMViewModelBlueprintEditor.cpp
deleted file mode 100644
diff --git a/Engine/Plugins/Runtime/ModelViewViewModel/Source/ModelViewViewModelEditor/Private/ViewModel/MVVMViewModelBlueprintFactory.cpp b/Engine/Plugins/Runtime/ModelViewViewModel/Source/ModelViewViewModelEditor/Private/ViewModel/MVVMViewModelBlueprintFactory.cpp
deleted file mode 100644
diff --git a/Engine/Plugins/Runtime/ModelViewViewModel/Source/ModelViewViewModelEditor/Private/ViewModel/MVVMViewModelBlueprintFactory.h b/Engine/Plugins/Runtime/ModelViewViewModel/Source/ModelViewViewModelEditor/Private/ViewModel/MVVMViewModelBlueprintFactory.h
deleted file mode 100644
diff --git a/Engine/Plugins/Runtime/ModelViewViewModel/Source/ModelViewViewModelEditor/Public/ViewModel/MVVMViewModelBlueprintEditor.h b/Engine/Plugins/Runtime/ModelViewViewModel/Source/ModelViewViewModelEditor/Public/ViewModel/MVVMViewModelBlueprintEditor.h
deleted file mode 100644
diff --git a/Engine/Plugins/Runtime/ModelViewViewModel/Source/ModelViewViewModelEditor/Public/ViewModel/MVVMViewModelBlueprintToolMenuContext.h b/Engine/Plugins/Runtime/ModelViewViewModel/Source/ModelViewViewModelEditor/Public/ViewModel/MVVMViewModelBlueprintToolMenuContext.h
deleted file mode 100644
diff --git a/Engine/Source/Editor/UMGEditor/Private/UMGEditorModule.cpp b/Engine/Source/Editor/UMGEditor/Private/UMGEditorModule.cpp
--- a/Engine/Source/Editor/UMGEditor/Private/UMGEditorModule.cpp
+++ b/Engine/Source/Editor/UMGEditor/Private/UMGEditorModule.cpp
@@ -92,8 +92,8 @@ public:
 		// Add Customization for variable in Graph editor
 		if (FBlueprintEditorModule* BlueprintEditorModule = FModuleManager::GetModulePtr<FBlueprintEditorModule>("Kismet"))
 		{
-			BlueprintVariableCustomizationHandle = BlueprintEditorModule->RegisterVariableCustomization(FProperty::StaticClass(), FOnGetVariableCustomizationInstance::CreateStatic(&FGraphVariableDetailsCustomization::MakeInstance));
-			BlueprintFunctionCustomizationHandle = BlueprintEditorModule->RegisterFunctionCustomization(UK2Node_FunctionEntry::StaticClass(), FOnGetFunctionCustomizationInstance::CreateStatic(&FGraphFunctionDetailsCustomization::MakeInstance));
+			//BlueprintVariableCustomizationHandle = BlueprintEditorModule->RegisterVariableCustomization(FProperty::StaticClass(), FOnGetVariableCustomizationInstance::CreateStatic(&FGraphVariableDetailsCustomization::MakeInstance));
+			//BlueprintFunctionCustomizationHandle = BlueprintEditorModule->RegisterFunctionCustomization(UK2Node_FunctionEntry::StaticClass(), FOnGetFunctionCustomizationInstance::CreateStatic(&FGraphFunctionDetailsCustomization::MakeInstance));
 		}
 		else
 		{
@@ -152,8 +152,8 @@ public:
 		{
 			if (FBlueprintEditorModule* BlueprintEditorModule = FModuleManager::GetModulePtr<FBlueprintEditorModule>("Kismet"))
 			{
-				BlueprintEditorModule->UnregisterVariableCustomization(FProperty::StaticClass(), BlueprintVariableCustomizationHandle);
-				BlueprintEditorModule->UnregisterFunctionCustomization(UK2Node_FunctionEntry::StaticClass(), BlueprintFunctionCustomizationHandle);
+				//BlueprintEditorModule->UnregisterVariableCustomization(FProperty::StaticClass(), BlueprintVariableCustomizationHandle);
+				//BlueprintEditorModule->UnregisterFunctionCustomization(UK2Node_FunctionEntry::StaticClass(), BlueprintFunctionCustomizationHandle);
 			}
 		}
 
@@ -334,8 +334,8 @@ private:
 			{
 				if (FBlueprintEditorModule* BlueprintEditorModule = FModuleManager::GetModulePtr<FBlueprintEditorModule>(KismetModule))
 				{
-					BlueprintVariableCustomizationHandle = BlueprintEditorModule->RegisterVariableCustomization(FProperty::StaticClass(), FOnGetVariableCustomizationInstance::CreateStatic(&FGraphVariableDetailsCustomization::MakeInstance));
-					BlueprintFunctionCustomizationHandle = BlueprintEditorModule->RegisterFunctionCustomization(UK2Node_FunctionEntry::StaticClass(), FOnGetFunctionCustomizationInstance::CreateStatic(&FGraphFunctionDetailsCustomization::MakeInstance));
+					//BlueprintVariableCustomizationHandle = BlueprintEditorModule->RegisterVariableCustomization(FProperty::StaticClass(), FOnGetVariableCustomizationInstance::CreateStatic(&FGraphVariableDetailsCustomization::MakeInstance));
+					//BlueprintFunctionCustomizationHandle = BlueprintEditorModule->RegisterFunctionCustomization(UK2Node_FunctionEntry::StaticClass(), FOnGetFunctionCustomizationInstance::CreateStatic(&FGraphFunctionDetailsCustomization::MakeInstance));
 				}
 			}
 			FModuleManager::Get().OnModulesChanged().Remove(ModuleChangedHandle);
diff --git a/Engine/Source/Editor/UMGEditor/Private/WidgetBlueprintCompiler.cpp b/Engine/Source/Editor/UMGEditor/Private/WidgetBlueprintCompiler.cpp
--- a/Engine/Source/Editor/UMGEditor/Private/WidgetBlueprintCompiler.cpp
+++ b/Engine/Source/Editor/UMGEditor/Private/WidgetBlueprintCompiler.cpp
@@ -405,7 +405,6 @@ void FWidgetBlueprintCompilerContext::CleanAndSanitizeClass(UBlueprintGeneratedC
 	NewWidgetBlueprintClass->Animations.Empty();
 	NewWidgetBlueprintClass->Bindings.Empty();
 	NewWidgetBlueprintClass->Extensions.Empty();
-	NewWidgetBlueprintClass->FieldNotifyNames.Empty();
 
 	if (UWidgetBlueprintGeneratedClass* WidgetClassToClean = Cast<UWidgetBlueprintGeneratedClass>(ClassToClean))
 	{
@@ -596,16 +595,6 @@ void FWidgetBlueprintCompilerContext::CreateClassVariablesFromBlueprint()
 		{
 			InExtension->CreateClassVariablesFromBlueprint(FCreateVariableContext(*Self));
 		});
-
-	//Add FieldNotifyNames
-	for (TFieldIterator<const FProperty> PropertyIt(NewClass, EFieldIterationFlags::None); PropertyIt; ++PropertyIt)
-	{
-		const FProperty* Property = *PropertyIt;
-		if (Property->HasMetaData(UE::FieldNotification::FCustomizationHelper::MetaData_FieldNotify))
-		{
-			NewWidgetBlueprintClass->FieldNotifyNames.Emplace(Property->GetFName());
-		}
-	}
 }
 
 void FWidgetBlueprintCompilerContext::CopyTermDefaultsToDefaultObject(UObject* DefaultObject)
@@ -1127,11 +1116,6 @@ void FWidgetBlueprintCompilerContext::FinishCompilingClass(UClass* Class)
 	Super::FinishCompilingClass(Class);
 
 	CA_ASSUME(BPGClass);
-	if (UUserWidget* UserWidget = Cast<UUserWidget>(BPGClass->GetDefaultObject()))
-	{
-		BPGClass->InitializeFieldNotification(UserWidget);
-	}
-
 	UWidgetBlueprintExtension::ForEachExtension(WidgetBlueprint(), [BPGClass](UWidgetBlueprintExtension* InExtension)
 		{
 			InExtension->FinishCompilingClass(BPGClass);
@@ -1300,25 +1284,6 @@ void FWidgetBlueprintCompilerContext::VerifyEventReplysAreNotEmpty(FKismetFuncti
 	}
 }
 
-void FWidgetBlueprintCompilerContext::PostcompileFunction(FKismetFunctionContext& Context)
-{
-	Super::PostcompileFunction(Context);
-
-	VerifyFieldNotifyFunction(Context);
-}
-
-void FWidgetBlueprintCompilerContext::VerifyFieldNotifyFunction(FKismetFunctionContext& Context)
-{
-	if (Context.Function && Context.Function->HasMetaData(UE::FieldNotification::FCustomizationHelper::MetaData_FieldNotify))
-	{
-		if (!UE::FieldNotification::Helpers::IsValidAsField(Context.Function))
-		{
-			MessageLog.Error(*LOCTEXT("FieldNotify_IsEventGraph", "Function @@ cannot be a FieldNotify. A function needs to be const, returns a single properties, has no inputs, not be an event or a net function.").ToString(), Context.EntryPoint);
-		}
-		NewWidgetBlueprintClass->FieldNotifyNames.Emplace(Context.Function->GetFName());
-	}
-}
-
 bool FWidgetBlueprintCompilerContext::ValidateGeneratedClass(UBlueprintGeneratedClass* Class)
 {
 	const bool bSuperResult = Super::ValidateGeneratedClass(Class);
diff --git a/Engine/Source/Editor/UMGEditor/Private/Widgets/SFieldNotificationGraphPin.cpp b/Engine/Source/Editor/UMGEditor/Private/Widgets/SFieldNotificationGraphPin.cpp
--- a/Engine/Source/Editor/UMGEditor/Private/Widgets/SFieldNotificationGraphPin.cpp
+++ b/Engine/Source/Editor/UMGEditor/Private/Widgets/SFieldNotificationGraphPin.cpp
@@ -11,6 +11,7 @@
 #include "EdGraphSchema_K2.h"
 #include "Engine/Blueprint.h"
 #include "K2Node.h"
+#include "K2Node_CallFunction.h"
 #include "Kismet2/BlueprintEditorUtils.h"
 #include "Misc/AssertionMacros.h"
 #include "Misc/OutputDeviceRedirector.h"
@@ -107,6 +108,17 @@ void SFieldNotificationGraphPin::Construct(const FArguments& InArgs, UEdGraphPin
 TSharedRef<SWidget> SFieldNotificationGraphPin::GetDefaultValueWidget()
 {
 	UEdGraphPin* SelfPin = GraphPinObj->GetOwningNode()->FindPin(UEdGraphSchema_K2::PSC_Self);
+	if (UK2Node_CallFunction* CallFunction = Cast<UK2Node_CallFunction>(GraphPinObj->GetOwningNode()))
+	{
+		if (UFunction* Function = CallFunction->GetTargetFunction())
+		{
+			const FString& PinName = Function->GetMetaData("FieldNotifyInterfaceParam");
+			if (PinName.Len() != 0)
+			{
+				SelfPin = GraphPinObj->GetOwningNode()->FindPin(*PinName);
+			}
+		}
+	}
 
 	return SNew(SFieldNotificationPicker)
 		.Value(this, &SFieldNotificationGraphPin::GetValue)
diff --git a/Engine/Source/Editor/UMGEditor/Public/WidgetBlueprintCompiler.h b/Engine/Source/Editor/UMGEditor/Public/WidgetBlueprintCompiler.h
--- a/Engine/Source/Editor/UMGEditor/Public/WidgetBlueprintCompiler.h
+++ b/Engine/Source/Editor/UMGEditor/Public/WidgetBlueprintCompiler.h
@@ -75,7 +75,6 @@ protected:
 	virtual void SpawnNewClass(const FString& NewClassName) override;
 	virtual void OnNewClassSet(UBlueprintGeneratedClass* ClassToUse) override;
 	virtual void PrecompileFunction(FKismetFunctionContext& Context, EInternalCompilerFlags InternalFlags) override;
-	virtual void PostcompileFunction(FKismetFunctionContext& Context) override;
 	virtual void CleanAndSanitizeClass(UBlueprintGeneratedClass* ClassToClean, UObject*& InOutOldCDO) override;
 	virtual void SaveSubObjectsFromCleanAndSanitizeClass(FSubobjectCollection& SubObjectsToSave, UBlueprintGeneratedClass* ClassToClean) override;
 	virtual void EnsureProperGeneratedClass(UClass*& TargetClass) override;
diff --git a/Engine/Source/Runtime/UMG/Private/Components/Widget.cpp b/Engine/Source/Runtime/UMG/Private/Components/Widget.cpp
--- a/Engine/Source/Runtime/UMG/Private/Components/Widget.cpp
+++ b/Engine/Source/Runtime/UMG/Private/Components/Widget.cpp
@@ -172,9 +172,9 @@ void UWidget::FFieldNotificationClassDescriptor::ForEachField(const UClass* Clas
 			return;
 		}
 	}
-	if (const UBlueprintGeneratedClass* ViewModelBPClass = Cast<const UBlueprintGeneratedClass>(Class))
+	if (const UBlueprintGeneratedClass* BPClass = Cast<const UBlueprintGeneratedClass>(Class))
 	{
-		ViewModelBPClass->ForEachFieldNotify(Callback, true);
+		BPClass->ForEachFieldNotify(Callback, true);
 	}
 }
 UE_FIELD_NOTIFICATION_IMPLEMENT_CLASS_DESCRIPTOR_ThreeFields(UWidget, ToolTipText, Visibility, bIsEnabled);
diff --git a/Engine/Source/Runtime/UMG/Private/WidgetBlueprintGeneratedClass.cpp b/Engine/Source/Runtime/UMG/Private/WidgetBlueprintGeneratedClass.cpp
--- a/Engine/Source/Runtime/UMG/Private/WidgetBlueprintGeneratedClass.cpp
+++ b/Engine/Source/Runtime/UMG/Private/WidgetBlueprintGeneratedClass.cpp
@@ -120,7 +120,6 @@ namespace
 // UWidgetBlueprintGeneratedClass
 
 UWidgetBlueprintGeneratedClass::UWidgetBlueprintGeneratedClass()
-	: FieldNotifyStartBitNumber(INDEX_NONE)
 {
 #if WITH_EDITORONLY_DATA
 	{
@@ -367,12 +366,6 @@ void UWidgetBlueprintGeneratedClass::PostLoad()
 #endif
 }
 
-void UWidgetBlueprintGeneratedClass::PostLoadDefaultObject(UObject* Object)
-{
-	Super::PostLoadDefaultObject(Object);
-	InitializeFieldNotification(Cast<UUserWidget>(Object));
-}
-
 void UWidgetBlueprintGeneratedClass::PurgeClass(bool bRecompilingOnLoad)
 {
 	Super::PurgeClass(bRecompilingOnLoad);
@@ -396,7 +389,6 @@ void UWidgetBlueprintGeneratedClass::PurgeClass(bool bRecompilingOnLoad)
 
 	Animations.Empty();
 	Bindings.Empty();
-	FieldNotifyNames.Empty();
 }
 
 bool UWidgetBlueprintGeneratedClass::NeedsLoadForServer() const
@@ -489,41 +481,6 @@ UWidgetBlueprintGeneratedClass* UWidgetBlueprintGeneratedClass::FindWidgetTreeOw
 	return nullptr;
 }
 
-void UWidgetBlueprintGeneratedClass::InitializeFieldNotification(const UUserWidget* UserWidget)
-{
-	FieldNotifyStartBitNumber = 0;
-	if (UserWidget && FieldNotifyNames.Num())
-	{
-		int32 NumberOfField = 0;
-		UserWidget->GetFieldNotificationDescriptor().ForEachField(this, [&NumberOfField](::UE::FieldNotification::FFieldId FielId)
-			{
-				++NumberOfField;
-				return true;
-			});
-		FieldNotifyStartBitNumber = NumberOfField - FieldNotifyNames.Num();
-		ensureMsgf(FieldNotifyStartBitNumber >= 0, TEXT("The FieldNotifyStartIndex is negative. The number of field should be positive."));
-	}
-}
-
-void UWidgetBlueprintGeneratedClass::ForEachField(TFunctionRef<bool(::UE::FieldNotification::FFieldId FielId)> Callback, bool bIncludeSuper) const
-{
-	ensureMsgf(FieldNotifyStartBitNumber >= 0, TEXT("The FieldNotifyStartIndex is negative. The number of field should be positive."));
-	for (int32 Index = 0; Index < FieldNotifyNames.Num(); ++Index)
-	{
-		if (!Callback(UE::FieldNotification::FFieldId(FieldNotifyNames[Index].GetFieldName(), Index + FieldNotifyStartBitNumber)))
-		{
-			break;
-		}
-	}
-	if (bIncludeSuper)
-	{
-		if (UWidgetBlueprintGeneratedClass* ParentClass = Cast<UWidgetBlueprintGeneratedClass>(GetSuperClass()))
-		{
-			ParentClass->ForEachField(Callback);
-		}
-	}
-}
-
 UWidgetBlueprintGeneratedClassExtension* UWidgetBlueprintGeneratedClass::GetExtension(TSubclassOf<UWidgetBlueprintGeneratedClassExtension> InExtensionType, bool bIncludeSuper)
 {
 	for (UWidgetBlueprintGeneratedClassExtension* Extension : Extensions)
diff --git a/Engine/Source/Runtime/UMG/Public/Blueprint/WidgetBlueprintGeneratedClass.h b/Engine/Source/Runtime/UMG/Public/Blueprint/WidgetBlueprintGeneratedClass.h
--- a/Engine/Source/Runtime/UMG/Public/Blueprint/WidgetBlueprintGeneratedClass.h
+++ b/Engine/Source/Runtime/UMG/Public/Blueprint/WidgetBlueprintGeneratedClass.h
@@ -75,12 +75,6 @@ private:
 	UPROPERTY()
 	TArray<TObjectPtr<UWidgetBlueprintGeneratedClassExtension>> Extensions;
 
-	/** List Field Notifies. No index here on purpose to prevent saving them. */
-	UPROPERTY()
-	TArray<FFieldNotificationId> FieldNotifyNames;
-
-	int32 FieldNotifyStartBitNumber;
-
 	/** The classes native parent requires a native tick */
 	UPROPERTY()
 	uint32 bClassRequiresNativeTick :1;
@@ -136,14 +130,10 @@ public:
 	// Walks up the hierarchy looking for a valid widget tree.
 	UWidgetBlueprintGeneratedClass* FindWidgetTreeOwningClass() const;
 
-	// Execute the callback for every FieldId defined in the BP class
-	void ForEachField(TFunctionRef<bool(::UE::FieldNotification::FFieldId FielId)> Callback, bool bIncludeSuper = true) const;
-
 	//~ Begin UObject interface
 	virtual void Serialize(FArchive& Ar) override;
 	virtual void PostLoad() override;
 	virtual bool NeedsLoadForServer() const override;
-	virtual void PostLoadDefaultObject(UObject* Object) override;
 	//~ End UObject interface
 
 	virtual void PurgeClass(bool bRecompilingOnLoad) override;
@@ -155,8 +145,6 @@ public:
 	 */
 	void InitializeWidget(UUserWidget* UserWidget) const;
 
-	void InitializeFieldNotification(const UUserWidget*);
-
 	static void InitializeWidgetStatic(UUserWidget* UserWidget
 		, const UClass* InClass
 		, UWidgetTree* InWidgetTree
