Creating new parent class for widget editing settings and setting up widget and editor utility widget variants.

diff --git a/Engine/Config/BaseEditor.ini b/Engine/Config/BaseEditor.ini
--- a/Engine/Config/BaseEditor.ini
+++ b/Engine/Config/BaseEditor.ini
@@ -633,3 +633,9 @@ LogCookedStatus=false
 [/Script/TextureUtilitiesCommon.TextureImportSettings]
 CompressedFormatForFloatTextures=HDRCompressed_BC6
 
+[/Script/Blutility.EditorUtilityWidgetProjectSettings]
+bShowWidgetsFromEngineContent=True
+bShowWidgetsFromDeveloperContent=True
+bUseWidgetTemplateSelector=True
++CommonRootWidgetClasses=/Script/UMG.StackBox
++CommonRootWidgetClasses=/Script/UMG.GridPanel
diff --git a/Engine/Source/Editor/Blutility/Blutility.Build.cs b/Engine/Source/Editor/Blutility/Blutility.Build.cs
--- a/Engine/Source/Editor/Blutility/Blutility.Build.cs
+++ b/Engine/Source/Editor/Blutility/Blutility.Build.cs
@@ -43,7 +43,8 @@ public class Blutility : ModuleRules
 				"RHI",
 				"RenderCore",
 				"ImageWrapper",
-				"ImageWriteQueue"
+				"ImageWriteQueue",
+				"DeveloperSettings",
 			}
 			);
 	}
diff --git a/Engine/Source/Editor/Blutility/Classes/EditorUtilityWidgetBlueprint.h b/Engine/Source/Editor/Blutility/Classes/EditorUtilityWidgetBlueprint.h
--- a/Engine/Source/Editor/Blutility/Classes/EditorUtilityWidgetBlueprint.h
+++ b/Engine/Source/Editor/Blutility/Classes/EditorUtilityWidgetBlueprint.h
@@ -74,6 +74,9 @@ public:
 	/** Returns the default desired tab display name that was specified for this widget */
 	FText GetTabDisplayName() const;
 
+	virtual UWidgetEditingProjectSettings* GetRelevantSettings() override;
+	virtual const UWidgetEditingProjectSettings* GetRelevantSettings() const override;
+
 private:
 	bool IsWidgetEnabled() const;
 
diff --git a/Engine/Source/Editor/Blutility/Private/BlutilityModule.cpp b/Engine/Source/Editor/Blutility/Private/BlutilityModule.cpp
--- a/Engine/Source/Editor/Blutility/Private/BlutilityModule.cpp
+++ b/Engine/Source/Editor/Blutility/Private/BlutilityModule.cpp
@@ -62,6 +62,9 @@
 #include "Widgets/Docking/SDockTab.h"
 #include "WorkspaceMenuStructure.h"
 #include "WorkspaceMenuStructureModule.h"
+#include "EditorUtilityWidgetSettingsCustomization.h"
+#include "EditorUtilityWidgetProjectSettings.h"
+#include "PropertyEditorModule.h"
 
 #define LOCTEXT_NAMESPACE "AssetTypeActions"
 
@@ -106,6 +109,9 @@ public:
 
 		FAssetRegistryModule& AssetRegistryModule = FModuleManager::LoadModuleChecked<FAssetRegistryModule>(TEXT("AssetRegistry"));
 		AssetRegistryModule.Get().OnAssetRemoved().AddRaw(this, &FBlutilityModule::HandleAssetRemoved);
+
+		FPropertyEditorModule& PropertyModule = FModuleManager::LoadModuleChecked<FPropertyEditorModule>("PropertyEditor");
+		PropertyModule.RegisterCustomClassLayout(UEditorUtilityWidgetProjectSettings::StaticClass()->GetFName(), FOnGetDetailCustomizationInstance::CreateStatic(&FEditorUtilityWidgetSettingsCustomization::MakeInstance));
 	}
 
 	void ReinitializeUIs()
@@ -226,7 +232,7 @@ public:
 	virtual TConstArrayView<FAssetCategoryPath> GetAssetCategories() const override
 	{
 		static const TArray<FAssetCategoryPath, TFixedAllocator<1>> Categories = {
-			FAssetCategoryPath(LOCTEXT("EditorUtilities", "EditorUtilities"))
+			FAssetCategoryPath(LOCTEXT("EditorUtilities", "Editor Utilities"))
 		};
 		return Categories;
 	}
diff --git a/Engine/Source/Editor/Blutility/Private/EditorUtilityWidgetBlueprint.cpp b/Engine/Source/Editor/Blutility/Private/EditorUtilityWidgetBlueprint.cpp
--- a/Engine/Source/Editor/Blutility/Private/EditorUtilityWidgetBlueprint.cpp
+++ b/Engine/Source/Editor/Blutility/Private/EditorUtilityWidgetBlueprint.cpp
@@ -27,6 +27,7 @@
 #include "Widgets/Docking/SDockTab.h"
 #include "Widgets/SBoxPanel.h"
 #include "Widgets/SNullWidget.h"
+#include "EditorUtilityWidgetProjectSettings.h"
 
 class SWidget;
 class UClass;
@@ -162,6 +163,16 @@ FText UEditorUtilityWidgetBlueprint::GetTabDisplayName() const
 	return FText::FromString(FName::NameToDisplayString(GetName(), false));
 }
 
+UWidgetEditingProjectSettings* UEditorUtilityWidgetBlueprint::GetRelevantSettings()
+{
+	return GetMutableDefault<UEditorUtilityWidgetProjectSettings>();
+}
+
+const UWidgetEditingProjectSettings* UEditorUtilityWidgetBlueprint::GetRelevantSettings() const
+{
+	return GetDefault<UEditorUtilityWidgetProjectSettings>();
+}
+
 void UEditorUtilityWidgetBlueprint::ChangeTabWorld(UWorld* World, EMapChangeType MapChangeType)
 {
 	if (MapChangeType == EMapChangeType::TearDownWorld)
diff --git a/Engine/Source/Editor/Blutility/Private/EditorUtilityWidgetBlueprintFactory.cpp b/Engine/Source/Editor/Blutility/Private/EditorUtilityWidgetBlueprintFactory.cpp
--- a/Engine/Source/Editor/Blutility/Private/EditorUtilityWidgetBlueprintFactory.cpp
+++ b/Engine/Source/Editor/Blutility/Private/EditorUtilityWidgetBlueprintFactory.cpp
@@ -30,11 +30,11 @@
 #include "Templates/Casts.h"
 #include "Templates/ChooseClass.h"
 #include "Templates/SharedPointer.h"
-#include "UMGEditorProjectSettings.h"
 #include "UObject/Class.h"
 #include "UObject/Object.h"
 #include "UObject/ObjectPtr.h"
 #include "UObject/UnrealNames.h"
+#include "EditorUtilityWidgetProjectSettings.h"
 
 class FFeedbackContext;
 
@@ -76,7 +76,7 @@ UEditorUtilityWidgetBlueprintFactory::UEditorUtilityWidgetBlueprintFactory(const
 
 bool UEditorUtilityWidgetBlueprintFactory::ConfigureProperties()
 {
-	if (GetDefault<UUMGEditorProjectSettings>()->bUseWidgetTemplateSelector)
+	if (GetDefault<UEditorUtilityWidgetProjectSettings>()->bUseWidgetTemplateSelector)
 	{
 		// Load the classviewer module to display a class picker
 		FClassViewerModule& ClassViewerModule = FModuleManager::LoadModuleChecked<FClassViewerModule>("ClassViewer");
@@ -86,10 +86,23 @@ bool UEditorUtilityWidgetBlueprintFactory::ConfigureProperties()
 		Options.Mode = EClassViewerMode::ClassPicker;
 		Options.bShowNoneOption = true;
 
-		Options.ExtraPickerCommonClasses.Add(UHorizontalBox::StaticClass());
-		Options.ExtraPickerCommonClasses.Add(UVerticalBox::StaticClass());
-		Options.ExtraPickerCommonClasses.Add(UGridPanel::StaticClass());
-		Options.ExtraPickerCommonClasses.Add(UCanvasPanel::StaticClass());
+		TArray<TSoftClassPtr<UPanelWidget>> CommonRootWidgetClasses = GetDefault<UEditorUtilityWidgetProjectSettings>()->CommonRootWidgetClasses;
+		for (int32 Index = 0; Index < CommonRootWidgetClasses.Num(); ++Index)
+		{
+			UClass* PanelWidgetClass = CommonRootWidgetClasses[Index].LoadSynchronous();
+			if (PanelWidgetClass && PanelWidgetClass->IsChildOf(UPanelWidget::StaticClass()))
+			{
+				if (!Options.ExtraPickerCommonClasses.Contains(PanelWidgetClass))
+				{
+					Options.ExtraPickerCommonClasses.Add(PanelWidgetClass);
+				}
+			}
+		}
+
+		if (Options.ExtraPickerCommonClasses.Num() == 0)
+		{
+			Options.ExtraPickerCommonClasses.Add(UCanvasPanel::StaticClass());
+		}
 
 		TSharedPtr<FEditorUtilityWidgetBlueprintFactoryFilter> Filter = MakeShareable(new FEditorUtilityWidgetBlueprintFactoryFilter);
 		Options.ClassFilters.Add(Filter.ToSharedRef());
@@ -97,8 +110,8 @@ bool UEditorUtilityWidgetBlueprintFactory::ConfigureProperties()
 		Filter->DisallowedClassFlags = CLASS_Abstract | CLASS_Deprecated | CLASS_NewerVersionExists;
 		Filter->AllowedChildrenOfClasses.Add(UPanelWidget::StaticClass());
 
-		const FText TitleText = LOCTEXT("CreateWidgetBlueprint", "Pick Root Widget for New Editor Utility Widget");
-		return SClassPickerDialog::PickClass(TitleText, Options, RootWidgetClass, UPanelWidget::StaticClass());
+		const FText TitleText = LOCTEXT("CreateRootWidgetBlueprint", "Pick Root Widget for New Editor Utility Widget");
+		return SClassPickerDialog::PickClass(TitleText, Options, static_cast<UClass*&>(RootWidgetClass), UPanelWidget::StaticClass());
 
 	}
 	return true;
@@ -121,7 +134,7 @@ UObject* UEditorUtilityWidgetBlueprintFactory::FactoryCreateNew(UClass* Class, U
 	else
 	{
 		// If the root widget selection dialog is not enabled, use a canvas panel as the root by default
-		if (!GetDefault<UUMGEditorProjectSettings>()->bUseWidgetTemplateSelector)
+		if (!GetDefault<UEditorUtilityWidgetProjectSettings>()->bUseWidgetTemplateSelector)
 		{
 			RootWidgetClass = UCanvasPanel::StaticClass();
 		}
diff --git a/Engine/Source/Editor/Blutility/Private/EditorUtilityWidgetProjectSettings.cpp b/Engine/Source/Editor/Blutility/Private/EditorUtilityWidgetProjectSettings.cpp
new file mode 100644
--- /dev/null
+++ b/Engine/Source/Editor/Blutility/Private/EditorUtilityWidgetProjectSettings.cpp
@@ -0,0 +1,44 @@
+// Copyright Epic Games, Inc. All Rights Reserved.
+
+#include "EditorUtilityWidgetProjectSettings.h"
+#include "WidgetBlueprint.h"
+#include "WidgetCompilerRule.h"
+#include "UObject/Package.h"
+#include "UObject/UObjectIterator.h"
+
+#include "Blueprint/UserWidget.h"
+#include "Components/CanvasPanel.h"	
+#include "Components/GridPanel.h"
+#include "Components/HorizontalBox.h"
+#include "Components/VerticalBox.h"
+#include "../Classes/EditorUtilityWidgetBlueprint.h"
+
+
+FText UEditorUtilityWidgetProjectSettings::GetSectionText() const
+{
+	return NSLOCTEXT("EditorUtilities", "EditorUtilityWidgetsTeamSettingsName", "Editor Utility Widgets (Team)");
+}
+
+FText UEditorUtilityWidgetProjectSettings::GetSectionDescription() const
+{
+	return NSLOCTEXT("EditorUtilities", "EditorUtilityWidgetsTeamSettingsDescription", "Configure options for Editor Utility Widgets that affect the whole team.");
+}
+
+void UEditorUtilityWidgetProjectSettings::PostEditChangeChainProperty(FPropertyChangedChainEvent& PropertyChangedEvent)
+{
+	Super::PostEditChangeChainProperty(PropertyChangedEvent);
+
+	// If there's a change, we should scan for widgets currently in the error or warning state and mark them as dirty
+	// so they get recompiled next time we PIE.  Don't mark all widgets dirty, or we're in for a very large recompile.
+	if (PropertyChangedEvent.ChangeType != EPropertyChangeType::Interactive)
+	{
+		for (TObjectIterator<UEditorUtilityWidgetBlueprint> BlueprintIt; BlueprintIt; ++BlueprintIt)
+		{
+			UEditorUtilityWidgetBlueprint* Blueprint = *BlueprintIt;
+			if (Blueprint->Status == BS_Error || Blueprint->Status == BS_UpToDateWithWarnings)
+			{
+				Blueprint->Status = BS_Dirty;
+			}
+		}
+	}
+}
diff --git a/Engine/Source/Editor/Blutility/Private/EditorUtilityWidgetSettingsCustomization.cpp b/Engine/Source/Editor/Blutility/Private/EditorUtilityWidgetSettingsCustomization.cpp
new file mode 100644
--- /dev/null
+++ b/Engine/Source/Editor/Blutility/Private/EditorUtilityWidgetSettingsCustomization.cpp
@@ -0,0 +1,44 @@
+// Copyright Epic Games, Inc. All Rights Reserved.
+
+#include "EditorUtilityWidgetSettingsCustomization.h"
+#include "DetailLayoutBuilder.h"
+#include "DetailWidgetRow.h"
+#include "WidgetEditingProjectSettings.h"
+ 
+
+TSharedRef<IDetailCustomization> FEditorUtilityWidgetSettingsCustomization::MakeInstance()
+{
+	return MakeShareable( new FEditorUtilityWidgetSettingsCustomization() );
+}
+
+void FEditorUtilityWidgetSettingsCustomization::CustomizeDetails(IDetailLayoutBuilder& DetailLayout)
+{
+	// Hide the parent class selector properties because the EditorUtilityWidgetBP factory doesn't use them yet
+	TSharedPtr<IPropertyHandle> ParentSelectorProperty = DetailLayout.GetProperty(GET_MEMBER_NAME_CHECKED(UWidgetEditingProjectSettings, bUseUserWidgetParentClassViewerSelector), UWidgetEditingProjectSettings::StaticClass());
+	ParentSelectorProperty->MarkHiddenByCustomization();
+	TSharedPtr<IPropertyHandle> FavoriteParentClassesProperty = DetailLayout.GetProperty(GET_MEMBER_NAME_CHECKED(UWidgetEditingProjectSettings, FavoriteWidgetParentClasses), UWidgetEditingProjectSettings::StaticClass());
+	FavoriteParentClassesProperty->MarkHiddenByCustomization();
+
+	// Hide properties where we don't want the Editor Utility Widget settings to differ from the standard UMG ones
+	TSharedPtr<IPropertyHandle> TemplateSelectorProperty = DetailLayout.GetProperty(GET_MEMBER_NAME_CHECKED(UWidgetEditingProjectSettings, bUseWidgetTemplateSelector), UWidgetEditingProjectSettings::StaticClass());
+	TemplateSelectorProperty->MarkHiddenByCustomization();
+	TSharedPtr<IPropertyHandle> MakeVariableProperty = DetailLayout.GetProperty(GET_MEMBER_NAME_CHECKED(UWidgetEditingProjectSettings, bEnableMakeVariable), UWidgetEditingProjectSettings::StaticClass());
+	MakeVariableProperty->MarkHiddenByCustomization();
+	TSharedPtr<IPropertyHandle> WidgetAnimationWindowProperty = DetailLayout.GetProperty(GET_MEMBER_NAME_CHECKED(UWidgetEditingProjectSettings, bEnableWidgetAnimationEditor), UWidgetEditingProjectSettings::StaticClass());
+	WidgetAnimationWindowProperty->MarkHiddenByCustomization();
+	TSharedPtr<IPropertyHandle> PaletteWindowProperty = DetailLayout.GetProperty(GET_MEMBER_NAME_CHECKED(UWidgetEditingProjectSettings, bEnablePaletteWindow), UWidgetEditingProjectSettings::StaticClass());
+	PaletteWindowProperty->MarkHiddenByCustomization();
+	TSharedPtr<IPropertyHandle> LibraryWindowProperty = DetailLayout.GetProperty(GET_MEMBER_NAME_CHECKED(UWidgetEditingProjectSettings, bEnableLibraryWindow), UWidgetEditingProjectSettings::StaticClass());
+	LibraryWindowProperty->MarkHiddenByCustomization();
+	TSharedPtr<IPropertyHandle> HierarchyWindowProperty = DetailLayout.GetProperty(GET_MEMBER_NAME_CHECKED(UWidgetEditingProjectSettings, bEnableHierarchyWindow), UWidgetEditingProjectSettings::StaticClass());
+	HierarchyWindowProperty->MarkHiddenByCustomization();
+	TSharedPtr<IPropertyHandle> BindWidgetWindowProperty = DetailLayout.GetProperty(GET_MEMBER_NAME_CHECKED(UWidgetEditingProjectSettings, bEnableBindWidgetWindow), UWidgetEditingProjectSettings::StaticClass());
+	BindWidgetWindowProperty->MarkHiddenByCustomization();
+	TSharedPtr<IPropertyHandle> NavigationSimWindowProperty = DetailLayout.GetProperty(GET_MEMBER_NAME_CHECKED(UWidgetEditingProjectSettings, bEnableNavigationSimulationWindow), UWidgetEditingProjectSettings::StaticClass());
+	NavigationSimWindowProperty->MarkHiddenByCustomization();
+	TSharedPtr<IPropertyHandle> ConfigPaletteProperty = DetailLayout.GetProperty(GET_MEMBER_NAME_CHECKED(UWidgetEditingProjectSettings, bUseEditorConfigPaletteFiltering), UWidgetEditingProjectSettings::StaticClass());
+	ConfigPaletteProperty->MarkHiddenByCustomization();
+	TSharedPtr<IPropertyHandle> ClassViewerProperty = DetailLayout.GetProperty(GET_MEMBER_NAME_CHECKED(UWidgetEditingProjectSettings, bUseUserWidgetParentDefaultClassViewerSelector), UWidgetEditingProjectSettings::StaticClass());
+	ClassViewerProperty->MarkHiddenByCustomization();
+}
+
diff --git a/Engine/Source/Editor/Blutility/Private/EditorUtilityWidgetSettingsCustomization.h b/Engine/Source/Editor/Blutility/Private/EditorUtilityWidgetSettingsCustomization.h
new file mode 100644
--- /dev/null
+++ b/Engine/Source/Editor/Blutility/Private/EditorUtilityWidgetSettingsCustomization.h
@@ -0,0 +1,29 @@
+// Copyright Epic Games, Inc. All Rights Reserved.
+
+#pragma once
+
+#include "IDetailCustomization.h"
+
+
+class IDetailLayoutBuilder;
+
+//////////////////////////////////////////////////////////////////////////
+// FEditorUtilityWidgetSettingsCustomization
+
+class FEditorUtilityWidgetSettingsCustomization : public IDetailCustomization
+{
+public:
+	FEditorUtilityWidgetSettingsCustomization() = default;
+	virtual ~FEditorUtilityWidgetSettingsCustomization() = default;
+
+	/** Makes a new instance of this detail layout class for a specific detail view requesting it */
+	static TSharedRef<IDetailCustomization> MakeInstance();
+
+	// IDetailCustomization interface
+	virtual void CustomizeDetails(IDetailLayoutBuilder& DetailLayout) override;
+	// End of IDetailCustomization interface
+
+private:
+
+
+};
diff --git a/Engine/Source/Editor/Blutility/Public/EditorUtilityWidgetProjectSettings.h b/Engine/Source/Editor/Blutility/Public/EditorUtilityWidgetProjectSettings.h
new file mode 100644
--- /dev/null
+++ b/Engine/Source/Editor/Blutility/Public/EditorUtilityWidgetProjectSettings.h
@@ -0,0 +1,37 @@
+// Copyright Epic Games, Inc. All Rights Reserved.
+
+#pragma once
+
+#include "CoreMinimal.h"
+#include "EditorConfigBase.h"
+#include "UObject/ObjectMacros.h"
+#include "UObject/Object.h"
+#include "UObject/SoftObjectPath.h"
+#include "Engine/EngineTypes.h"
+#include "Misc/NamePermissionList.h"
+#include "WidgetEditingProjectSettings.h"
+#include "EditorUtilityWidgetProjectSettings.generated.h"
+
+class UWidgetCompilerRule;
+class UUserWidget;
+class UWidgetBlueprint;
+class UPanelWidget;
+
+/**
+ * Implements the settings for Editor Utility Widget Editing Project Settings
+ */
+UCLASS(config=Editor, defaultconfig)
+class BLUTILITY_API UEditorUtilityWidgetProjectSettings : public UWidgetEditingProjectSettings
+{
+	GENERATED_BODY()
+
+public:
+
+	virtual FText GetSectionText() const override;
+	virtual FText GetSectionDescription() const override;
+
+	// Begin UObject Interface
+	virtual void PostEditChangeChainProperty(struct FPropertyChangedChainEvent& PropertyChangedEvent) override;
+	// End UObject Interface
+
+};
diff --git a/Engine/Source/Editor/UMGEditor/Private/Animation/UMGDetailKeyframeHandler.cpp b/Engine/Source/Editor/UMGEditor/Private/Animation/UMGDetailKeyframeHandler.cpp
--- a/Engine/Source/Editor/UMGEditor/Private/Animation/UMGDetailKeyframeHandler.cpp
+++ b/Engine/Source/Editor/UMGEditor/Private/Animation/UMGDetailKeyframeHandler.cpp
@@ -2,7 +2,6 @@
 
 #include "Animation/UMGDetailKeyframeHandler.h"
 #include "Animation/WidgetAnimation.h"
-#include "UMGEditorProjectSettings.h"
 
 #include "PropertyHandle.h"
 #include "MovieScene.h"
diff --git a/Engine/Source/Editor/UMGEditor/Private/BlueprintModes/WidgetDesignerApplicationMode.cpp b/Engine/Source/Editor/UMGEditor/Private/BlueprintModes/WidgetDesignerApplicationMode.cpp
--- a/Engine/Source/Editor/UMGEditor/Private/BlueprintModes/WidgetDesignerApplicationMode.cpp
+++ b/Engine/Source/Editor/UMGEditor/Private/BlueprintModes/WidgetDesignerApplicationMode.cpp
@@ -10,7 +10,7 @@
 #include "HAL/PlatformApplicationMisc.h"
 #include "WidgetBlueprintEditorToolbar.h"
 #include "UMGEditorModule.h"
-#include "UMGEditorProjectSettings.h"
+#include "WidgetEditingProjectSettings.h"
 #include "StatusBarSubsystem.h"
 #include "ToolMenus.h"
 #include "WidgetDrawerConfig.h"
@@ -25,6 +25,7 @@
 #include "TabFactory/AnimationTabSummoner.h"
 #include "TabFactory/NavigationTabSummoner.h"
 #include "BlueprintModes/WidgetBlueprintApplicationModes.h"
+#include "WidgetBlueprintEditorUtils.h"
 
 #define LOCTEXT_NAMESPACE "WidgetDesignerMode"
 
@@ -49,35 +50,35 @@ FWidgetDesignerApplicationMode::FWidgetDesignerApplicationMode(TSharedPtr<FWidge
 	const FVector2D WindowSize = (CenterScale * DisplaySize) / DPIScale;
 
 	TSharedPtr<FTabManager::FStack> PaletteTab;
-	if (GetDefault<UUMGEditorProjectSettings>()->bEnablePaletteWindow || GetDefault<UUMGEditorProjectSettings>()->bEnableLibraryWindow)
+	if (FWidgetBlueprintEditorUtils::GetRelevantSettings(InWidgetEditor)->bEnablePaletteWindow || FWidgetBlueprintEditorUtils::GetRelevantSettings(InWidgetEditor)->bEnableLibraryWindow)
 	{
 		PaletteTab = FTabManager::NewStack()
 			->SetSizeCoefficient(0.5f);
 
-		if (GetDefault<UUMGEditorProjectSettings>()->bEnablePaletteWindow)
+		if (FWidgetBlueprintEditorUtils::GetRelevantSettings(InWidgetEditor)->bEnablePaletteWindow)
 		{
 			PaletteTab->SetForegroundTab(FPaletteTabSummoner::TabID)
 				->AddTab(FPaletteTabSummoner::TabID, ETabState::OpenedTab);
 		}
-		if (GetDefault<UUMGEditorProjectSettings>()->bEnableLibraryWindow)
+		if (FWidgetBlueprintEditorUtils::GetRelevantSettings(InWidgetEditor)->bEnableLibraryWindow)
 		{
 			PaletteTab->AddTab(FLibraryTabSummoner::TabID, ETabState::OpenedTab);
 		}
 	}
 
 	TSharedPtr<FTabManager::FStack> EnableHierarchyTab;
-	if (GetDefault<UUMGEditorProjectSettings>()->bEnableHierarchyWindow || GetDefault<UUMGEditorProjectSettings>()->bEnableBindWidgetWindow)
+	if (FWidgetBlueprintEditorUtils::GetRelevantSettings(InWidgetEditor)->bEnableHierarchyWindow || FWidgetBlueprintEditorUtils::GetRelevantSettings(InWidgetEditor)->bEnableBindWidgetWindow)
 	{
 		EnableHierarchyTab = FTabManager::NewStack()
 			->SetSizeCoefficient(0.5f);
 
-		if (GetDefault<UUMGEditorProjectSettings>()->bEnableHierarchyWindow)
+		if (FWidgetBlueprintEditorUtils::GetRelevantSettings(InWidgetEditor)->bEnableHierarchyWindow)
 		{
 			EnableHierarchyTab->SetForegroundTab(FHierarchyTabSummoner::TabID)
 				->AddTab(FHierarchyTabSummoner::TabID, ETabState::OpenedTab);
 		}
 
-		if (GetDefault<UUMGEditorProjectSettings>()->bEnableBindWidgetWindow)
+		if (FWidgetBlueprintEditorUtils::GetRelevantSettings(InWidgetEditor)->bEnableBindWidgetWindow)
 		{
 			EnableHierarchyTab->AddTab(FBindWidgetTabSummoner::TabID, ETabState::OpenedTab);
 		}
@@ -166,23 +167,23 @@ FWidgetDesignerApplicationMode::FWidgetDesignerApplicationMode(TSharedPtr<FWidge
 	TabFactories.RegisterFactory(MakeShareable(new FDesignerTabSummoner(InWidgetEditor)));
 	TabFactories.RegisterFactory(MakeShareable(new FCompilerResultsSummoner(InWidgetEditor)));
 	TabFactories.RegisterFactory(MakeShareable(new FAnimationTabSummoner(InWidgetEditor)));
-	if (GetDefault<UUMGEditorProjectSettings>()->bEnableHierarchyWindow)
+	if (FWidgetBlueprintEditorUtils::GetRelevantSettings(InWidgetEditor)->bEnableHierarchyWindow)
 	{
 		TabFactories.RegisterFactory(MakeShareable(new FHierarchyTabSummoner(InWidgetEditor)));
 	}
-	if (GetDefault<UUMGEditorProjectSettings>()->bEnableBindWidgetWindow)
+	if (FWidgetBlueprintEditorUtils::GetRelevantSettings(InWidgetEditor)->bEnableBindWidgetWindow)
 	{
 		TabFactories.RegisterFactory(MakeShareable(new FBindWidgetTabSummoner(InWidgetEditor)));
 	}
-	if (GetDefault<UUMGEditorProjectSettings>()->bEnablePaletteWindow)
+	if (FWidgetBlueprintEditorUtils::GetRelevantSettings(InWidgetEditor)->bEnablePaletteWindow)
 	{
 		TabFactories.RegisterFactory(MakeShareable(new FPaletteTabSummoner(InWidgetEditor)));
 	}
-	if (GetDefault<UUMGEditorProjectSettings>()->bEnableLibraryWindow)
+	if (FWidgetBlueprintEditorUtils::GetRelevantSettings(InWidgetEditor)->bEnableLibraryWindow)
 	{
 		TabFactories.RegisterFactory(MakeShareable(new FLibraryTabSummoner(InWidgetEditor)));
 	}
-	if (GetDefault<UUMGEditorProjectSettings>()->bEnableNavigationSimulationWindow)
+	if (FWidgetBlueprintEditorUtils::GetRelevantSettings(InWidgetEditor)->bEnableNavigationSimulationWindow)
 	{
 		TabFactories.RegisterFactory(MakeShareable(new FNavigationTabSummoner(InWidgetEditor)));
 	}
@@ -246,7 +247,7 @@ void FWidgetDesignerApplicationMode::PostActivateMode()
 
 	TSharedPtr<FWidgetBlueprintEditor> BP = GetBlueprintEditor();
 
-	if (GetDefault<UUMGEditorProjectSettings>()->bEnableWidgetAnimationEditor)
+	if (FWidgetBlueprintEditorUtils::GetRelevantSettings(BP)->bEnableWidgetAnimationEditor)
 	{
 		FWidgetDrawerConfig WidgetAnimSequencerDrawer(FAnimationTabSummoner::WidgetAnimSequencerDrawerID);
 		WidgetAnimSequencerDrawer.GetDrawerContentDelegate.BindSP(BP.Get(), &FWidgetBlueprintEditor::OnGetWidgetAnimSequencer);
diff --git a/Engine/Source/Editor/UMGEditor/Private/Designer/SDesignerView.cpp b/Engine/Source/Editor/UMGEditor/Private/Designer/SDesignerView.cpp
--- a/Engine/Source/Editor/UMGEditor/Private/Designer/SDesignerView.cpp
+++ b/Engine/Source/Editor/UMGEditor/Private/Designer/SDesignerView.cpp
@@ -71,13 +71,14 @@
 #include "Input/HittestGrid.h"
 
 #include "Fonts/FontMeasure.h"
-#include "UMGEditorProjectSettings.h"
+#include "WidgetEditingProjectSettings.h"
 #include "DeviceProfiles/DeviceProfile.h"
 #include "DeviceProfiles/DeviceProfileManager.h"
 #include "Engine/DPICustomScalingRule.h"
 #include "UMGEditorModule.h"
 #include "ToolMenus.h"
 #include "Styling/ToolBarStyle.h"
+#include "UMGEditorProjectSettings.h"
 
 #define LOCTEXT_NAMESPACE "UMG"
 
@@ -137,6 +138,11 @@ public:
 		if (bResizing)
 		{
 			TSharedPtr<SDesignerView> DesignerView = Designer.Pin();
+			if (!DesignerView)
+			{
+				return FReply::Unhandled();
+            }
+
 			DesignerView->BeginResizingArea();
 
 			const float ZoomAmount = DesignerView->GetZoomAmount();
@@ -145,7 +151,7 @@ public:
 			AreaSize /= ZoomAmount;
 			AreaSize /= MyGeometry.Scale;
 
-			if (const UUMGEditorProjectSettings* Settings = GetDefault<UUMGEditorProjectSettings>())
+			if (const UWidgetEditingProjectSettings* Settings = DesignerView->GetRelevantSettings())
 			{
 				for (const FDebugResolution& Resolution : Settings->DebugResolutions)
 				{
@@ -1022,6 +1028,16 @@ void SDesignerView::EndResizingArea()
 	bShowResolutionOutlines = false;
 }
 
+const UWidgetEditingProjectSettings* SDesignerView::GetRelevantSettings() const
+{
+	if (UWidgetBlueprint* WidgetBlueprint = GetBlueprint())
+	{
+		return WidgetBlueprint->GetRelevantSettings();
+	}
+	// Default to the UMG Editor project settings
+	return GetDefault<UUMGEditorProjectSettings>();
+}
+
 void SDesignerView::SetPreviewAreaSize(int32 Width, int32 Height)
 {
 	if (UUserWidget* DefaultWidget = GetDefaultWidget())
@@ -2396,7 +2412,7 @@ void SDesignerView::OnPaintBackground(const FGeometry& AllottedGeometry, const F
 
 	if (bShowResolutionOutlines)
 	{
-		if (const UUMGEditorProjectSettings* Settings = GetDefault<UUMGEditorProjectSettings>())
+		if (const UWidgetEditingProjectSettings* Settings = FWidgetBlueprintEditorUtils::GetRelevantSettings(BlueprintEditor))
 		{
 			for (const FDebugResolution& Resolution : Settings->DebugResolutions)
 			{
diff --git a/Engine/Source/Editor/UMGEditor/Private/Details/DetailWidgetExtensionHandler.cpp b/Engine/Source/Editor/UMGEditor/Private/Details/DetailWidgetExtensionHandler.cpp
--- a/Engine/Source/Editor/UMGEditor/Private/Details/DetailWidgetExtensionHandler.cpp
+++ b/Engine/Source/Editor/UMGEditor/Private/Details/DetailWidgetExtensionHandler.cpp
@@ -5,7 +5,6 @@
 #include "Customizations/UMGDetailCustomizations.h"
 #include "Engine/Blueprint.h"
 #include "UMGEditorModule.h"
-#include "UMGEditorProjectSettings.h"
 #include "WidgetBlueprint.h"
 #include "WidgetBlueprintEditor.h"
 
diff --git a/Engine/Source/Editor/UMGEditor/Private/Details/SWidgetDetailsView.cpp b/Engine/Source/Editor/UMGEditor/Private/Details/SWidgetDetailsView.cpp
--- a/Engine/Source/Editor/UMGEditor/Private/Details/SWidgetDetailsView.cpp
+++ b/Engine/Source/Editor/UMGEditor/Private/Details/SWidgetDetailsView.cpp
@@ -32,7 +32,7 @@
 #include "ScopedTransaction.h"
 #include "Styling/SlateIconFinder.h"
 #include "UMGEditorModule.h"
-#include "UMGEditorProjectSettings.h"
+#include "WidgetEditingProjectSettings.h"
 
 #define LOCTEXT_NAMESPACE "UMG"
 
@@ -51,7 +51,7 @@ void SWidgetDetailsView::Construct(const FArguments& InArgs, TSharedPtr<FWidgetB
 	PropertyView = EditModule.CreateDetailView(DetailsViewArgs);
 
 	// Create a handler for keyframing via the details panel
-	if (GetDefault<UUMGEditorProjectSettings>()->bEnableWidgetAnimationEditor)
+	if (FWidgetBlueprintEditorUtils::GetRelevantSettings(BlueprintEditor)->bEnableWidgetAnimationEditor)
 	{
 		TSharedRef<IDetailKeyframeHandler> KeyframeHandler = MakeShareable(new FUMGDetailKeyframeHandler(InBlueprintEditor));
 		PropertyView->SetKeyframeHandler(KeyframeHandler);
@@ -162,7 +162,7 @@ void SWidgetDetailsView::Construct(const FArguments& InArgs, TSharedPtr<FWidgetB
 					+ SHorizontalBox::Slot()
 					.AutoWidth()
 					[
-						GetDefault<UUMGEditorProjectSettings>()->bEnableMakeVariable ? IsVariableCheckbox : SNullWidget::NullWidget
+						FWidgetBlueprintEditorUtils::GetRelevantSettings(BlueprintEditor)->bEnableMakeVariable ? IsVariableCheckbox : SNullWidget::NullWidget
 					]
 
 					+ SHorizontalBox::Slot()
diff --git a/Engine/Source/Editor/UMGEditor/Private/Library/SLibraryView.cpp b/Engine/Source/Editor/UMGEditor/Private/Library/SLibraryView.cpp
--- a/Engine/Source/Editor/UMGEditor/Private/Library/SLibraryView.cpp
+++ b/Engine/Source/Editor/UMGEditor/Private/Library/SLibraryView.cpp
@@ -31,7 +31,6 @@
 #include "IContentBrowserSingleton.h"
 #include "WidgetBlueprintEditorUtils.h"
 
-#include "UMGEditorProjectSettings.h"
 #include "Library/SLibraryView.h"
 
 #define LOCTEXT_NAMESPACE "UMG"
diff --git a/Engine/Source/Editor/UMGEditor/Private/Library/SLibraryViewModel.cpp b/Engine/Source/Editor/UMGEditor/Private/Library/SLibraryViewModel.cpp
--- a/Engine/Source/Editor/UMGEditor/Private/Library/SLibraryViewModel.cpp
+++ b/Engine/Source/Editor/UMGEditor/Private/Library/SLibraryViewModel.cpp
@@ -24,7 +24,6 @@
 #include "IContentBrowserDataModule.h"
 #include "Settings/ContentBrowserSettings.h"
 #include "Settings/WidgetDesignerSettings.h"
-#include "UMGEditorProjectSettings.h"
 #include "WidgetPaletteFavorites.h"
 #include "Library/SLibraryViewModel.h"
 
diff --git a/Engine/Source/Editor/UMGEditor/Private/Palette/SPaletteView.cpp b/Engine/Source/Editor/UMGEditor/Private/Palette/SPaletteView.cpp
--- a/Engine/Source/Editor/UMGEditor/Private/Palette/SPaletteView.cpp
+++ b/Engine/Source/Editor/UMGEditor/Private/Palette/SPaletteView.cpp
@@ -30,7 +30,6 @@
 
 
 
-#include "UMGEditorProjectSettings.h"
 
 #define LOCTEXT_NAMESPACE "UMG"
 
diff --git a/Engine/Source/Editor/UMGEditor/Private/Palette/SPaletteViewModel.cpp b/Engine/Source/Editor/UMGEditor/Private/Palette/SPaletteViewModel.cpp
--- a/Engine/Source/Editor/UMGEditor/Private/Palette/SPaletteViewModel.cpp
+++ b/Engine/Source/Editor/UMGEditor/Private/Palette/SPaletteViewModel.cpp
@@ -27,7 +27,7 @@
 
 #include "Settings/ContentBrowserSettings.h"
 #include "Settings/WidgetDesignerSettings.h"
-#include "UMGEditorProjectSettings.h"
+#include "WidgetEditingProjectSettings.h"
 #include "WidgetPaletteFavorites.h"
 
 #define LOCTEXT_NAMESPACE "UMG"
@@ -215,6 +215,12 @@ void FWidgetCatalogViewModel::BuildClassWidgetList()
 {
 	const UClass* ActiveWidgetBlueprintClass = GetBlueprint()->GeneratedClass;
 	FName ActiveWidgetBlueprintClassName = ActiveWidgetBlueprintClass->GetFName();
+	TSharedPtr<FWidgetBlueprintEditor> PinnedBPEditor = BlueprintEditor.Pin();
+
+	if (!PinnedBPEditor)
+	{
+		return;
+	}
 
 	// Locate all UWidget classes from code and loaded widget BPs
 	for (TObjectIterator<UClass> ClassIt; ClassIt; ++ClassIt)
@@ -226,7 +232,7 @@ void FWidgetCatalogViewModel::BuildClassWidgetList()
 			continue;
 		}
 
-		if (!FWidgetBlueprintEditorUtils::IsUsableWidgetClass(WidgetClass))
+		if (!FWidgetBlueprintEditorUtils::IsUsableWidgetClass(WidgetClass, PinnedBPEditor.ToSharedRef()))
 		{
 			continue;
 		}
@@ -269,7 +275,7 @@ void FWidgetCatalogViewModel::BuildClassWidgetList()
 			continue;
 		}
 
-		TValueOrError<FWidgetBlueprintEditorUtils::FUsableWidgetClassResult, void> Usable = FWidgetBlueprintEditorUtils::IsUsableWidgetClass(BPAssetData);
+		TValueOrError<FWidgetBlueprintEditorUtils::FUsableWidgetClassResult, void> Usable = FWidgetBlueprintEditorUtils::IsUsableWidgetClass(BPAssetData, PinnedBPEditor.ToSharedRef());
 		if (Usable.HasError())
 		{
 			continue;
@@ -293,7 +299,7 @@ void FWidgetCatalogViewModel::BuildClassWidgetList()
 
 	TArray<FAssetData> AllGeneratedBPsAssetData; // if it's a widget already compiled
 
-	if (AssetRegistryModule && GetDefault<UUMGEditorProjectSettings>()->bUseEditorConfigPaletteFiltering)
+	if (AssetRegistryModule && FWidgetBlueprintEditorUtils::GetRelevantSettings(BlueprintEditor)->bUseEditorConfigPaletteFiltering)
 	{
 		AssetRegistryModule->Get().GetAssetsByClass(UBlueprintGeneratedClass::StaticClass()->GetClassPathName(), AllGeneratedBPsAssetData, true);
 	}
@@ -312,7 +318,7 @@ void FWidgetCatalogViewModel::BuildClassWidgetList()
 			continue;
 		}
 
-		TValueOrError<FWidgetBlueprintEditorUtils::FUsableWidgetClassResult, void> Usable = FWidgetBlueprintEditorUtils::IsUsableWidgetClass(BPAssetData);
+		TValueOrError<FWidgetBlueprintEditorUtils::FUsableWidgetClassResult, void> Usable = FWidgetBlueprintEditorUtils::IsUsableWidgetClass(BPAssetData, PinnedBPEditor.ToSharedRef());
 		if (Usable.HasError())
 		{
 			continue;
@@ -332,7 +338,8 @@ void FWidgetCatalogViewModel::AddWidgetTemplate(TSharedPtr<FWidgetTemplate> Temp
 	FString Category = *Template->GetCategory().ToString();
 
 	// Hide user specific categories
-	const TArray<FString>& CategoriesToHide = GetDefault<UUMGEditorProjectSettings>()->CategoriesToHide;
+	const TArray<FString>& CategoriesToHide = FWidgetBlueprintEditorUtils::GetRelevantSettings(BlueprintEditor)->CategoriesToHide;
+	
 	for (const FString& CategoryName : CategoriesToHide)
 	{
 		if (Category == CategoryName)
diff --git a/Engine/Source/Editor/UMGEditor/Private/Settings/UMGEditorProjectSettings.cpp b/Engine/Source/Editor/UMGEditor/Private/Settings/UMGEditorProjectSettings.cpp
--- a/Engine/Source/Editor/UMGEditor/Private/Settings/UMGEditorProjectSettings.cpp
+++ b/Engine/Source/Editor/UMGEditor/Private/Settings/UMGEditorProjectSettings.cpp
@@ -56,103 +56,6 @@ FText UUMGEditorProjectSettings::GetSectionDescription() const
 
 #endif
 
-bool UUMGEditorProjectSettings::CompilerOption_AllowBlueprintTick(const class UWidgetBlueprint* WidgetBlueprint) const
-{
-	return GetFirstCompilerOption(WidgetBlueprint, &FWidgetCompilerOptions::bAllowBlueprintTick, true);
-}
-
-bool UUMGEditorProjectSettings::CompilerOption_AllowBlueprintPaint(const class UWidgetBlueprint* WidgetBlueprint) const
-{
-	return GetFirstCompilerOption(WidgetBlueprint, &FWidgetCompilerOptions::bAllowBlueprintPaint, true);
-}
-
-EPropertyBindingPermissionLevel UUMGEditorProjectSettings::CompilerOption_PropertyBindingRule(const class UWidgetBlueprint* WidgetBlueprint) const
-{
-	return GetFirstCompilerOption(WidgetBlueprint, &FWidgetCompilerOptions::PropertyBindingRule, EPropertyBindingPermissionLevel::Allow);
-}
-
-TArray<UWidgetCompilerRule*> UUMGEditorProjectSettings::CompilerOption_Rules(const class UWidgetBlueprint* WidgetBlueprint) const
-{
-	TArray<UWidgetCompilerRule*> Rules;
-	GetCompilerOptionsForWidget(WidgetBlueprint, [&Rules](const FWidgetCompilerOptions& Options) {
-		for (const TSoftClassPtr<UWidgetCompilerRule>& RuleClassPtr : Options.Rules)
-		{
-			// The compiling rule may not be loaded yet in early loading phases, we'll
-			// just have to skip the rules in those cases.
-			RuleClassPtr.LoadSynchronous();
-			if (RuleClassPtr)
-			{
-				if (UWidgetCompilerRule* Rule = RuleClassPtr->GetDefaultObject<UWidgetCompilerRule>())
-				{
-					Rules.Add(Rule);
-				}
-			}
-		}
-		return false;
-	});
-	return Rules;
-}
-
-FNamePermissionList& UUMGEditorProjectSettings::GetAllowedPaletteCategories()
-{
-	return AllowedPaletteCategories;
-}
-
-const FNamePermissionList& UUMGEditorProjectSettings::GetAllowedPaletteCategories() const
-{
-	return AllowedPaletteCategories;
-}
-
-FPathPermissionList& UUMGEditorProjectSettings::GetAllowedPaletteWidgets()
-{
-	return AllowedPaletteWidgets;
-}
-
-const FPathPermissionList& UUMGEditorProjectSettings::GetAllowedPaletteWidgets() const
-{
-	return AllowedPaletteWidgets;
-}
-
-void UUMGEditorProjectSettings::GetCompilerOptionsForWidget(const UWidgetBlueprint* WidgetBlueprint, TFunctionRef<bool(const FWidgetCompilerOptions&)> Operator) const
-{
-	FString AssetPath = WidgetBlueprint->GetOutermost()->GetName();
-	FSoftObjectPath SoftObjectPath = WidgetBlueprint->GetPathName();
-	
-	// Don't apply the rules to the engine widgets.
-	if (AssetPath.StartsWith(TEXT("/Engine")))
-	{
-		return;
-	}
-
-	for (int32 DirectoryIndex = DirectoryCompilerOptions.Num() - 1; DirectoryIndex >= 0; DirectoryIndex--)
-	{
-		const FDirectoryWidgetCompilerOptions& CompilerOptions = DirectoryCompilerOptions[DirectoryIndex];
-
-		const FString& DirectoryPath = CompilerOptions.Directory.Path;
-		if (!DirectoryPath.IsEmpty())
-		{
-			if (AssetPath.StartsWith(DirectoryPath))
-			{
-				const bool bIgnoreWidget = CompilerOptions.IgnoredWidgets.ContainsByPredicate([&SoftObjectPath](const TSoftObjectPtr<UWidgetBlueprint>& IgnoredWidget) {
-					return IgnoredWidget.ToSoftObjectPath() == SoftObjectPath;
-				});
-
-				if (bIgnoreWidget)
-				{
-					continue;
-				}
-
-				if (Operator(CompilerOptions.Options))
-				{
-					return;
-				}
-			}
-		}
-	}
-
-	Operator(DefaultCompilerOptions);
-}
-
 #if WITH_EDITOR
 void UUMGEditorProjectSettings::PostEditChangeChainProperty(FPropertyChangedChainEvent& PropertyChangedEvent)
 {
@@ -177,22 +80,3 @@ void UUMGEditorProjectSettings::PostEditChangeChainProperty(FPropertyChangedChai
 	}
 }
 #endif
-
-void UUMGEditorProjectSettings::PostInitProperties()
-{
-	Super::PostInitProperties();
-
-	if (Version < CurrentVersion)
-	{
-		for (int32 FromVersion = Version + 1; FromVersion <= CurrentVersion; FromVersion++)
-		{
-			PerformUpgradeStepForVersion(FromVersion);
-		}
-
-		Version = CurrentVersion;
-	}
-}
-
-void UUMGEditorProjectSettings::PerformUpgradeStepForVersion(int32 ForVersion)
-{
-}
\ No newline at end of file
diff --git a/Engine/Source/Editor/UMGEditor/Private/Settings/WidgetEditingProjectSettings.cpp b/Engine/Source/Editor/UMGEditor/Private/Settings/WidgetEditingProjectSettings.cpp
new file mode 100644
--- /dev/null
+++ b/Engine/Source/Editor/UMGEditor/Private/Settings/WidgetEditingProjectSettings.cpp
@@ -0,0 +1,149 @@
+// Copyright Epic Games, Inc. All Rights Reserved.
+
+#include "WidgetEditingProjectSettings.h"
+#include "WidgetBlueprint.h"
+#include "WidgetCompilerRule.h"
+#include "UObject/Package.h"
+#include "UObject/UObjectIterator.h"
+
+#include "Blueprint/UserWidget.h"
+#include "Components/CanvasPanel.h"	
+#include "Components/GridPanel.h"
+#include "Components/HorizontalBox.h"
+#include "Components/VerticalBox.h"
+
+UWidgetEditingProjectSettings::UWidgetEditingProjectSettings()
+{
+	Version = 0;
+	CurrentVersion = 1;
+
+	bEnableMakeVariable = true;
+	bEnableWidgetAnimationEditor = true;
+	bEnablePaletteWindow = true;
+	bEnableLibraryWindow = true;
+	bEnableHierarchyWindow = true;
+	bEnableBindWidgetWindow = true;
+	bEnableNavigationSimulationWindow = true;
+
+	bUseEditorConfigPaletteFiltering = false;
+	bUseUserWidgetParentClassViewerSelector = true;
+	bUseUserWidgetParentDefaultClassViewerSelector = true;
+
+	bUseWidgetTemplateSelector = false;
+}
+
+bool UWidgetEditingProjectSettings::CompilerOption_AllowBlueprintTick(const class UWidgetBlueprint* WidgetBlueprint) const
+{
+	return GetFirstCompilerOption(WidgetBlueprint, &FWidgetCompilerOptions::bAllowBlueprintTick, true);
+}
+
+bool UWidgetEditingProjectSettings::CompilerOption_AllowBlueprintPaint(const class UWidgetBlueprint* WidgetBlueprint) const
+{
+	return GetFirstCompilerOption(WidgetBlueprint, &FWidgetCompilerOptions::bAllowBlueprintPaint, true);
+}
+
+EPropertyBindingPermissionLevel UWidgetEditingProjectSettings::CompilerOption_PropertyBindingRule(const class UWidgetBlueprint* WidgetBlueprint) const
+{
+	return GetFirstCompilerOption(WidgetBlueprint, &FWidgetCompilerOptions::PropertyBindingRule, EPropertyBindingPermissionLevel::Allow);
+}
+
+TArray<UWidgetCompilerRule*> UWidgetEditingProjectSettings::CompilerOption_Rules(const class UWidgetBlueprint* WidgetBlueprint) const
+{
+	TArray<UWidgetCompilerRule*> Rules;
+	GetCompilerOptionsForWidget(WidgetBlueprint, [&Rules](const FWidgetCompilerOptions& Options) {
+		for (const TSoftClassPtr<UWidgetCompilerRule>& RuleClassPtr : Options.Rules)
+		{
+			// The compiling rule may not be loaded yet in early loading phases, we'll
+			// just have to skip the rules in those cases.
+			RuleClassPtr.LoadSynchronous();
+			if (RuleClassPtr)
+			{
+				if (UWidgetCompilerRule* Rule = RuleClassPtr->GetDefaultObject<UWidgetCompilerRule>())
+				{
+					Rules.Add(Rule);
+				}
+			}
+		}
+		return false;
+	});
+	return Rules;
+}
+
+FNamePermissionList& UWidgetEditingProjectSettings::GetAllowedPaletteCategories()
+{
+	return AllowedPaletteCategories;
+}
+
+const FNamePermissionList& UWidgetEditingProjectSettings::GetAllowedPaletteCategories() const
+{
+	return AllowedPaletteCategories;
+}
+
+FPathPermissionList& UWidgetEditingProjectSettings::GetAllowedPaletteWidgets()
+{
+	return AllowedPaletteWidgets;
+}
+
+const FPathPermissionList& UWidgetEditingProjectSettings::GetAllowedPaletteWidgets() const
+{
+	return AllowedPaletteWidgets;
+}
+
+void UWidgetEditingProjectSettings::GetCompilerOptionsForWidget(const UWidgetBlueprint* WidgetBlueprint, TFunctionRef<bool(const FWidgetCompilerOptions&)> Operator) const
+{
+	FString AssetPath = WidgetBlueprint->GetOutermost()->GetName();
+	FSoftObjectPath SoftObjectPath = WidgetBlueprint->GetPathName();
+	
+	// Don't apply the rules to the engine widgets.
+	if (AssetPath.StartsWith(TEXT("/Engine")))
+	{
+		return;
+	}
+
+	for (int32 DirectoryIndex = DirectoryCompilerOptions.Num() - 1; DirectoryIndex >= 0; DirectoryIndex--)
+	{
+		const FDirectoryWidgetCompilerOptions& CompilerOptions = DirectoryCompilerOptions[DirectoryIndex];
+
+		const FString& DirectoryPath = CompilerOptions.Directory.Path;
+		if (!DirectoryPath.IsEmpty())
+		{
+			if (AssetPath.StartsWith(DirectoryPath))
+			{
+				const bool bIgnoreWidget = CompilerOptions.IgnoredWidgets.ContainsByPredicate([&SoftObjectPath](const TSoftObjectPtr<UWidgetBlueprint>& IgnoredWidget) {
+					return IgnoredWidget.ToSoftObjectPath() == SoftObjectPath;
+				});
+
+				if (bIgnoreWidget)
+				{
+					continue;
+				}
+
+				if (Operator(CompilerOptions.Options))
+				{
+					return;
+				}
+			}
+		}
+	}
+
+	Operator(DefaultCompilerOptions);
+}
+
+void UWidgetEditingProjectSettings::PostInitProperties()
+{
+	Super::PostInitProperties();
+
+	if (Version < CurrentVersion)
+	{
+		for (int32 FromVersion = Version + 1; FromVersion <= CurrentVersion; FromVersion++)
+		{
+			PerformUpgradeStepForVersion(FromVersion);
+		}
+
+		Version = CurrentVersion;
+	}
+}
+
+void UWidgetEditingProjectSettings::PerformUpgradeStepForVersion(int32 ForVersion)
+{
+}
\ No newline at end of file
diff --git a/Engine/Source/Editor/UMGEditor/Private/UMGEditorModule.cpp b/Engine/Source/Editor/UMGEditor/Private/UMGEditorModule.cpp
--- a/Engine/Source/Editor/UMGEditor/Private/UMGEditorModule.cpp
+++ b/Engine/Source/Editor/UMGEditor/Private/UMGEditorModule.cpp
@@ -32,7 +32,6 @@
 
 #include "ClassIconFinder.h"
 
-#include "UMGEditorProjectSettings.h"
 #include "ISettingsModule.h"
 #include "SequencerSettings.h"
 
diff --git a/Engine/Source/Editor/UMGEditor/Private/WidgetBlueprint.cpp b/Engine/Source/Editor/UMGEditor/Private/WidgetBlueprint.cpp
--- a/Engine/Source/Editor/UMGEditor/Private/WidgetBlueprint.cpp
+++ b/Engine/Source/Editor/UMGEditor/Private/WidgetBlueprint.cpp
@@ -36,6 +36,7 @@
 #include "K2Node_Composite.h"
 #include "K2Node_FunctionResult.h"
 #include "Blueprint/WidgetNavigation.h"
+#include "WidgetEditingProjectSettings.h"
 
 #define LOCTEXT_NAMESPACE "UMG"
 
@@ -1401,7 +1402,7 @@ void UWidgetBlueprint::UpdateTickabilityStats(bool& OutHasLatentActions, bool& O
 
 bool UWidgetBlueprint::ArePropertyBindingsAllowed() const
 {
-	return GetDefault<UUMGEditorProjectSettings>()->CompilerOption_PropertyBindingRule(this) == EPropertyBindingPermissionLevel::Allow;
+	return GetRelevantSettings()->CompilerOption_PropertyBindingRule(this) == EPropertyBindingPermissionLevel::Allow;
 }
 
 TArray<FName> UWidgetBlueprint::GetInheritedAvailableNamedSlots() const
@@ -1414,6 +1415,16 @@ TArray<FName> UWidgetBlueprint::GetInheritedAvailableNamedSlots() const
 	return TArray<FName>();
 }
 
+UWidgetEditingProjectSettings* UWidgetBlueprint::GetRelevantSettings()
+{
+	return GetMutableDefault<UUMGEditorProjectSettings>();
+}
+
+const UWidgetEditingProjectSettings* UWidgetBlueprint::GetRelevantSettings() const
+{
+	return GetDefault<UUMGEditorProjectSettings>();
+}
+
 #if WITH_EDITOR
 void UWidgetBlueprint::LoadModulesRequiredForCompilation()
 {
diff --git a/Engine/Source/Editor/UMGEditor/Private/WidgetBlueprintCompiler.cpp b/Engine/Source/Editor/UMGEditor/Private/WidgetBlueprintCompiler.cpp
--- a/Engine/Source/Editor/UMGEditor/Private/WidgetBlueprintCompiler.cpp
+++ b/Engine/Source/Editor/UMGEditor/Private/WidgetBlueprintCompiler.cpp
@@ -20,7 +20,7 @@
 #include "WidgetBlueprintEditorUtils.h"
 #include "WidgetGraphSchema.h"
 #include "IUMGModule.h"
-#include "UMGEditorProjectSettings.h"
+#include "WidgetEditingProjectSettings.h"
 #include "WidgetCompilerRule.h"
 #include "WidgetBlueprintExtension.h"
 #include "Editor/WidgetCompilerLog.h"
@@ -856,7 +856,7 @@ void FWidgetBlueprintCompilerContext::FinishCompilingClass(UClass* Class)
 				}
 			}
 
-			const EPropertyBindingPermissionLevel PropertyBindingRule = GetDefault<UUMGEditorProjectSettings>()->CompilerOption_PropertyBindingRule(WidgetBP);
+			const EPropertyBindingPermissionLevel PropertyBindingRule = WidgetBP->GetRelevantSettings()->CompilerOption_PropertyBindingRule(WidgetBP);
 			if (PropertyBindingRule != EPropertyBindingPermissionLevel::Allow)
 			{
 				if (WidgetBP->Bindings.Num() > 0)
@@ -884,7 +884,7 @@ void FWidgetBlueprintCompilerContext::FinishCompilingClass(UClass* Class)
 				}
 			}
 
-			if (!GetDefault<UUMGEditorProjectSettings>()->CompilerOption_AllowBlueprintTick(WidgetBP))
+			if (!WidgetBP->GetRelevantSettings()->CompilerOption_AllowBlueprintTick(WidgetBP))
 			{
 				const UFunction* ReceiveTickEvent = FKismetCompilerUtilities::FindOverriddenImplementableEvent(GET_FUNCTION_NAME_CHECKED(UUserWidget, Tick), NewWidgetBlueprintClass);
 				if (ReceiveTickEvent)
@@ -893,7 +893,7 @@ void FWidgetBlueprintCompilerContext::FinishCompilingClass(UClass* Class)
 				}
 			}
 
-			if (!GetDefault<UUMGEditorProjectSettings>()->CompilerOption_AllowBlueprintPaint(WidgetBP))
+			if (!WidgetBP->GetRelevantSettings()->CompilerOption_AllowBlueprintPaint(WidgetBP))
 			{
 				if (const UFunction* ReceivePaintEvent = FKismetCompilerUtilities::FindOverriddenImplementableEvent(GET_FUNCTION_NAME_CHECKED(UUserWidget, OnPaint), NewWidgetBlueprintClass))
 				{
@@ -903,7 +903,7 @@ void FWidgetBlueprintCompilerContext::FinishCompilingClass(UClass* Class)
 
 			// It's possible we may encounter some rules that haven't had a chance to load yet during early loading phases
 			// They're automatically removed from the returned set.
-			TArray<UWidgetCompilerRule*> CustomRules = GetDefault<UUMGEditorProjectSettings>()->CompilerOption_Rules(WidgetBP);
+			TArray<UWidgetCompilerRule*> CustomRules = WidgetBP->GetRelevantSettings()->CompilerOption_Rules(WidgetBP);
 			for (UWidgetCompilerRule* CustomRule : CustomRules)
 			{
 				CustomRule->ExecuteRule(WidgetBP, MessageLog);
diff --git a/Engine/Source/Editor/UMGEditor/Private/WidgetBlueprintEditorUtils.cpp b/Engine/Source/Editor/UMGEditor/Private/WidgetBlueprintEditorUtils.cpp
--- a/Engine/Source/Editor/UMGEditor/Private/WidgetBlueprintEditorUtils.cpp
+++ b/Engine/Source/Editor/UMGEditor/Private/WidgetBlueprintEditorUtils.cpp
@@ -63,6 +63,7 @@
 #include "Slate/WidgetRenderer.h"
 #include "Widgets/SVirtualWindow.h"
 #include "GraphEditorActions.h"
+#include "WidgetEditingProjectSettings.h"
 
 #define LOCTEXT_NAMESPACE "UMG"
 
@@ -818,7 +819,7 @@ void FWidgetBlueprintEditorUtils::BuildWrapWithMenu(FMenuBuilder& Menu, TSharedR
 	for ( TObjectIterator<UClass> ClassIt; ClassIt; ++ClassIt )
 	{
 		UClass* WidgetClass = *ClassIt;
-		if ( FWidgetBlueprintEditorUtils::IsUsableWidgetClass(WidgetClass) )
+		if ( FWidgetBlueprintEditorUtils::IsUsableWidgetClass(WidgetClass, BlueprintEditor) )
 		{
 			if ( WidgetClass->IsChildOf(UPanelWidget::StaticClass()) && !WidgetClass->HasAnyClassFlags(CLASS_HideDropDown) )
 			{
@@ -1007,7 +1008,7 @@ void FWidgetBlueprintEditorUtils::BuildReplaceWithMenu(FMenuBuilder& Menu, TShar
 		for ( TObjectIterator<UClass> ClassIt; ClassIt; ++ClassIt )
 		{
 			UClass* WidgetClass = *ClassIt;
-			if ( FWidgetBlueprintEditorUtils::IsUsableWidgetClass(WidgetClass) )
+			if ( FWidgetBlueprintEditorUtils::IsUsableWidgetClass(WidgetClass, BlueprintEditor) )
 			{
 				if ( WidgetClass->IsChildOf(UPanelWidget::StaticClass()) && !WidgetClass->HasAnyClassFlags(CLASS_HideDropDown) )
 				{
@@ -2171,13 +2172,12 @@ private:
 	bool bIsNormalBlueprintType;
 };
 
-bool IsUsableWidgetClass(const FString& WidgetPathName, const FAssetData& WidgetAssetData, FName Category, const UClass* WidgetClass)
+bool IsUsableWidgetClass(const FString& WidgetPathName, const FAssetData& WidgetAssetData, FName Category, const UClass* WidgetClass, TSharedRef<FWidgetBlueprintEditor> InCurrentActiveBlueprintEditor)
 {
-	const UUMGEditorProjectSettings* UMGEditorProjectSettings = GetDefault<UUMGEditorProjectSettings>();
-	const UContentBrowserSettings* ContentBrowserSettings = GetDefault<UContentBrowserSettings>();
+	const UWidgetEditingProjectSettings* UMGEditorProjectSettings = FWidgetBlueprintEditorUtils::GetRelevantSettings(InCurrentActiveBlueprintEditor);
 
 	// Excludes engine content if user sets it to false
-	if (!ContentBrowserSettings->GetDisplayEngineFolder() || !UMGEditorProjectSettings->bShowWidgetsFromEngineContent)
+	if (!UMGEditorProjectSettings->bShowWidgetsFromEngineContent)
 	{
 		if (WidgetPathName.StartsWith(TEXT("/Engine")))
 		{
@@ -2186,7 +2186,7 @@ bool IsUsableWidgetClass(const FString& WidgetPathName, const FAssetData& Widget
 	}
 
 	// Excludes developer content if user sets it to false
-	if (!ContentBrowserSettings->GetDisplayDevelopersFolder() || !UMGEditorProjectSettings->bShowWidgetsFromDeveloperContent)
+	if (!UMGEditorProjectSettings->bShowWidgetsFromDeveloperContent)
 	{
 		if (WidgetPathName.StartsWith(TEXT("/Game/Developers")))
 		{
@@ -2194,6 +2194,26 @@ bool IsUsableWidgetClass(const FString& WidgetPathName, const FAssetData& Widget
 		}
 	}
 
+	UWidgetBlueprint* WidgetBP = InCurrentActiveBlueprintEditor->GetWidgetBlueprintObj();
+	const bool bAllowEditorWidget = WidgetBP ? WidgetBP->AllowEditorWidget() : false;
+	if (!bAllowEditorWidget)
+	{
+		if (WidgetClass && IsEditorOnlyObject(WidgetClass))
+		{
+			return false;
+		}
+		else if (WidgetAssetData.IsValid())
+		{
+			// should not load since the default for GetClass is EResolveClass::No
+			const UClass* AssetClass = WidgetAssetData.GetClass();
+			if (AssetClass && IsEditorOnlyObject(AssetClass))
+			{
+				return false;
+			}
+			
+		}
+	}
+
 	if (UMGEditorProjectSettings->bUseEditorConfigPaletteFiltering)
 	{
 		if (UMGEditorProjectSettings->GetAllowedPaletteWidgets().PassesFilter(WidgetPathName))
@@ -2233,6 +2253,16 @@ bool IsUsableWidgetClass(const FString& WidgetPathName, const FAssetData& Widget
 }
 
 bool FWidgetBlueprintEditorUtils::IsUsableWidgetClass(const UClass* WidgetClass)
+{
+	return false;
+}
+
+TValueOrError<FWidgetBlueprintEditorUtils::FUsableWidgetClassResult, void> FWidgetBlueprintEditorUtils::IsUsableWidgetClass(const FAssetData& WidgetAsset)
+{
+	return MakeError();
+}
+
+bool FWidgetBlueprintEditorUtils::IsUsableWidgetClass(const UClass* WidgetClass, TSharedRef<FWidgetBlueprintEditor> BlueprintEditor)
 {
 	if (WidgetClass->IsChildOf(UWidget::StaticClass()))
 	{
@@ -2255,17 +2285,17 @@ bool FWidgetBlueprintEditorUtils::IsUsableWidgetClass(const UClass* WidgetClass)
 			return false;
 		}
 
-		return UE::UMG::Private::IsUsableWidgetClass(WidgetClass->GetPathName(), FAssetData(), *WidgetClass->GetDefaultObject<UWidget>()->GetPaletteCategory().ToString(), WidgetClass);
+		return UE::UMG::Private::IsUsableWidgetClass(WidgetClass->GetPathName(), FAssetData(), *WidgetClass->GetDefaultObject<UWidget>()->GetPaletteCategory().ToString(), WidgetClass, BlueprintEditor);
 	}
 
 	return false;
 }
 
-TValueOrError<FWidgetBlueprintEditorUtils::FUsableWidgetClassResult, void> FWidgetBlueprintEditorUtils::IsUsableWidgetClass(const FAssetData& WidgetAsset)
+TValueOrError<FWidgetBlueprintEditorUtils::FUsableWidgetClassResult, void> FWidgetBlueprintEditorUtils::IsUsableWidgetClass(const FAssetData& WidgetAsset, TSharedRef<FWidgetBlueprintEditor> InCurrentActiveBlueprintEditor)
 {
 	if (const UClass* WidgetAssetClass = WidgetAsset.GetClass(EResolveClass::No))
 	{
-		if (IsUsableWidgetClass(WidgetAssetClass))
+		if (IsUsableWidgetClass(WidgetAssetClass, InCurrentActiveBlueprintEditor))
 		{
 			FWidgetBlueprintEditorUtils::FUsableWidgetClassResult Result;
 			Result.NativeParentClass = WidgetAssetClass;
@@ -2309,7 +2339,7 @@ TValueOrError<FWidgetBlueprintEditorUtils::FUsableWidgetClassResult, void> FWidg
 	}
 
 	const FName CategoryName = *GetPaletteCategory(WidgetAsset, TSubclassOf<UWidget>(NativeParentClass)).ToString();
-	if (UE::UMG::Private::IsUsableWidgetClass(WidgetAsset.GetObjectPathString(), WidgetAsset, CategoryName, nullptr))
+	if (UE::UMG::Private::IsUsableWidgetClass(WidgetAsset.GetObjectPathString(), WidgetAsset, CategoryName, nullptr, InCurrentActiveBlueprintEditor))
 	{
 		FWidgetBlueprintEditorUtils::FUsableWidgetClassResult Result;
 		Result.NativeParentClass = NativeParentClass;
@@ -2540,6 +2570,32 @@ TOptional<FWidgetBlueprintEditorUtils::FWidgetThumbnailProperties> FWidgetBluepr
 	return DrawSWidgetInRenderTargetInternal(WidgetInstance, nullptr, RenderTarget2D, FVector2D(256.f,256.f), false, TOptional<FVector2D>(), EThumbnailPreviewSizeMode::MatchDesignerMode);
 }
 
+UWidgetEditingProjectSettings* FWidgetBlueprintEditorUtils::GetRelevantMutableSettings(TWeakPtr<FWidgetBlueprintEditor> CurrentEditor)
+{
+	if (TSharedPtr<FWidgetBlueprintEditor> PinnedEditor = CurrentEditor.Pin())
+	{
+		if (UWidgetBlueprint* WidgetBP = PinnedEditor->GetWidgetBlueprintObj())
+		{
+			return WidgetBP->GetRelevantSettings();
+		}
+	}
+	// Fall back to the UMG Editor settings as default
+	return GetMutableDefault<UUMGEditorProjectSettings>();
+}
+
+const UWidgetEditingProjectSettings* FWidgetBlueprintEditorUtils::GetRelevantSettings(TWeakPtr<FWidgetBlueprintEditor> CurrentEditor)
+{
+	if (TSharedPtr<FWidgetBlueprintEditor> PinnedEditor = CurrentEditor.Pin())
+	{
+		if (UWidgetBlueprint* WidgetBP = PinnedEditor->GetWidgetBlueprintObj())
+		{
+			return WidgetBP->GetRelevantSettings();
+		}
+	}
+	// Fall back to the UMG Editor settings as default
+	return GetDefault<UUMGEditorProjectSettings>();
+}
+
 TOptional<FWidgetBlueprintEditorUtils::FWidgetThumbnailProperties> FWidgetBlueprintEditorUtils::DrawSWidgetInRenderTargetForThumbnail(UUserWidget* WidgetInstance, FRenderTarget* RenderTarget2D, FVector2D ThumbnailSize, TOptional<FVector2D> ThumbnailCustomSize, EThumbnailPreviewSizeMode ThumbnailSizeMode)
 {
 	return DrawSWidgetInRenderTargetInternal(WidgetInstance, RenderTarget2D, nullptr,ThumbnailSize, true, ThumbnailCustomSize, ThumbnailSizeMode);
diff --git a/Engine/Source/Editor/UMGEditor/Public/Designer/SDesignerView.h b/Engine/Source/Editor/UMGEditor/Public/Designer/SDesignerView.h
--- a/Engine/Source/Editor/UMGEditor/Public/Designer/SDesignerView.h
+++ b/Engine/Source/Editor/UMGEditor/Public/Designer/SDesignerView.h
@@ -38,6 +38,7 @@ class UWidgetBlueprint;
 class FHittestGrid;
 struct FOnPaintHandlerParams;
 struct FWidgetHitResult;
+class UWidgetEditingProjectSettings;
 
 /**
  * The designer for widgets.  Allows for laying out widgets in a drag and drop environment.
@@ -124,6 +125,8 @@ public:
 
 	void BeginResizingArea();
 	void EndResizingArea();
+	
+	const UWidgetEditingProjectSettings* GetRelevantSettings() const;
 
 protected:
 	virtual void OnPaintBackground(const FGeometry& AllottedGeometry, const FSlateRect& MyCullingRect, FSlateWindowElementList& OutDrawElements, int32 LayerId) const override;
diff --git a/Engine/Source/Editor/UMGEditor/Public/UMGEditorProjectSettings.h b/Engine/Source/Editor/UMGEditor/Public/UMGEditorProjectSettings.h
--- a/Engine/Source/Editor/UMGEditor/Public/UMGEditorProjectSettings.h
+++ b/Engine/Source/Editor/UMGEditor/Public/UMGEditorProjectSettings.h
@@ -7,7 +7,7 @@
 #include "UObject/ObjectMacros.h"
 #include "UObject/Object.h"
 #include "UObject/SoftObjectPath.h"
-#include "Engine/DeveloperSettings.h"
+#include "WidgetEditingProjectSettings.h"
 #include "Engine/EngineTypes.h"
 #include "Misc/NamePermissionList.h"
 #include "UMGEditorProjectSettings.generated.h"
@@ -17,110 +17,12 @@ class UUserWidget;
 class UWidgetBlueprint;
 class UPanelWidget;
 
-USTRUCT()
-struct FDebugResolution
-{
-	GENERATED_BODY()
-
-	UPROPERTY(EditAnywhere, Category = Resolution)
-	int32 Width = 0;
-
-	UPROPERTY(EditAnywhere, Category = Resolution)
-	int32 Height = 0;
-
-	UPROPERTY(EditAnywhere, Category=Resolution)
-	FString Description;
-
-	UPROPERTY(EditAnywhere, Category = Resolution)
-	FLinearColor Color = FLinearColor::Black;
-};
-
-/** Controls the level of support you want to have for widget property binding. */
-UENUM()
-enum class EPropertyBindingPermissionLevel : uint8
-{
-	/** Freely allow the use of property binding. */
-	Allow,
-	/**
-	 * Prevent any new property binding, will still allow you to edit widgets with property binding, but
-	 * the buttons will be missing on all existing widgets that don't have bindings.
-	 */
-	Prevent,
-	/**
-	 * Prevent any new property binding, and warn when compiling any existing bindings.
-	 */
-	PreventAndWarn,
-	/**
-	* Prevent any new property binding, and error when compiling any existing bindings.
-	*/
-	PreventAndError
-};
-
-USTRUCT()
-struct FWidgetCompilerOptions
-{
-	GENERATED_BODY()
-
-public:
-
-	/**
-	 * If you disable this, these widgets these compiler options apply to will not be allowed to implement Tick.
-	 */
-	UPROPERTY(EditAnywhere, Category = Compiler)
-	bool bAllowBlueprintTick = true;
-
-	/**
-	 * If you disable this, these widgets these compiler options apply to will not be allowed to implement Paint.
-	 */
-	UPROPERTY(EditAnywhere, Category = Compiler)
-	bool bAllowBlueprintPaint = true;
-
-	/**
-	 * Controls if you allow property bindings in widgets.  They can have a large performance impact if used.
-	 */
-	UPROPERTY(EditAnywhere, Category = Compiler)
-	EPropertyBindingPermissionLevel PropertyBindingRule = EPropertyBindingPermissionLevel::Allow;
-
-	/**
-	 * Custom rules.
-	 */
-	UPROPERTY(EditAnywhere, Category = Compiler)
-	TArray<TSoftClassPtr<UWidgetCompilerRule>> Rules;
-
-	//TODO Allow limiting delay nodes?
-	//TODO Allow limiting 'size of blueprint'
-	//TODO Allow preventing blueprint inheritance
-};
-
-
-/**  */
-USTRUCT()
-struct FDirectoryWidgetCompilerOptions
-{
-	GENERATED_BODY()
-
-public:
-
-	/** The directory to limit the rules effects to. */
-	UPROPERTY(EditAnywhere, Category = Compiler, meta = (ContentDir))
-	FDirectoryPath Directory;
-
-#if WITH_EDITORONLY_DATA
-	/** These widgets are ignored, and they will use the next most applicable directory to determine their rules. */
-	UPROPERTY(EditAnywhere, Category = Compiler)
-	TArray<TSoftObjectPtr<UWidgetBlueprint>> IgnoredWidgets;
-#endif // WITH_EDITORONLY_DATA
-
-	/** The directory specific compiler options for these widgets. */
-	UPROPERTY(EditAnywhere, Category = Compiler)
-	FWidgetCompilerOptions Options;
-};
 
 /**
  * Implements the settings for the UMG Editor Project Settings
  */
 UCLASS(config=Editor, defaultconfig)
-class UMGEDITOR_API UUMGEditorProjectSettings : public UDeveloperSettings
+class UMGEDITOR_API UUMGEditorProjectSettings : public UWidgetEditingProjectSettings
 {
 	GENERATED_BODY()
 
@@ -132,139 +34,10 @@ public:
 	virtual FText GetSectionDescription() const override;
 #endif
 
-	virtual void PostInitProperties() override;
-
 #if WITH_EDITOR
 	// Begin UObject Interface
 	virtual void PostEditChangeChainProperty(struct FPropertyChangedChainEvent& PropertyChangedEvent) override;
 	// End UObject Interface
 #endif
 
-public:
-	UPROPERTY(EditAnywhere, config, Category = Compiler)
-	FWidgetCompilerOptions DefaultCompilerOptions;
-
-private:
-	UPROPERTY(EditAnywhere, config, Category = Compiler)
-	TArray<FDirectoryWidgetCompilerOptions> DirectoryCompilerOptions;
-
-public:
-
-	UPROPERTY(EditAnywhere, config, Category="Class Filtering")
-	bool bShowWidgetsFromEngineContent;
-
-	UPROPERTY(EditAnywhere, config, Category="Class Filtering")
-	bool bShowWidgetsFromDeveloperContent;
-
-	UPROPERTY(EditAnywhere, config, AdvancedDisplay, Category="Class Filtering")
-	TArray<FString> CategoriesToHide;
-
-	UPROPERTY(EditAnywhere, config, Category = "Class Filtering", meta = (MetaClass = "/Script/UMG.Widget"))
-	TArray<FSoftClassPath> WidgetClassesToHide;
-
-public:
-
-	/** Enables a dialog that lets you select a root widget before creating a widget blueprint */
-	UPROPERTY(EditAnywhere, config, Category = Designer)
-	bool bUseWidgetTemplateSelector;
-
-	/** This list populates the common class section of the root widget selection dialog */
-	UPROPERTY(EditAnywhere, config, Category = Designer, meta = (EditCondition = "bUseWidgetTemplateSelector"))
-	TArray<TSoftClassPtr<UPanelWidget>> CommonRootWidgetClasses;
-
-	/** The panel widget to place at the root of all newly constructed widget blueprints. Can be empty. */
-	UPROPERTY(EditAnywhere, config, Category = Designer)
-	TSubclassOf<UPanelWidget> DefaultRootWidget;
-
-	/** Set true to enable the Is Variable checkbox in the UMG editor DetailView. */
-	UPROPERTY(EditAnywhere, Category = "Designer | Window")
-	bool bEnableMakeVariable;
-
-	/** Set true to hide widget animation related elements in the UMG editor. */
-	UPROPERTY(EditAnywhere, Category = "Designer | Window")
-	bool bEnableWidgetAnimationEditor;
-
-	/** Set true to enabled the Palette window in the UMG editor. */
-	UPROPERTY(EditAnywhere, Category = "Designer | Window")
-	bool bEnablePaletteWindow;
-
-	/** Set true to enabled the LIbrary window in the UMG editor. */
-	UPROPERTY(EditAnywhere, Category = "Designer | Window")
-	bool bEnableLibraryWindow;
-
-	/** Set true to enabled the Widget Hierarchy window in the UMG editor. */
-	UPROPERTY(EditAnywhere, Category = "Designer | Window")
-	bool bEnableHierarchyWindow;
-	
-	/** Set true to enabled the Bind Widget window in the UMG editor. */
-	UPROPERTY(EditAnywhere, Category = "Designer | Window")
-	bool bEnableBindWidgetWindow;
-
-	/** Set true to enabled the Navigation Simulation window in the UMG editor. */
-	UPROPERTY(EditAnywhere, Category = "Designer | Window")
-	bool bEnableNavigationSimulationWindow;
-
-	/** Set true to filter all categories and widgets out in the palette, selectively enabling them later via permission lists. */
-	UPROPERTY(EditAnywhere, config, Category = Designer)
-	bool bUseEditorConfigPaletteFiltering;
-
-	/** Enables a dialog that lets you select the parent class in a tree view. */
-	UPROPERTY(EditAnywhere, config, Category = Designer)
-	bool bUseUserWidgetParentClassViewerSelector;
-
-	/** Enables a dialog that lets you select the parent class in a default view. */
-	UPROPERTY(EditAnywhere, config, Category = Designer)
-	bool bUseUserWidgetParentDefaultClassViewerSelector;
-
-	/**
-	 * The list of parent classes to choose from for newly constructed widget blueprints.
-	 * The classes must have empty widget hierarchies.
-	 */
-	UPROPERTY(EditAnywhere, config, Category = Designer, meta = (AllowAbstract = "", EditCondition="bUseUserWidgetParentDefaultClassViewerSelector"))
-	TArray<TSoftClassPtr<UUserWidget>> FavoriteWidgetParentClasses;
-
-	UPROPERTY(EditAnywhere, config, Category=Designer)
-	TArray<FDebugResolution> DebugResolutions;
-
-	bool CompilerOption_AllowBlueprintTick(const class UWidgetBlueprint* WidgetBlueprint) const;
-	bool CompilerOption_AllowBlueprintPaint(const class UWidgetBlueprint* WidgetBlueprint) const;
-	EPropertyBindingPermissionLevel CompilerOption_PropertyBindingRule(const class UWidgetBlueprint* WidgetBlueprint) const;
-	TArray<UWidgetCompilerRule*> CompilerOption_Rules(const class UWidgetBlueprint* WidgetBlueprint) const;
-
-	/** Get the permission list that controls which categories are exposed in config palette filtering */
-	FNamePermissionList& GetAllowedPaletteCategories();
-	const FNamePermissionList& GetAllowedPaletteCategories() const;
-
-	/** Get the permission list that controls which widgets are exposed in config palette filtering */
-	FPathPermissionList& GetAllowedPaletteWidgets();
-	const FPathPermissionList& GetAllowedPaletteWidgets() const;
-
-private:
-	template<typename ReturnType, typename T>
-	ReturnType GetFirstCompilerOption(const class UWidgetBlueprint* WidgetBlueprint, T FWidgetCompilerOptions::* OptionMember, ReturnType bDefaultValue) const
-	{
-		ReturnType bValue = bDefaultValue;
-		GetCompilerOptionsForWidget(WidgetBlueprint, [&bValue, OptionMember](const FWidgetCompilerOptions& Options) {
-			bValue = Options.*OptionMember;
-			return true;
-		});
-		return bValue;
-	}
-
-	void GetCompilerOptionsForWidget(const class UWidgetBlueprint* WidgetBlueprint, TFunctionRef<bool(const FWidgetCompilerOptions&)> Operator) const;
-
-protected:
-	virtual void PerformUpgradeStepForVersion(int32 ForVersion);
-
-	UPROPERTY(config)
-	int32 Version;
-
-	/** This one is unsaved, we compare it on post init to see if the save matches real */
-	int32 CurrentVersion;
-
-	/** Palette categories to allow all widgets within when using permission list palette filtering */
-	FNamePermissionList AllowedPaletteCategories;
-
-	/** Individual widgets to always allow when using permission list palette filtering, regardless of category */
-	FPathPermissionList AllowedPaletteWidgets;
 };
diff --git a/Engine/Source/Editor/UMGEditor/Public/WidgetBlueprint.h b/Engine/Source/Editor/UMGEditor/Public/WidgetBlueprint.h
--- a/Engine/Source/Editor/UMGEditor/Public/WidgetBlueprint.h
+++ b/Engine/Source/Editor/UMGEditor/Public/WidgetBlueprint.h
@@ -25,6 +25,7 @@ class FKismetCompilerContext;
 class UWidgetBlueprint;
 enum class EWidgetTickFrequency : uint8;
 enum class EWidgetCompileTimeTickPrediction : uint8;
+class UWidgetEditingProjectSettings;
 
 
 /** Widget Delegates */
@@ -329,6 +330,9 @@ public:
 	/** Gets any named slots exposed by the parent generated class that can be slotted into by the subclass. */
 	TArray<FName> GetInheritedAvailableNamedSlots() const;
 
+	virtual UWidgetEditingProjectSettings* GetRelevantSettings();
+	virtual const UWidgetEditingProjectSettings* GetRelevantSettings() const;
+
 protected:
 	virtual void LoadModulesRequiredForCompilation() override;
 
diff --git a/Engine/Source/Editor/UMGEditor/Public/WidgetBlueprintEditorUtils.h b/Engine/Source/Editor/UMGEditor/Public/WidgetBlueprintEditorUtils.h
--- a/Engine/Source/Editor/UMGEditor/Public/WidgetBlueprintEditorUtils.h
+++ b/Engine/Source/Editor/UMGEditor/Public/WidgetBlueprintEditorUtils.h
@@ -14,6 +14,7 @@ class UWidgetBlueprint;
 class UWidgetSlotPair;
 class UWidgetTree;
 class SWindow;
+class UWidgetEditingProjectSettings;
 
 //////////////////////////////////////////////////////////////////////////
 // FWidgetBlueprintEditorUtils
@@ -61,9 +62,15 @@ public:
 		const UClass* NativeParentClass = nullptr;
 		EClassFlags AssetClassFlags = EClassFlags::CLASS_None;
 	};
+
+	UE_DEPRECATED(5.3, "This function has been deprecated. Use the version of IsUsableWidgetClass that takes a second argument of TSharedRef<FWidgetBlueprintEditor>.")
 	static bool IsUsableWidgetClass(const UClass* WidgetClass);
+	UE_DEPRECATED(5.3, "This function has been deprecated. Use the version of IsUsableWidgetClass that takes a second argument of TSharedRef<FWidgetBlueprintEditor>.")
 	static TValueOrError<FUsableWidgetClassResult, void> IsUsableWidgetClass(const FAssetData& WidgetAsset);
 
+	static bool IsUsableWidgetClass(const UClass* WidgetClass, TSharedRef<FWidgetBlueprintEditor> InCurrentActiveBlueprintEditor);
+	static TValueOrError<FUsableWidgetClassResult, void> IsUsableWidgetClass(const FAssetData& WidgetAsset, TSharedRef<FWidgetBlueprintEditor> InCurrentActiveBlueprintEditor);
+
 	static void ExportWidgetsToText(TArray<UWidget*> WidgetsToExport, /*out*/ FString& ExportedText);
 
 	static void ImportWidgetsFromText(UWidgetBlueprint* BP, const FString& TextToImport, /*out*/ TSet<UWidget*>& ImportedWidgetSet, /*out*/ TMap<FName, UWidgetSlotPair*>& PastedExtraSlotData);
@@ -109,6 +116,9 @@ public:
 
 	static TOptional<FWidgetThumbnailProperties> DrawSWidgetInRenderTarget(UUserWidget* WidgetInstance, UTextureRenderTarget2D* RenderTarget2D);
 
+	static UWidgetEditingProjectSettings* GetRelevantMutableSettings(TWeakPtr<FWidgetBlueprintEditor> CurrentEditor);
+	static const UWidgetEditingProjectSettings* GetRelevantSettings(TWeakPtr<FWidgetBlueprintEditor> CurrentEditor);
+
 private:
 
 	static FString CopyWidgetsInternal(UWidgetBlueprint* BP, TSet<FWidgetReference> Widgets);
diff --git a/Engine/Source/Editor/UMGEditor/Public/WidgetEditingProjectSettings.h b/Engine/Source/Editor/UMGEditor/Public/WidgetEditingProjectSettings.h
new file mode 100644
--- /dev/null
+++ b/Engine/Source/Editor/UMGEditor/Public/WidgetEditingProjectSettings.h
@@ -0,0 +1,259 @@
+// Copyright Epic Games, Inc. All Rights Reserved.
+
+#pragma once
+
+#include "CoreMinimal.h"
+#include "EditorConfigBase.h"
+#include "UObject/ObjectMacros.h"
+#include "UObject/Object.h"
+#include "UObject/SoftObjectPath.h"
+#include "Engine/DeveloperSettings.h"
+#include "Engine/EngineTypes.h"
+#include "Misc/NamePermissionList.h"
+#include "WidgetEditingProjectSettings.generated.h"
+
+class UWidgetCompilerRule;
+class UUserWidget;
+class UWidgetBlueprint;
+class UPanelWidget;
+
+USTRUCT()
+struct FDebugResolution
+{
+	GENERATED_BODY()
+
+	UPROPERTY(EditAnywhere, Category = Resolution)
+	int32 Width = 0;
+
+	UPROPERTY(EditAnywhere, Category = Resolution)
+	int32 Height = 0;
+
+	UPROPERTY(EditAnywhere, Category=Resolution)
+	FString Description;
+
+	UPROPERTY(EditAnywhere, Category = Resolution)
+	FLinearColor Color = FLinearColor::Black;
+};
+
+/** Controls the level of support you want to have for widget property binding. */
+UENUM()
+enum class EPropertyBindingPermissionLevel : uint8
+{
+	/** Freely allow the use of property binding. */
+	Allow,
+	/**
+	 * Prevent any new property binding, will still allow you to edit widgets with property binding, but
+	 * the buttons will be missing on all existing widgets that don't have bindings.
+	 */
+	Prevent,
+	/**
+	 * Prevent any new property binding, and warn when compiling any existing bindings.
+	 */
+	PreventAndWarn,
+	/**
+	* Prevent any new property binding, and error when compiling any existing bindings.
+	*/
+	PreventAndError
+};
+
+USTRUCT()
+struct FWidgetCompilerOptions
+{
+	GENERATED_BODY()
+
+public:
+
+	/**
+	 * If you disable this, these widgets these compiler options apply to will not be allowed to implement Tick.
+	 */
+	UPROPERTY(EditAnywhere, Category = Compiler)
+	bool bAllowBlueprintTick = true;
+
+	/**
+	 * If you disable this, these widgets these compiler options apply to will not be allowed to implement Paint.
+	 */
+	UPROPERTY(EditAnywhere, Category = Compiler)
+	bool bAllowBlueprintPaint = true;
+
+	/**
+	 * Controls if you allow property bindings in widgets.  They can have a large performance impact if used.
+	 */
+	UPROPERTY(EditAnywhere, Category = Compiler)
+	EPropertyBindingPermissionLevel PropertyBindingRule = EPropertyBindingPermissionLevel::Allow;
+
+	/**
+	 * Custom rules.
+	 */
+	UPROPERTY(EditAnywhere, Category = Compiler)
+	TArray<TSoftClassPtr<UWidgetCompilerRule>> Rules;
+
+	//TODO Allow limiting delay nodes?
+	//TODO Allow limiting 'size of blueprint'
+	//TODO Allow preventing blueprint inheritance
+};
+
+
+/**  */
+USTRUCT()
+struct FDirectoryWidgetCompilerOptions
+{
+	GENERATED_BODY()
+
+public:
+
+	/** The directory to limit the rules effects to. */
+	UPROPERTY(EditAnywhere, Category = Compiler, meta = (ContentDir))
+	FDirectoryPath Directory;
+
+#if WITH_EDITORONLY_DATA
+	/** These widgets are ignored, and they will use the next most applicable directory to determine their rules. */
+	UPROPERTY(EditAnywhere, Category = Compiler)
+	TArray<TSoftObjectPtr<UWidgetBlueprint>> IgnoredWidgets;
+#endif // WITH_EDITORONLY_DATA
+
+	/** The directory specific compiler options for these widgets. */
+	UPROPERTY(EditAnywhere, Category = Compiler)
+	FWidgetCompilerOptions Options;
+};
+
+/**
+ * Implements the settings for the UMG Editor Project Settings
+ */
+UCLASS(Abstract, config=Editor, defaultconfig)
+class UMGEDITOR_API UWidgetEditingProjectSettings : public UDeveloperSettings
+{
+	GENERATED_BODY()
+
+public:
+	UWidgetEditingProjectSettings();
+
+	virtual void PostInitProperties() override;
+
+public:
+	UPROPERTY(EditAnywhere, config, Category = Compiler)
+	FWidgetCompilerOptions DefaultCompilerOptions;
+
+protected:
+	UPROPERTY(EditAnywhere, config, Category = Compiler)
+	TArray<FDirectoryWidgetCompilerOptions> DirectoryCompilerOptions;
+
+public:
+
+	UPROPERTY(EditAnywhere, config, Category="Class Filtering")
+	bool bShowWidgetsFromEngineContent;
+
+	UPROPERTY(EditAnywhere, config, Category="Class Filtering")
+	bool bShowWidgetsFromDeveloperContent;
+
+	UPROPERTY(EditAnywhere, config, AdvancedDisplay, Category="Class Filtering")
+	TArray<FString> CategoriesToHide;
+
+	UPROPERTY(EditAnywhere, config, Category = "Class Filtering", meta = (MetaClass = "/Script/UMG.Widget"))
+	TArray<FSoftClassPath> WidgetClassesToHide;
+
+public:
+
+	/** Enables a dialog that lets you select a root widget before creating a widget blueprint */
+	UPROPERTY(EditAnywhere, config, Category = Designer)
+	bool bUseWidgetTemplateSelector;
+
+	/** This list populates the common class section of the root widget selection dialog */
+	UPROPERTY(EditAnywhere, config, Category = Designer, meta = (EditCondition = "bUseWidgetTemplateSelector"))
+	TArray<TSoftClassPtr<UPanelWidget>> CommonRootWidgetClasses;
+
+	/** The panel widget to place at the root of all newly constructed widget blueprints. Can be empty. */
+	UPROPERTY(EditAnywhere, config, Category = Designer)
+	TSubclassOf<UPanelWidget> DefaultRootWidget;
+
+	/** Set true to enable the Is Variable checkbox in the UMG editor DetailView. */
+	UPROPERTY(EditAnywhere, Category = "Designer | Window")
+	bool bEnableMakeVariable;
+
+	/** Set true to hide widget animation related elements in the UMG editor. */
+	UPROPERTY(EditAnywhere, Category = "Designer | Window")
+	bool bEnableWidgetAnimationEditor;
+
+	/** Set true to enabled the Palette window in the UMG editor. */
+	UPROPERTY(EditAnywhere, Category = "Designer | Window")
+	bool bEnablePaletteWindow;
+
+	/** Set true to enabled the LIbrary window in the UMG editor. */
+	UPROPERTY(EditAnywhere, Category = "Designer | Window")
+	bool bEnableLibraryWindow;
+
+	/** Set true to enabled the Widget Hierarchy window in the UMG editor. */
+	UPROPERTY(EditAnywhere, Category = "Designer | Window")
+	bool bEnableHierarchyWindow;
+	
+	/** Set true to enabled the Bind Widget window in the UMG editor. */
+	UPROPERTY(EditAnywhere, Category = "Designer | Window")
+	bool bEnableBindWidgetWindow;
+
+	/** Set true to enabled the Navigation Simulation window in the UMG editor. */
+	UPROPERTY(EditAnywhere, Category = "Designer | Window")
+	bool bEnableNavigationSimulationWindow;
+
+	/** Set true to filter all categories and widgets out in the palette, selectively enabling them later via permission lists. */
+	UPROPERTY(EditAnywhere, config, Category = Designer)
+	bool bUseEditorConfigPaletteFiltering;
+
+	/** Enables a dialog that lets you select the parent class in a tree view. */
+	UPROPERTY(EditAnywhere, config, Category = Designer)
+	bool bUseUserWidgetParentClassViewerSelector;
+
+	/** Enables a dialog that lets you select the parent class in a default view. */
+	UPROPERTY(EditAnywhere, config, Category = Designer)
+	bool bUseUserWidgetParentDefaultClassViewerSelector;
+
+	/**
+	 * The list of parent classes to choose from for newly constructed widget blueprints.
+	 * The classes must have empty widget hierarchies.
+	 */
+	UPROPERTY(EditAnywhere, config, Category = Designer, meta = (AllowAbstract = "", EditCondition="bUseUserWidgetParentDefaultClassViewerSelector"))
+	TArray<TSoftClassPtr<UUserWidget>> FavoriteWidgetParentClasses;
+
+	UPROPERTY(EditAnywhere, config, Category=Designer)
+	TArray<FDebugResolution> DebugResolutions;
+
+	bool CompilerOption_AllowBlueprintTick(const class UWidgetBlueprint* WidgetBlueprint) const;
+	bool CompilerOption_AllowBlueprintPaint(const class UWidgetBlueprint* WidgetBlueprint) const;
+	EPropertyBindingPermissionLevel CompilerOption_PropertyBindingRule(const class UWidgetBlueprint* WidgetBlueprint) const;
+	TArray<UWidgetCompilerRule*> CompilerOption_Rules(const class UWidgetBlueprint* WidgetBlueprint) const;
+
+	/** Get the permission list that controls which categories are exposed in config palette filtering */
+	FNamePermissionList& GetAllowedPaletteCategories();
+	const FNamePermissionList& GetAllowedPaletteCategories() const;
+
+	/** Get the permission list that controls which widgets are exposed in config palette filtering */
+	FPathPermissionList& GetAllowedPaletteWidgets();
+	const FPathPermissionList& GetAllowedPaletteWidgets() const;
+
+private:
+	template<typename ReturnType, typename T>
+	ReturnType GetFirstCompilerOption(const class UWidgetBlueprint* WidgetBlueprint, T FWidgetCompilerOptions::* OptionMember, ReturnType bDefaultValue) const
+	{
+		ReturnType bValue = bDefaultValue;
+		GetCompilerOptionsForWidget(WidgetBlueprint, [&bValue, OptionMember](const FWidgetCompilerOptions& Options) {
+			bValue = Options.*OptionMember;
+			return true;
+		});
+		return bValue;
+	}
+
+	void GetCompilerOptionsForWidget(const class UWidgetBlueprint* WidgetBlueprint, TFunctionRef<bool(const FWidgetCompilerOptions&)> Operator) const;
+
+protected:
+	virtual void PerformUpgradeStepForVersion(int32 ForVersion);
+
+	UPROPERTY(config)
+	int32 Version;
+
+	/** This one is unsaved, we compare it on post init to see if the save matches real */
+	int32 CurrentVersion;
+
+	/** Palette categories to allow all widgets within when using permission list palette filtering */
+	FNamePermissionList AllowedPaletteCategories;
+
+	/** Individual widgets to always allow when using permission list palette filtering, regardless of category */
+	FPathPermissionList AllowedPaletteWidgets;
+};
