Display different corner text in the graph in the Widget Blueprint Editor for runtime widgets vs. editor utility widgets

diff --git a/Engine/Source/Editor/UMGEditor/Private/WidgetBlueprintEditor.cpp b/Engine/Source/Editor/UMGEditor/Private/WidgetBlueprintEditor.cpp
--- a/Engine/Source/Editor/UMGEditor/Private/WidgetBlueprintEditor.cpp
+++ b/Engine/Source/Editor/UMGEditor/Private/WidgetBlueprintEditor.cpp
@@ -1717,7 +1717,11 @@ FGraphAppearanceInfo FWidgetBlueprintEditor::GetGraphAppearance(UEdGraph* InGrap
 {
 	FGraphAppearanceInfo AppearanceInfo = Super::GetGraphAppearance(InGraph);
 
-	if ( GetBlueprintObj()->IsA(UWidgetBlueprint::StaticClass()) )
+	if (FBlueprintEditorUtils::IsEditorUtilityBlueprint(GetBlueprintObj()))
+	{
+		AppearanceInfo.CornerText = LOCTEXT("EditorUtilityWidgetAppearanceCornerText", "EDITOR UTILITY WIDGET");
+	}
+	else if ( GetBlueprintObj()->IsA(UWidgetBlueprint::StaticClass()) )
 	{
 		AppearanceInfo.CornerText = LOCTEXT("AppearanceCornerText", "WIDGET BLUEPRINT");
 	}
