Don't reset text IDs when restoring the old widget tree state post-compile to avoid cook determinism issues

This was likely a regression caused by 13773255, as prior to that change each UUserWidget managed its own persistent widget tree and would use PPF_DuplicateVerbatim when building it. That use of PPF_DuplicateVerbatim was lost when that code was removed in favor of building an archetype widget tree on the BPGC.

diff --git a/Engine/Source/Editor/UMGEditor/Private/WidgetBlueprintCompiler.cpp b/Engine/Source/Editor/UMGEditor/Private/WidgetBlueprintCompiler.cpp
--- a/Engine/Source/Editor/UMGEditor/Private/WidgetBlueprintCompiler.cpp
+++ b/Engine/Source/Editor/UMGEditor/Private/WidgetBlueprintCompiler.cpp
@@ -772,6 +772,7 @@ void FWidgetBlueprintCompilerContext::FinishCompilingClass(UClass* Class)
 			FObjectDuplicationParameters DupParams(WidgetBP->WidgetTree, BPGClass);
 			DupParams.DestName = DupParams.SourceObject->GetFName();
 			DupParams.FlagMask = RF_AllFlags & ~RF_DefaultSubObject;
+			DupParams.PortFlags |= PPF_DuplicateVerbatim; // Skip resetting text IDs
 
 			// if we are recompiling the BP on load, skip post load and defer it to the loading process
 			FUObjectSerializeContext* LinkerLoadingContext = nullptr;
