UMG: Add option to expand all the widgets when the view is created. True by default to keep previous behavior.

diff --git a/Engine/Source/Editor/UMGEditor/Private/Hierarchy/SReadOnlyHierarchyView.cpp b/Engine/Source/Editor/UMGEditor/Private/Hierarchy/SReadOnlyHierarchyView.cpp
--- a/Engine/Source/Editor/UMGEditor/Private/Hierarchy/SReadOnlyHierarchyView.cpp
+++ b/Engine/Source/Editor/UMGEditor/Private/Hierarchy/SReadOnlyHierarchyView.cpp
@@ -15,6 +15,7 @@ void SReadOnlyHierarchyView::Construct(const FArguments& InArgs, const UWidgetBl
 	WidgetBlueprint = InWidgetBlueprint;
 	ShowOnly = InArgs._ShowOnly;
 	RootSelectionMode = InArgs._RootSelectionMode;
+	bExpandAll = InArgs._ExpandAll;
 
 	SearchFilter = MakeShared<FTextFilter>(FTextFilter::FItemToStringArray::CreateSP(this, &SReadOnlyHierarchyView::GetFilterStringsForItem));
 	FilterHandler = MakeShared<FTreeFilterHandler>();
@@ -57,8 +58,6 @@ void SReadOnlyHierarchyView::Construct(const FArguments& InArgs, const UWidgetBl
 	[
 		ContentBox
 	];
-
-	ExpandAll();
 }
 
 SReadOnlyHierarchyView::~SReadOnlyHierarchyView()
@@ -96,7 +95,10 @@ void SReadOnlyHierarchyView::Refresh()
 	FilteredRootWidgets.Reset();
 	RebuildTree();
 	FilterHandler->RefreshAndFilterTree();
-	ExpandAll();
+	if (bExpandAll)
+	{
+		ExpandAll();
+	}
 }
 
 void SReadOnlyHierarchyView::SetItemExpansionRecursive(TSharedPtr<FItem> Item, bool bShouldBeExpanded)
@@ -326,7 +328,10 @@ void SReadOnlyHierarchyView::RebuildTree()
 		BuildWidgetChildren(RootWidget, RootItem);
 	}
 
-	ExpandAll();
+	if (bExpandAll)
+	{
+		ExpandAll();
+	}
 }
 
 void SReadOnlyHierarchyView::ExpandAll()
diff --git a/Engine/Source/Editor/UMGEditor/Public/Hierarchy/SReadOnlyHierarchyView.h b/Engine/Source/Editor/UMGEditor/Public/Hierarchy/SReadOnlyHierarchyView.h
--- a/Engine/Source/Editor/UMGEditor/Public/Hierarchy/SReadOnlyHierarchyView.h
+++ b/Engine/Source/Editor/UMGEditor/Public/Hierarchy/SReadOnlyHierarchyView.h
@@ -58,6 +58,7 @@ public:
 		SLATE_ARGUMENT_DEFAULT(ESelectionMode::Type, SelectionMode) = ESelectionMode::Single;
 		SLATE_EVENT(FOnSelectionChanged, OnSelectionChanged)
 		SLATE_ARGUMENT(TArray<FName>, ShowOnly)
+		SLATE_ARGUMENT_DEFAULT(bool, ExpandAll) = true;
 	SLATE_END_ARGS()
 
 	virtual ~SReadOnlyHierarchyView();
@@ -112,4 +113,6 @@ private:
 	TSharedPtr<FTreeFilterHandler> FilterHandler;
 
 	TSharedPtr<STreeView<TSharedPtr<FItem>>> TreeView;
+
+	bool bExpandAll;
 };
