UMG Animation: Check for the animation binding in named slot content.

diff --git a/Engine/Source/Editor/UMGEditor/Private/WidgetBlueprintCompiler.cpp b/Engine/Source/Editor/UMGEditor/Private/WidgetBlueprintCompiler.cpp
--- a/Engine/Source/Editor/UMGEditor/Private/WidgetBlueprintCompiler.cpp
+++ b/Engine/Source/Editor/UMGEditor/Private/WidgetBlueprintCompiler.cpp
@@ -1168,9 +1168,13 @@ void FWidgetBlueprintCompilerContext::ValidateWidgetAnimations()
 			// If any of the FoundObjects is null, we do not play the animation.
 			if (FoundObject == nullptr)
 			{
-				// Notify the user of the null track in the editor
-				const FText AnimationNullTrackMessage = LOCTEXT("AnimationNullTrack", "UMG Animation '{0}' from '{1}' is trying to animate a non-existent widget through binding '{2}'. Please re-bind or delete this object from the animation.");
-				BlueprintLog.Warning(FText::Format(AnimationNullTrackMessage, InAnimation->GetDisplayName(), FText::FromString(UserWidget->GetClass()->GetName()), FText::FromName(Binding.WidgetName)));
+				FoundObject = Binding.FindRuntimeObject(*WidgetBP->WidgetTree, *UserWidget);
+				if (FoundObject == nullptr)
+				{
+					// Notify the user of the null track in the editor
+					const FText AnimationNullTrackMessage = LOCTEXT("AnimationNullTrack", "UMG Animation '{0}' from '{1}' is trying to animate a non-existent widget through binding '{2}'. Please re-bind or delete this object from the animation.");
+					BlueprintLog.Warning(FText::Format(AnimationNullTrackMessage, InAnimation->GetDisplayName(), FText::FromString(UserWidget->GetClass()->GetName()), FText::FromName(Binding.WidgetName)));
+				}
 			}
 		}
 	}
