Fix crash in UMG desginer when pasting multiple widgets into the slot of a single-slot widget.
Don't try to insert at a specific index if all we allow is one slot.

diff --git a/Engine/Source/Editor/UMGEditor/Private/WidgetBlueprintEditorUtils.cpp b/Engine/Source/Editor/UMGEditor/Private/WidgetBlueprintEditorUtils.cpp
--- a/Engine/Source/Editor/UMGEditor/Private/WidgetBlueprintEditorUtils.cpp
+++ b/Engine/Source/Editor/UMGEditor/Private/WidgetBlueprintEditorUtils.cpp
@@ -1707,7 +1707,7 @@ TArray<UWidget*> FWidgetBlueprintEditorUtils::PasteWidgetsInternal(TSharedRef<FW
 				if (UWidget* WidgetTemplate = ParentWidgetRef.GetTemplate())
 				{
 					ParentWidget = WidgetTemplate->GetParent();
-					if (ParentWidget)
+					if (ParentWidget && ParentWidget->CanHaveMultipleChildren())
 					{
 						IndexToInsert = ParentWidget->GetChildIndex(WidgetTemplate) + 1;
 					}
