Fix to bug in UMG Designer Palette where the LibraryView was removing all Favorites setupped by the Palette view (broken by  CL22513071).

diff --git a/Engine/Source/Editor/UMGEditor/Classes/WidgetPaletteFavorites.cpp b/Engine/Source/Editor/UMGEditor/Classes/WidgetPaletteFavorites.cpp
--- a/Engine/Source/Editor/UMGEditor/Classes/WidgetPaletteFavorites.cpp
+++ b/Engine/Source/Editor/UMGEditor/Classes/WidgetPaletteFavorites.cpp
@@ -20,7 +20,7 @@ void UWidgetPaletteFavorites::Add(const FString& InWidgetTemplateName)
 
 void UWidgetPaletteFavorites::Remove(const FString& InWidgetTemplateName)
 {
-	Favorites.Remove(InWidgetTemplateName);
+	Favorites.RemoveSingle(InWidgetTemplateName);
 
 	SaveConfig();
 
diff --git a/Engine/Source/Editor/UMGEditor/Private/Library/SLibraryViewModel.cpp b/Engine/Source/Editor/UMGEditor/Private/Library/SLibraryViewModel.cpp
--- a/Engine/Source/Editor/UMGEditor/Private/Library/SLibraryViewModel.cpp
+++ b/Engine/Source/Editor/UMGEditor/Private/Library/SLibraryViewModel.cpp
@@ -31,6 +31,9 @@
 
 void FLibraryViewModel::BuildWidgetTemplateCategory(FString& Category, TArray<TSharedPtr<FWidgetTemplate>>& Templates, TArray<FString>& FavoritesList)
 {
+	//This function do not handle Favorites. So we empty the Favorite Array. See FLibraryViewModel::BuildWidgetTemplateCategory to implement favorite cleaning.
+	FavoritesList.Empty();
+
 	TSharedPtr<FWidgetHeaderViewModel> Header = MakeShareable(new FWidgetHeaderViewModel());
 	Header->GroupName = FText::FromString(Category);
 
@@ -39,49 +42,9 @@ void FLibraryViewModel::BuildWidgetTemplateCategory(FString& Category, TArray<TS
 	Header->Children.Add(TemplateViewModel);
 	WidgetTemplateListViewModels.Add(TemplateViewModel);
 
-	// @TODO: DarenC - Reference for when implementing favorites system 
-	//for ( auto& Template : Entry.Value )
-	//{
-	//	TSharedPtr<FWidgetTemplateViewModel> TemplateViewModel = MakeShareable(new FWidgetTemplateViewModel());
-	//	TemplateViewModel->Template = Template;
-	//	TemplateViewModel->FavortiesViewModel = this;
-	//	Header->Children.Add(TemplateViewModel);
-
-	//	// If it's a favorite, we also add it to the Favorite section
-	//	int32 index = FavoritesList.Find(Template->Name.ToString());
-	//	if (index != INDEX_NONE)
-	//	{
-	//		TemplateViewModel->SetFavorite();
-
-	//		// We have to create a second copy of the ViewModel for the treeview has it doesn't support to have the same element twice.
-	//		TSharedPtr<FWidgetTemplateViewModel> FavoriteTemplateViewModel = MakeShareable(new FWidgetTemplateViewModel());
-	//		FavoriteTemplateViewModel->Template = Template;
-	//		FavoriteTemplateViewModel->FavortiesViewModel = this;
-	//		FavoriteTemplateViewModel->SetFavorite();
-
-	//		FavoriteHeader->Children.Add(FavoriteTemplateViewModel);
-
-	//		// Remove the favorite from the temporary list
-	//		FavoritesList.RemoveAt(index);
-	//	}
-
-	//}
-
 	Header->Children.Sort([](const TSharedPtr<FWidgetViewModel>& L, const TSharedPtr<FWidgetViewModel>& R) { return R->GetName().CompareTo(L->GetName()) > 0; });
 
-	WidgetViewModels.Add(Header);
-}
-
-void FLibraryViewModel::AddToFavorites(const FWidgetTemplateViewModel* WidgetTemplateViewModel)
-{
-	UWidgetPaletteFavorites* Favorites = GetDefault<UWidgetDesignerSettings>()->Favorites;
-	Favorites->Add(WidgetTemplateViewModel->GetName().ToString());
-}
-
-void FLibraryViewModel::RemoveFromFavorites(const FWidgetTemplateViewModel* WidgetTemplateViewModel)
-{
-	UWidgetPaletteFavorites* Favorites = GetDefault<UWidgetDesignerSettings>()->Favorites;
-	Favorites->Remove(WidgetTemplateViewModel->GetName().ToString());
+	AddHeader(Header);
 }
 
 void FLibraryViewModel::SetSearchText(const FText& InSearchText)
@@ -111,7 +74,6 @@ void FWidgetTemplateListViewModel::ConstructListView(TArray<TSharedPtr<FWidgetTe
 
 	if (!TemplatesFilter)
 	{
-		UContentBrowserDataSubsystem* ContentBrowserDataSubsystem = IContentBrowserDataModule::Get().GetSubsystem();
 		NumAssets = 0;
 		NumClasses = 0;
 
diff --git a/Engine/Source/Editor/UMGEditor/Private/Library/SLibraryViewModel.h b/Engine/Source/Editor/UMGEditor/Private/Library/SLibraryViewModel.h
--- a/Engine/Source/Editor/UMGEditor/Private/Library/SLibraryViewModel.h
+++ b/Engine/Source/Editor/UMGEditor/Private/Library/SLibraryViewModel.h
@@ -75,12 +75,9 @@ public:
 
 	//~ Begin FWidgetCatalogViewModel Interface
 	virtual void BuildWidgetTemplateCategory(FString& Category, TArray<TSharedPtr<FWidgetTemplate>>& Templates, TArray<FString>& FavoritesList) override;
-	virtual void AddToFavorites(const FWidgetTemplateViewModel* WidgetTemplateViewModel) override;
-	virtual void RemoveFromFavorites(const FWidgetTemplateViewModel* WidgetTemplateViewModel) override;
 	virtual void SetSearchText(const FText& InSearchText) override;
 	//~ End FWidgetCatalogViewModel Interface
 
-	ViewModelsArray& GetWidgetViewModels() { return WidgetViewModels; }
 	ViewModelsArray& GetWidgetTemplateListViewModels() { return WidgetTemplateListViewModels; }
 
 protected:
diff --git a/Engine/Source/Editor/UMGEditor/Private/Palette/SPaletteViewModel.cpp b/Engine/Source/Editor/UMGEditor/Private/Palette/SPaletteViewModel.cpp
--- a/Engine/Source/Editor/UMGEditor/Private/Palette/SPaletteViewModel.cpp
+++ b/Engine/Source/Editor/UMGEditor/Private/Palette/SPaletteViewModel.cpp
@@ -173,20 +173,23 @@ void FWidgetCatalogViewModel::BuildWidgetList()
 	bool bHasFavorites = FavoriteHeader->Children.Num() != 0;
 	FavoriteHeader->Children.Reset();
 	
-	// Copy of the list of favorites to be able to do some cleanup in the real list
-	UWidgetPaletteFavorites* FavoritesPalette = GetDefault<UWidgetDesignerSettings>()->Favorites;
-	TArray<FString> FavoritesList = FavoritesPalette->GetFavorites();
-
-	// For each entry in the category create a view model for the widget template
-	for ( auto& Entry : WidgetTemplateCategories )
+	// Build ViewModel and clean the Favorite list if needed
 	{
-		BuildWidgetTemplateCategory(Entry.Key, Entry.Value, FavoritesList);
-	}
+		// Copy of the list of favorites to be able to do some cleanup in the real list
+		UWidgetPaletteFavorites* FavoritesPalette = GetDefault<UWidgetDesignerSettings>()->Favorites;
+		TArray<FString> FavoritesList = FavoritesPalette->GetFavorites();
 
-	// Remove all Favorites that may be left in the list.Typically happening when the list of favorite contains widget that were deleted since the last opening.
-	for (const FString& favoriteName : FavoritesList)
-	{
-		FavoritesPalette->Remove(favoriteName);
+		// For each entry in the category create a view model for the widget template
+		for (auto& Entry : WidgetTemplateCategories)
+		{
+			BuildWidgetTemplateCategory(Entry.Key, Entry.Value, FavoritesList);
+		}
+
+		// Remove all Favorites that may be left in the list.Typically happening when the list of favorite contains widget that were deleted since the last opening.
+		for (const FString& FavoriteName : FavoritesList)
+		{
+			FavoritesPalette->Remove(FavoriteName);
+		}
 	}
 
 	// Sort the view models by name
@@ -333,6 +336,19 @@ void FWidgetCatalogViewModel::BuildClassWidgetList()
 	}
 }
 
+void FWidgetCatalogViewModel::AddHeader(TSharedPtr<FWidgetHeaderViewModel>& Header)
+{
+	WidgetViewModels.Add(Header);
+}
+
+void FWidgetCatalogViewModel::AddToFavoriteHeader(TSharedPtr<FWidgetTemplateViewModel>& Favorite)
+{
+	if (FavoriteHeader)
+	{
+		FavoriteHeader->Children.Add(Favorite);
+	}
+}
+
 void FWidgetCatalogViewModel::AddWidgetTemplate(TSharedPtr<FWidgetTemplate> Template)
 {
 	FString Category = *Template->GetCategory().ToString();
@@ -385,10 +401,6 @@ void FPaletteViewModel::BuildWidgetTemplateCategory(FString& Category, TArray<TS
 {
 	TSharedPtr<FWidgetHeaderViewModel> Header = MakeShareable(new FWidgetHeaderViewModel());
 	Header->GroupName = FText::FromString(Category);
-
-	// Copy of the list of favorites to be able to do some cleanup in the real list
-	UWidgetPaletteFavorites* FavoritesPalette = GetDefault<UWidgetDesignerSettings>()->Favorites;
-
 	for (auto& Template : Templates)
 	{
 		TSharedPtr<FWidgetTemplateViewModel> TemplateViewModel = MakeShareable(new FWidgetTemplateViewModel());
@@ -408,7 +420,7 @@ void FPaletteViewModel::BuildWidgetTemplateCategory(FString& Category, TArray<TS
 			FavoriteTemplateViewModel->FavortiesViewModel = this;
 			FavoriteTemplateViewModel->SetFavorite();
 
-			FavoriteHeader->Children.Add(FavoriteTemplateViewModel);
+			AddToFavoriteHeader(FavoriteTemplateViewModel);
 
 			// Remove the favorite from the temporary list
 			FavoritesList.RemoveAt(index);
@@ -418,16 +430,16 @@ void FPaletteViewModel::BuildWidgetTemplateCategory(FString& Category, TArray<TS
 	
 	Header->Children.Sort([](const TSharedPtr<FWidgetViewModel>& L, const TSharedPtr<FWidgetViewModel>& R) { return R->GetName().CompareTo(L->GetName()) > 0; });
 
-	WidgetViewModels.Add(Header);
+	AddHeader(Header);
 }
 
-void FPaletteViewModel::AddToFavorites(const FWidgetTemplateViewModel* WidgetTemplateViewModel)
+void FWidgetCatalogViewModel::AddToFavorites(const FWidgetTemplateViewModel* WidgetTemplateViewModel)
 {
 	UWidgetPaletteFavorites* Favorites = GetDefault<UWidgetDesignerSettings>()->Favorites;
 	Favorites->Add(WidgetTemplateViewModel->GetName().ToString());
 }
 
-void FPaletteViewModel::RemoveFromFavorites(const FWidgetTemplateViewModel* WidgetTemplateViewModel)
+void FWidgetCatalogViewModel::RemoveFromFavorites(const FWidgetTemplateViewModel* WidgetTemplateViewModel)
 {
 	UWidgetPaletteFavorites* Favorites = GetDefault<UWidgetDesignerSettings>()->Favorites;
 	Favorites->Remove(WidgetTemplateViewModel->GetName().ToString());
diff --git a/Engine/Source/Editor/UMGEditor/Private/Palette/SPaletteViewModel.h b/Engine/Source/Editor/UMGEditor/Private/Palette/SPaletteViewModel.h
--- a/Engine/Source/Editor/UMGEditor/Private/Palette/SPaletteViewModel.h
+++ b/Engine/Source/Editor/UMGEditor/Private/Palette/SPaletteViewModel.h
@@ -166,10 +166,10 @@ public:
 	bool NeedUpdate() const { return bRebuildRequested; }
 
 	/** Add the widget template to the list of favorites */
-	virtual void AddToFavorites(const FWidgetTemplateViewModel* WidgetTemplateViewModel) = 0;
+	static void AddToFavorites(const FWidgetTemplateViewModel* WidgetTemplateViewModel);
 
 	/** Remove the widget template to the list of favorites */
-	virtual void RemoveFromFavorites(const FWidgetTemplateViewModel* WidgetTemplateViewModel) = 0;
+	static void RemoveFromFavorites(const FWidgetTemplateViewModel* WidgetTemplateViewModel);
 
 	typedef TArray< TSharedPtr<FWidgetViewModel> > ViewModelsArray;
 	ViewModelsArray& GetWidgetViewModels() { return WidgetViewModels; }
@@ -184,11 +184,15 @@ public:
 	FText GetSearchText() const { return SearchText; }
 
 protected:
+	virtual void BuildWidgetList();
+
+	void AddHeader(TSharedPtr<FWidgetHeaderViewModel>& Header);
+	void AddToFavoriteHeader(TSharedPtr<FWidgetTemplateViewModel>& Favorite);
+
+private:
 	FWidgetCatalogViewModel() {};
 
 	UWidgetBlueprint* GetBlueprint() const;
-
-	virtual void BuildWidgetList();
 	virtual void BuildWidgetTemplateCategory(FString& Category, TArray<TSharedPtr<FWidgetTemplate>>& Templates, TArray<FString>& FavoritesList) = 0;
 	void BuildClassWidgetList();
 
@@ -232,8 +236,6 @@ public:
 
 	//~ Begin FWidgetCatalogViewModel Interface
 	virtual void BuildWidgetTemplateCategory(FString& Category, TArray<TSharedPtr<FWidgetTemplate>>& Templates, TArray<FString>& FavoritesList) override;
-	virtual void AddToFavorites(const FWidgetTemplateViewModel* WidgetTemplateViewModel) override;
-	virtual void RemoveFromFavorites(const FWidgetTemplateViewModel* WidgetTemplateViewModel) override;
 	//~ End FWidgetCatalogViewModel Interface
 };
 
