Sequencer: Widget Blueprint Editor should call Close() on Sequencer. This fixes an issue where the time display format changed delegate wasn't being removed.

This issue surfaced recently because of the UE Core Thread-safe delegate race detection changes.
This mimics what happens in ~FLevelSequenceEditorToolkit.

diff --git a/Engine/Source/Editor/UMGEditor/Private/WidgetBlueprintEditor.cpp b/Engine/Source/Editor/UMGEditor/Private/WidgetBlueprintEditor.cpp
--- a/Engine/Source/Editor/UMGEditor/Private/WidgetBlueprintEditor.cpp
+++ b/Engine/Source/Editor/UMGEditor/Private/WidgetBlueprintEditor.cpp
@@ -126,6 +126,7 @@ FWidgetBlueprintEditor::~FWidgetBlueprintEditor()
 		{
 			Sequencer->OnMovieSceneDataChanged().RemoveAll(this);
 			Sequencer->OnMovieSceneBindingsPasted().RemoveAll(this);
+			Sequencer->Close();
 			Sequencer.Reset();
 		}
 	}
