Allow property drag-drop support for preview widget.

diff --git a/Engine/Plugins/Runtime/ModelViewViewModel/Source/ModelViewViewModelEditor/Private/Widgets/SMVVMViewBindingListView.cpp b/Engine/Plugins/Runtime/ModelViewViewModel/Source/ModelViewViewModelEditor/Private/Widgets/SMVVMViewBindingListView.cpp
--- a/Engine/Plugins/Runtime/ModelViewViewModel/Source/ModelViewViewModelEditor/Private/Widgets/SMVVMViewBindingListView.cpp
+++ b/Engine/Plugins/Runtime/ModelViewViewModel/Source/ModelViewViewModelEditor/Private/Widgets/SMVVMViewBindingListView.cpp
@@ -3,6 +3,7 @@
 #include "Widgets/SMVVMViewBindingListView.h"
 
 #include "Bindings/MVVMBindingHelper.h"
+#include "Blueprint/WidgetTree.h"
 #include "BlueprintEditor.h"
 #include "Details/WidgetPropertyDragDropOp.h"
 #include "Dialog/SCustomDialog.h"
@@ -774,7 +775,16 @@ private:
 					{
 						if (UWidget* OwnerWidgetPtr = WidgetPropertyDragDropOp->OwnerWidget.Get())
 						{
-							PropertyPath.SetWidgetName(OwnerWidgetPtr->GetFName());
+							// Search for the widget by its name in the widget tree
+							// If the widget is not found, we know it is the root preview widget so we use the blueprint name.
+							if (WidgetBlueprintPtr->WidgetTree->FindWidget(OwnerWidgetPtr->GetFName()))
+							{
+								PropertyPath.SetWidgetName(OwnerWidgetPtr->GetFName());
+							}
+							else
+							{
+								PropertyPath.SetWidgetName(WidgetBlueprintPtr->GetFName());
+							}
 						}
 					}
 
diff --git a/Engine/Source/Editor/UMGEditor/Private/Details/DetailWidgetExtensionHandler.cpp b/Engine/Source/Editor/UMGEditor/Private/Details/DetailWidgetExtensionHandler.cpp
--- a/Engine/Source/Editor/UMGEditor/Private/Details/DetailWidgetExtensionHandler.cpp
+++ b/Engine/Source/Editor/UMGEditor/Private/Details/DetailWidgetExtensionHandler.cpp
@@ -88,7 +88,7 @@ bool FDetailWidgetExtensionHandler::IsPropertyExtendable(const UClass* ObjectCla
 	}
 
 	TSharedPtr<FWidgetBlueprintEditor> BPEd = BlueprintEditor.Pin();
-	if (BPEd == nullptr || Objects[0] == BPEd->GetPreview())
+	if (BPEd == nullptr)
 	{
 		return false;
 	}
@@ -168,22 +168,25 @@ void FDetailWidgetExtensionHandler::ExtendWidgetRow(
 		return;
 	}
 
-	UFunction* SignatureFunction = nullptr;
-	if (FDelegateProperty* DelegateProperty = Private::FindDelegateProperty(PropertyHandleRef.Get()))
-	{
-		SignatureFunction = DelegateProperty->SignatureFunction;
-	}
-
-	InWidgetRow.ExtensionContent()
-	[
-		FBlueprintWidgetCustomization::MakePropertyBindingWidget(BlueprintEditor, SignatureFunction, InPropertyHandle.ToSharedRef(), true, bShouldShowOldBindingWidget)
-	];
-
 	if (Objects.Num() > 0)
 	{
 		if (UWidget* Widget = Cast<UWidget>(Objects[0]))
 		{
 			InWidgetRow.DragDropHandler(MakeShared<FWidgetPropertyDragDropHandler>(Widget, InPropertyHandle, WidgetBlueprint));
+
+			if (Widget != BlueprintEditor.Pin()->GetPreview())
+			{
+				UFunction* SignatureFunction = nullptr;
+				if (FDelegateProperty* DelegateProperty = Private::FindDelegateProperty(PropertyHandleRef.Get()))
+				{
+					SignatureFunction = DelegateProperty->SignatureFunction;
+				}
+
+				InWidgetRow.ExtensionContent()
+					[
+						FBlueprintWidgetCustomization::MakePropertyBindingWidget(BlueprintEditor, SignatureFunction, InPropertyHandle.ToSharedRef(), true, bShouldShowOldBindingWidget)
+					];
+			}
 		}
 	}
 }
