UMG: add support for dynamic bindings in widget animations

diff --git a/Engine/Source/Editor/Sequencer/Public/MVVM/ObjectBindingModelStorageExtension.h b/Engine/Source/Editor/Sequencer/Public/MVVM/ObjectBindingModelStorageExtension.h
--- a/Engine/Source/Editor/Sequencer/Public/MVVM/ObjectBindingModelStorageExtension.h
+++ b/Engine/Source/Editor/Sequencer/Public/MVVM/ObjectBindingModelStorageExtension.h
@@ -25,7 +25,7 @@ class FSequenceModel;
 class FViewModel;
 struct FViewModelChildren;
 
-class FObjectBindingModelStorageExtension
+class SEQUENCER_API FObjectBindingModelStorageExtension
 	: public IDynamicExtension
 	, private UE::MovieScene::TIntrusiveEventHandler<UE::MovieScene::ISequenceDataEventHandler>
 {
diff --git a/Engine/Source/Editor/UMGEditor/Private/WidgetBlueprintEditor.cpp b/Engine/Source/Editor/UMGEditor/Private/WidgetBlueprintEditor.cpp
--- a/Engine/Source/Editor/UMGEditor/Private/WidgetBlueprintEditor.cpp
+++ b/Engine/Source/Editor/UMGEditor/Private/WidgetBlueprintEditor.cpp
@@ -53,6 +53,9 @@
 
 #include "Animation/MovieSceneWidgetMaterialTrack.h"
 #include "Animation/WidgetMaterialTrackUtilities.h"
+#include "MVVM/ObjectBindingModelStorageExtension.h"
+#include "MVVM/ViewModels/ObjectBindingModel.h"
+#include "MVVM/ViewModels/SequencerEditorViewModel.h"
 
 #include "ScopedTransaction.h"
 
@@ -1971,7 +1974,7 @@ void FWidgetBlueprintEditor::OnBuildCustomContextMenuForGuid(FMenuBuilder& MenuB
 			}
 		}
 		
-		if(ValidSelectedWidgets.Num() > 0)
+		if (ValidSelectedWidgets.Num() > 0)
 		{
 			MenuBuilder.AddMenuSeparator();
 			
@@ -2021,6 +2024,11 @@ void FWidgetBlueprintEditor::OnBuildCustomContextMenuForGuid(FMenuBuilder& MenuB
 				FSlateIcon(),
 				FExecuteAction::CreateRaw(this, &FWidgetBlueprintEditor::RemoveMissingWidgetsFromTrack, ObjectBinding)
 			);
+
+			MenuBuilder.AddSubMenu(
+				LOCTEXT("DynamicPossession", "Dynamic Possession"),
+				LOCTEXT("DynamicPossessionToolTip", "Specify a Blueprint method that will find a compatible widget for this binding"),
+				FNewMenuDelegate::CreateRaw(this, &FWidgetBlueprintEditor::AddDynamicPossessionMenu, ObjectBinding));
 		}
 	}
 }
@@ -2297,6 +2305,35 @@ void FWidgetBlueprintEditor::ReplaceTrackWithWidgets(TArray<FWidgetReference> Wi
 	SyncSequencersMovieSceneData();
 }
 
+void FWidgetBlueprintEditor::AddDynamicPossessionMenu(FMenuBuilder& MenuBuilder, FGuid ObjectId)
+{
+	using namespace UE::Sequencer;
+
+	TSharedPtr<ISequencer>& ActiveSequencer = GetSequencer();
+
+	UMovieScene* MovieScene = ActiveSequencer->GetFocusedMovieSceneSequence()->GetMovieScene();
+	FMovieScenePossessable* Possessable = MovieScene->FindPossessable(ObjectId);
+	if (!Possessable)
+	{
+		return;
+	}
+
+	TSharedPtr<FSequencerEditorViewModel> SequencerViewModel = ActiveSequencer->GetViewModel();
+	FObjectBindingModelStorageExtension* ObjectStorage = SequencerViewModel->GetRootModel()->CastDynamic<FObjectBindingModelStorageExtension>();
+	if (!ObjectStorage)
+	{
+		return;
+	}
+
+	TSharedPtr<FObjectBindingModel> ObjectBindingModel = ObjectStorage->FindModelForObjectBinding(ObjectId);
+	if (!ObjectBindingModel)
+	{
+		return;
+	}
+
+	ObjectBindingModel->AddDynamicBindingMenu(MenuBuilder, Possessable->DynamicBinding);
+}
+
 void FWidgetBlueprintEditor::AddSlotTrack( UPanelSlot* Slot )
 {
 	TSharedPtr<ISequencer>& ActiveSequencer = GetSequencer();
diff --git a/Engine/Source/Editor/UMGEditor/Public/WidgetBlueprintEditor.h b/Engine/Source/Editor/UMGEditor/Public/WidgetBlueprintEditor.h
--- a/Engine/Source/Editor/UMGEditor/Public/WidgetBlueprintEditor.h
+++ b/Engine/Source/Editor/UMGEditor/Public/WidgetBlueprintEditor.h
@@ -395,6 +395,9 @@ private:
 	/** Replace current widget bindings on a track with new widget bindings */
 	void ReplaceTrackWithWidgets(const TArray<FWidgetReference> Widgets, FGuid ObjectId);
 
+	/** Dynamic binding */
+	void AddDynamicPossessionMenu(FMenuBuilder& MenuBuilder, FGuid ObjectId);
+
 	/** Add an animation track for the supplied slot to the current animation. */
 	void AddSlotTrack( UPanelSlot* Slot );
 
diff --git a/Engine/Source/Editor/UMGEditor/UMGEditor.Build.cs b/Engine/Source/Editor/UMGEditor/UMGEditor.Build.cs
--- a/Engine/Source/Editor/UMGEditor/UMGEditor.Build.cs
+++ b/Engine/Source/Editor/UMGEditor/UMGEditor.Build.cs
@@ -12,6 +12,7 @@ public class UMGEditor : ModuleRules
 			new string[]
 			{
 				"FieldNotification",
+				"SequencerCore",
 				"Sequencer",
 			});
 
