UMG: Expose SaveSubObjectsFromCleanAndSanitizeClas to WidgetBlueprintExtension.

diff --git a/Engine/Source/Editor/UMGEditor/Private/WidgetBlueprintCompiler.cpp b/Engine/Source/Editor/UMGEditor/Private/WidgetBlueprintCompiler.cpp
--- a/Engine/Source/Editor/UMGEditor/Private/WidgetBlueprintCompiler.cpp
+++ b/Engine/Source/Editor/UMGEditor/Private/WidgetBlueprintCompiler.cpp
@@ -447,6 +447,11 @@ void FWidgetBlueprintCompilerContext::SaveSubObjectsFromCleanAndSanitizeClass(FS
 			SubObjectsToSave.AddObject(CDONamedSlotBinding.Content);
 		}
 	}
+
+	UWidgetBlueprintExtension::ForEachExtension(WidgetBlueprint(), [&SubObjectsToSave, LocalClass = NewWidgetBlueprintClass](UWidgetBlueprintExtension* InExtension)
+		{
+			SubObjectsToSave.AddObjects(InExtension->SaveSubObjectsFromCleanAndSanitizeClass(LocalClass));
+		});
 }
 
 void FWidgetBlueprintCompilerContext::CreateClassVariablesFromBlueprint()
diff --git a/Engine/Source/Editor/UMGEditor/Public/WidgetBlueprintExtension.h b/Engine/Source/Editor/UMGEditor/Public/WidgetBlueprintExtension.h
--- a/Engine/Source/Editor/UMGEditor/Public/WidgetBlueprintExtension.h
+++ b/Engine/Source/Editor/UMGEditor/Public/WidgetBlueprintExtension.h
@@ -74,6 +74,7 @@ protected:
 
 	virtual void HandleCreateFunctionList() {}
 	virtual void HandleCleanAndSanitizeClass(UWidgetBlueprintGeneratedClass* ClassToClean, UObject* OldCDO) {}
+	virtual TArray<UObject*> HandleSaveSubObjectsFromCleanAndSanitizeClass(UWidgetBlueprintGeneratedClass* ClassToClean) { return TArray<UObject*>();  }
 	virtual void HandleCreateClassVariablesFromBlueprint(const FWidgetBlueprintCompilerContext::FCreateVariableContext& Context) {}
 	virtual void HandleCopyTermDefaultsToDefaultObject(UObject* DefaultObject) {}
 	virtual void HandleFinishCompilingClass(UWidgetBlueprintGeneratedClass* Class) {}
@@ -104,6 +105,11 @@ private:
 		HandleCleanAndSanitizeClass(ClassToClean, OldCDO);
 	}
 
+	TArray<UObject*> SaveSubObjectsFromCleanAndSanitizeClass(UWidgetBlueprintGeneratedClass* ClassToClean)
+	{
+		return HandleSaveSubObjectsFromCleanAndSanitizeClass(ClassToClean);
+	}
+
 	void CreateClassVariablesFromBlueprint(const FWidgetBlueprintCompilerContext::FCreateVariableContext& Context)
 	{
 		HandleCreateClassVariablesFromBlueprint(Context);
