UMG: Rename the Debug mode to Preview Mode. Add a Detail View, remove the log widget. Add a container to share the information and editor context, instead of use the WidgetEditor directly.

diff --git a/Engine/Source/Editor/UMGEditor/Private/BlueprintModes/WidgetBlueprintApplicationModes.cpp b/Engine/Source/Editor/UMGEditor/Private/BlueprintModes/WidgetBlueprintApplicationModes.cpp
--- a/Engine/Source/Editor/UMGEditor/Private/BlueprintModes/WidgetBlueprintApplicationModes.cpp
+++ b/Engine/Source/Editor/UMGEditor/Private/BlueprintModes/WidgetBlueprintApplicationModes.cpp
@@ -8,7 +8,8 @@
 // Mode constants
 const FName FWidgetBlueprintApplicationModes::DesignerMode("DesignerName");
 const FName FWidgetBlueprintApplicationModes::GraphMode("GraphName");
-const FName FWidgetBlueprintApplicationModes::DebugMode("DebugName");
+const FName FWidgetBlueprintApplicationModes::DebugMode("PreviewName");
+const FName FWidgetBlueprintApplicationModes::PreviewMode("PreviewName");
 
 FText FWidgetBlueprintApplicationModes::GetLocalizedMode(const FName InMode)
 {
@@ -20,21 +21,21 @@ FText FWidgetBlueprintApplicationModes::GetLocalizedMode(const FName InMode)
 	{
 		return NSLOCTEXT("WidgetBlueprintModes", "GraphMode", "Graph");
 	}
-	else if (InMode == FWidgetBlueprintApplicationModes::DebugMode)
+	else if (InMode == FWidgetBlueprintApplicationModes::PreviewMode)
 	{
-		return NSLOCTEXT("WidgetBlueprintModes", "DebugMode", "Debug");
+		return NSLOCTEXT("WidgetBlueprintModes", "PreviewMode", "Preview");
 	}
 	return FText::GetEmpty();
 }
 
-static bool EnableDebugMode = false;
-static FAutoConsoleVariableRef CVarEnableDebugMode(
-	TEXT("UMG.EnableDebugMode"), 
-	EnableDebugMode, 
-	TEXT("Whether or not to enable the UMG Debug mode.")
+static bool bEnablePreviewMode = false;
+static FAutoConsoleVariableRef CVarEnablePreviewMode(
+	TEXT("UMG.EnablePreviewMode"), 
+	bEnablePreviewMode, 
+	TEXT("Whether or not to enable the UMG Preview mode.")
 );
 
-bool FWidgetBlueprintApplicationModes::IsDebugModeEnabled()
+bool FWidgetBlueprintApplicationModes::IsPreviewModeEnabled()
 {
-	return EnableDebugMode;
+	return bEnablePreviewMode;
 }
diff --git a/Engine/Source/Editor/UMGEditor/Private/BlueprintModes/WidgetDebugApplicationMode.cpp b/Engine/Source/Editor/UMGEditor/Private/BlueprintModes/WidgetDebugApplicationMode.cpp
deleted file mode 100644
diff --git a/Engine/Source/Editor/UMGEditor/Private/BlueprintModes/WidgetDebugApplicationMode.h b/Engine/Source/Editor/UMGEditor/Private/BlueprintModes/WidgetDebugApplicationMode.h
deleted file mode 100644
diff --git a/Engine/Source/Editor/UMGEditor/Private/BlueprintModes/WidgetPreviewApplicationMode.cpp b/Engine/Source/Editor/UMGEditor/Private/BlueprintModes/WidgetPreviewApplicationMode.cpp
new file mode 100644
--- /dev/null
+++ b/Engine/Source/Editor/UMGEditor/Private/BlueprintModes/WidgetPreviewApplicationMode.cpp
@@ -0,0 +1,74 @@
+// Copyright Epic Games, Inc. All Rights Reserved.
+
+#include "BlueprintModes/WidgetPreviewApplicationMode.h"
+
+#include "BlueprintModes/WidgetBlueprintApplicationModes.h"
+#include "SBlueprintEditorToolbar.h"
+#include "TabFactory/PreviewDetailsTabSummoner.h"
+#include "TabFactory/DetailsTabSummoner.h"
+#include "TabFactory/PreviewTabSummoner.h"
+#include "ToolMenus.h"
+#include "UMGEditorModule.h"
+#include "WidgetBlueprintEditorToolbar.h"
+
+namespace UE::UMG::Editor
+{
+
+FWidgetPreviewApplicationMode::FWidgetPreviewApplicationMode(TSharedPtr<FWidgetBlueprintEditor> InWidgetEditor)
+	: FWidgetBlueprintApplicationMode(InWidgetEditor, FWidgetBlueprintApplicationModes::PreviewMode)
+{
+	TabLayout = FTabManager::NewLayout( "WidgetBlueprintEditor_Preview_Layout_v1" )
+		->AddArea
+		(
+			FTabManager::NewPrimaryArea()
+			->SetOrientation(Orient_Horizontal)
+			->Split
+			(
+				FTabManager::NewStack()
+				->SetSizeCoefficient(0.8f)
+				->AddTab(FWidgetPreviewTabSummoner::TabID, ETabState::OpenedTab)
+				->SetHideTabWell(true)
+			)
+			->Split
+			(
+				FTabManager::NewSplitter()
+				->SetOrientation(Orient_Vertical)
+				->SetSizeCoefficient(0.15f)
+				->Split
+				(
+					FTabManager::NewStack()
+					->AddTab(FPreviewDetailsTabSummoner::TabID, ETabState::OpenedTab)
+				)
+			)
+		);
+	
+	TabFactories.RegisterFactory(MakeShareable(new FWidgetPreviewTabSummoner(InWidgetEditor)));
+	TabFactories.RegisterFactory(MakeShareable(new FPreviewDetailsTabSummoner(InWidgetEditor)));
+	//TabFactories.RegisterFactory(MakeShareable(new FPreviewSettingsTabSummoner(InWidgetEditor))); Named: Parameters? contains: localization, background, size
+
+	IUMGEditorModule& UMGEditorModule = FModuleManager::LoadModuleChecked<IUMGEditorModule>("UMGEditor");
+	UMGEditorModule.OnRegisterTabsForEditor().Broadcast(*this, TabFactories);
+
+	ToolbarExtender = UMGEditorModule.GetToolBarExtensibilityManager()->GetAllExtenders();
+	InWidgetEditor->GetWidgetToolbarBuilder()->AddWidgetBlueprintEditorModesToolbar(ToolbarExtender);
+
+	if (UToolMenu* Toolbar = InWidgetEditor->RegisterModeToolbarIfUnregistered(GetModeName()))
+	{
+		InWidgetEditor->GetWidgetToolbarBuilder()->AddWidgetReflector(Toolbar);
+	}
+}
+
+void FWidgetPreviewApplicationMode::RegisterTabFactories(TSharedPtr<FTabManager> InTabManager)
+{
+	TSharedPtr<FBlueprintEditor> BP = GetBlueprintEditor();
+	
+	BP->RegisterToolbarTab(InTabManager.ToSharedRef());
+	BP->PushTabFactories(TabFactories);
+}
+
+void FWidgetPreviewApplicationMode::PostActivateMode()
+{
+
+}
+
+} //namespace
diff --git a/Engine/Source/Editor/UMGEditor/Private/BlueprintModes/WidgetPreviewApplicationMode.h b/Engine/Source/Editor/UMGEditor/Private/BlueprintModes/WidgetPreviewApplicationMode.h
new file mode 100644
--- /dev/null
+++ b/Engine/Source/Editor/UMGEditor/Private/BlueprintModes/WidgetPreviewApplicationMode.h
@@ -0,0 +1,24 @@
+// Copyright Epic Games, Inc. All Rights Reserved.
+
+#pragma once
+
+#include "CoreMinimal.h"
+#include "BlueprintModes/WidgetBlueprintApplicationMode.h"
+
+namespace UE::UMG::Editor
+{
+
+class FWidgetPreviewApplicationMode : public FWidgetBlueprintApplicationMode
+{
+public:
+	FWidgetPreviewApplicationMode(TSharedPtr<FWidgetBlueprintEditor> InWidgetEditor);
+
+	static bool IsDebugModeEnabled();
+
+	// FApplicationMode interface
+	virtual void RegisterTabFactories(TSharedPtr<FTabManager> InTabManager) override;
+	virtual void PostActivateMode() override;
+	// End of FApplicationMode interface
+};
+
+} // namespace
\ No newline at end of file
diff --git a/Engine/Source/Editor/UMGEditor/Private/Debug/SDebugLog.cpp b/Engine/Source/Editor/UMGEditor/Private/Debug/SDebugLog.cpp
deleted file mode 100644
diff --git a/Engine/Source/Editor/UMGEditor/Private/Debug/SDebugLog.h b/Engine/Source/Editor/UMGEditor/Private/Debug/SDebugLog.h
deleted file mode 100644
diff --git a/Engine/Source/Editor/UMGEditor/Private/Debug/SDebugPreview.cpp b/Engine/Source/Editor/UMGEditor/Private/Debug/SDebugPreview.cpp
deleted file mode 100644
diff --git a/Engine/Source/Editor/UMGEditor/Private/Debug/SDebugPreview.h b/Engine/Source/Editor/UMGEditor/Private/Debug/SDebugPreview.h
deleted file mode 100644
diff --git a/Engine/Source/Editor/UMGEditor/Private/Preview/PreviewMode.cpp b/Engine/Source/Editor/UMGEditor/Private/Preview/PreviewMode.cpp
new file mode 100644
--- /dev/null
+++ b/Engine/Source/Editor/UMGEditor/Private/Preview/PreviewMode.cpp
@@ -0,0 +1,54 @@
+// Copyright Epic Games, Inc. All Rights Reserved.
+
+#include "Preview/PreviewMode.h"
+#include "Blueprint/UserWidget.h"
+
+namespace UE::UMG::Editor
+{
+
+void FPreviewMode::SetSelectedObject(TArray<TWeakObjectPtr<UObject>> Objects)
+{
+	SelectedObjects = MoveTemp(Objects);
+	SelectedObjectChangedDelegate.Broadcast();
+}
+
+void FPreviewMode::SetSelectedObject(TArrayView<UObject*> Objects)
+{
+	SelectedObjects.Reset(Objects.Num());
+	for (UObject* Obj : Objects)
+	{
+		SelectedObjects.Emplace(Obj);
+	}
+	SelectedObjectChangedDelegate.Broadcast();
+}
+
+TArray<UObject*> FPreviewMode::GetSelectedObjectList() const
+{
+	TArray<UObject*> Result;
+	Result.Reset(SelectedObjects.Num());
+	for (TWeakObjectPtr<UObject> Obj : SelectedObjects)
+	{
+		if (UObject* Ptr = Obj.Get())
+		{
+			Result.Add(Ptr);
+		}
+	}
+	return Result;
+}
+
+void FPreviewMode::SetPreviewWidget(UUserWidget* Widget)
+{
+	UUserWidget* Previous = PreviewWidget.Get();
+	if (Previous != Widget)
+	{
+		PreviewWidget = Widget;
+		PreviewWidgetChangedDelegate.Broadcast();
+		if (SelectedObjects.ContainsByPredicate([Previous](const TWeakObjectPtr<UObject>& Other){ return Other.Get() == Previous; }))
+		{
+			SelectedObjects.Reset();
+			SelectedObjectChangedDelegate.Broadcast();
+		}
+	}
+}
+
+} // namespace UE::UMG
diff --git a/Engine/Source/Editor/UMGEditor/Private/Preview/SPreviewDetails.cpp b/Engine/Source/Editor/UMGEditor/Private/Preview/SPreviewDetails.cpp
new file mode 100644
--- /dev/null
+++ b/Engine/Source/Editor/UMGEditor/Private/Preview/SPreviewDetails.cpp
@@ -0,0 +1,99 @@
+// Copyright Epic Games, Inc. All Rights Reserved.
+
+#include "SPreviewDetails.h"
+
+#include "Preview/PreviewMode.h"
+#include "INotifyFieldValueChanged.h"
+#include "PropertyEditorModule.h"
+#include "Modules/ModuleManager.h"
+#include "WidgetBlueprintEditor.h"
+
+#include "IDetailsView.h"
+
+namespace UE::UMG::Editor
+{
+
+void SPreviewDetails::Construct(const FArguments& Args, TSharedPtr<FWidgetBlueprintEditor> InBlueprintEditor)
+{
+	WeakEditor = InBlueprintEditor;
+
+	FPropertyEditorModule& PropertyEditor = FModuleManager::LoadModuleChecked<FPropertyEditorModule>("PropertyEditor");
+
+	FDetailsViewArgs DetailsViewArgs;
+	DetailsViewArgs.bShowOptions = false;
+	DetailsViewArgs.bAllowSearch = true;
+	DetailsViewArgs.bShowPropertyMatrixButton = false;
+	DetailsViewArgs.bUpdatesFromSelection = false;
+	DetailsViewArgs.bLockable = false;
+	DetailsViewArgs.NameAreaSettings = FDetailsViewArgs::HideNameArea;
+	DetailsViewArgs.NotifyHook = this;
+
+	DetailsView = PropertyEditor.CreateDetailView(DetailsViewArgs);
+
+	TSharedPtr<FPreviewMode> Context = InBlueprintEditor && InBlueprintEditor->GetPreviewMode() ? InBlueprintEditor->GetPreviewMode() : TSharedPtr<FPreviewMode>();
+	if (Context)
+	{
+		Context->OnSelectedObjectChanged().AddSP(this, &SPreviewDetails::HandleSelectedObjectChanged);
+	}
+	HandleSelectedObjectChanged();
+
+	ChildSlot
+	[
+		DetailsView.ToSharedRef()
+	];
+}
+
+
+void SPreviewDetails::NotifyPostChange(const FPropertyChangedEvent& PropertyChangedEvent, class FEditPropertyChain* PropertyThatChanged)
+{
+
+	for (int32 Index = 0; Index < PropertyChangedEvent.GetNumObjectsBeingEdited(); ++Index)
+	{
+		if (const UObject* Object = PropertyChangedEvent.GetObjectBeingEdited(Index))
+		{
+			if (PropertyThatChanged->GetActiveMemberNode())
+			{
+				FProperty* Property = PropertyThatChanged->GetActiveMemberNode()->GetValue();
+				if (Property && Object->GetClass()->ImplementsInterface(UNotifyFieldValueChanged::StaticClass()))
+				{
+					TScriptInterface<INotifyFieldValueChanged> Interface = const_cast<UObject*>(Object);
+					UE::FieldNotification::FFieldId FieldId = Interface->GetFieldNotificationDescriptor().GetField(Object->GetClass(), Property->GetFName());
+					if (FieldId.IsValid())
+					{
+						Interface->BroadcastFieldValueChanged(FieldId);
+					}
+				}
+			}
+		}
+	}
+}
+
+
+void SPreviewDetails::HandleSelectedObjectChanged()
+{
+	TSharedPtr<FWidgetBlueprintEditor> Editor = WeakEditor.Pin();
+	if (!Editor)
+	{
+		return;
+	}
+
+	TSharedPtr<FPreviewMode> Context = Editor && Editor->GetPreviewMode() ? Editor->GetPreviewMode() : TSharedPtr<FPreviewMode>();
+	if (Context)
+	{
+		TArray<UObject*> Objects = Context->GetSelectedObjectList();
+		if (Objects.Num() == 0)
+		{
+			if (Context->GetPreviewWidget())
+			{
+				Objects.Add(Context->GetPreviewWidget());
+			}
+		}
+		DetailsView->SetObjects(Objects);
+	}
+	else
+	{
+		DetailsView->SetObject(nullptr);
+	}
+}
+
+} // namespace
diff --git a/Engine/Source/Editor/UMGEditor/Private/Preview/SPreviewDetails.h b/Engine/Source/Editor/UMGEditor/Private/Preview/SPreviewDetails.h
new file mode 100644
--- /dev/null
+++ b/Engine/Source/Editor/UMGEditor/Private/Preview/SPreviewDetails.h
@@ -0,0 +1,31 @@
+// Copyright Epic Games, Inc. All Rights Reserved.
+
+#pragma once
+
+#include "Misc/NotifyHook.h"
+#include "Widgets/SCompoundWidget.h"
+
+class IDetailsView;
+class FWidgetBlueprintEditor;
+
+namespace UE::UMG::Editor
+{
+
+class SPreviewDetails : public SCompoundWidget, public FNotifyHook
+{
+	SLATE_BEGIN_ARGS(SPreviewDetails) {}
+	SLATE_END_ARGS()
+
+	void Construct(const FArguments& Args, TSharedPtr<FWidgetBlueprintEditor> InBlueprintEditor);
+
+	virtual void NotifyPostChange(const FPropertyChangedEvent& PropertyChangedEvent, class FEditPropertyChain* PropertyThatChanged) override;
+
+private:
+	void HandleSelectedObjectChanged();
+
+private:
+	TWeakPtr<FWidgetBlueprintEditor> WeakEditor;
+	TSharedPtr<IDetailsView> DetailsView;
+};
+
+} // namespace
diff --git a/Engine/Source/Editor/UMGEditor/Private/Preview/SWidgetPreview.cpp b/Engine/Source/Editor/UMGEditor/Private/Preview/SWidgetPreview.cpp
new file mode 100644
--- /dev/null
+++ b/Engine/Source/Editor/UMGEditor/Private/Preview/SWidgetPreview.cpp
@@ -0,0 +1,62 @@
+// Copyright Epic Games, Inc. All Rights Reserved.
+
+#include "Preview/SWidgetPreview.h"
+
+#include "Blueprint/UserWidget.h"
+#include "Engine/World.h"
+
+#include "Preview/PreviewMode.h"
+#include "Editor.h"
+#include "Engine/Engine.h"
+#include "WidgetBlueprint.h"
+#include "WidgetBlueprintEditor.h"
+
+namespace UE::UMG::Editor
+{
+
+void SWidgetPreview::Construct(const FArguments& Args, TSharedPtr<FWidgetBlueprintEditor> InWidgetBlueprintEditor)
+{
+	WeakEditor = InWidgetBlueprintEditor;
+
+	TSharedPtr<SWidget> SlateWidget;
+
+	UWidgetBlueprint* WidgetBlueprint = InWidgetBlueprintEditor->GetWidgetBlueprintObj();
+	UWorld* World = GEditor->GetEditorWorldContext().World();
+	if (WidgetBlueprint && WidgetBlueprint->GeneratedClass && World)
+	{
+		if (UWidgetBlueprintGeneratedClass* WidgetClass = Cast<UWidgetBlueprintGeneratedClass>(WidgetBlueprint->GeneratedClass))
+		{
+			CreatedWidget.Reset(CreateWidget(World, WidgetClass, WidgetClass->GetFName()));
+			if (CreatedWidget.Get())
+			{
+				CreatedWidget->SetDesignerFlags(EWidgetDesignFlags::Previewing);
+				SlateWidget = CreatedWidget->TakeWidget();
+			}
+		}
+	}
+
+	if (InWidgetBlueprintEditor->GetPreviewMode())
+	{
+		InWidgetBlueprintEditor->GetPreviewMode()->SetPreviewWidget(CreatedWidget.Get());
+	}
+	//Editor->GetDebugStuff()->SetPreviewWidget(CreateWidget);
+	// if widget compiled, (it could be the parent), need to clean this up like with widget blutility
+
+	ChildSlot
+	[
+		SlateWidget.IsValid() ? SlateWidget.ToSharedRef() : SNullWidget::NullWidget
+	];
+}
+
+SWidgetPreview::~SWidgetPreview()
+{
+	if (TSharedPtr<FWidgetBlueprintEditor> Editor = WeakEditor.Pin())
+	{
+		if (Editor->GetPreviewMode())
+		{
+			Editor->GetPreviewMode()->SetPreviewWidget(nullptr);
+		}
+	}
+}
+
+} // namespace UE::UMG
diff --git a/Engine/Source/Editor/UMGEditor/Private/Preview/SWidgetPreview.h b/Engine/Source/Editor/UMGEditor/Private/Preview/SWidgetPreview.h
new file mode 100644
--- /dev/null
+++ b/Engine/Source/Editor/UMGEditor/Private/Preview/SWidgetPreview.h
@@ -0,0 +1,31 @@
+// Copyright Epic Games, Inc. All Rights Reserved.
+
+#pragma once
+
+#include "UObject/StrongObjectPtr.h"
+#include "Widgets/SCompoundWidget.h"
+
+class FWidgetBlueprintEditor;
+class UPlayer;
+class UUserWidget;
+
+namespace UE::UMG::Editor
+{
+
+class SWidgetPreview : public SCompoundWidget
+{
+public:
+	SLATE_BEGIN_ARGS(SWidgetPreview) {}
+	SLATE_END_ARGS()
+
+	~SWidgetPreview();
+
+	void Construct(const FArguments& Args, TSharedPtr<FWidgetBlueprintEditor> InWidgetBlueprintEditor);
+
+private:
+
+	TWeakPtr<FWidgetBlueprintEditor> WeakEditor;
+	TStrongObjectPtr<UUserWidget> CreatedWidget;
+};
+
+} // namespace UE::UMG
\ No newline at end of file
diff --git a/Engine/Source/Editor/UMGEditor/Private/TabFactory/DebugLogTabSummoner.cpp b/Engine/Source/Editor/UMGEditor/Private/TabFactory/DebugLogTabSummoner.cpp
deleted file mode 100644
diff --git a/Engine/Source/Editor/UMGEditor/Private/TabFactory/DebugLogTabSummoner.h b/Engine/Source/Editor/UMGEditor/Private/TabFactory/DebugLogTabSummoner.h
deleted file mode 100644
diff --git a/Engine/Source/Editor/UMGEditor/Private/TabFactory/PreviewDetailsTabSummoner.cpp b/Engine/Source/Editor/UMGEditor/Private/TabFactory/PreviewDetailsTabSummoner.cpp
new file mode 100644
--- /dev/null
+++ b/Engine/Source/Editor/UMGEditor/Private/TabFactory/PreviewDetailsTabSummoner.cpp
@@ -0,0 +1,36 @@
+// Copyright Epic Games, Inc. All Rights Reserved.
+
+#include "TabFactory/PreviewDetailsTabSummoner.h"
+
+#include "UMGStyle.h"
+#include "WidgetBlueprintEditor.h"
+#include "Preview/SPreviewDetails.h"
+
+#define LOCTEXT_NAMESPACE "UMG"
+
+namespace UE::UMG::Editor
+{
+
+const FName FPreviewDetailsTabSummoner::TabID("DebugDetails");
+
+FPreviewDetailsTabSummoner::FPreviewDetailsTabSummoner(TSharedPtr<FWidgetBlueprintEditor> InBlueprintEditor)
+	: FWorkflowTabFactory(TabID, InBlueprintEditor)
+	, BlueprintEditor(InBlueprintEditor)
+{
+	TabLabel = LOCTEXT("DebugDetails_TabLabel", "Detail View");
+	TabIcon = FSlateIcon(FAppStyle::GetAppStyleSetName(), "BlueprintDebugger.TabIcon");
+
+	bIsSingleton = true;
+
+	ViewMenuDescription = LOCTEXT("DebugDetails_ViewMenu_Desc", "Detail View");
+	ViewMenuTooltip = LOCTEXT("DebugDetails_ViewMenu_ToolTip", "Show the Detail View");
+}
+
+TSharedRef<SWidget> FPreviewDetailsTabSummoner::CreateTabBody(const FWorkflowTabSpawnInfo& Info) const
+{
+	return SNew(SPreviewDetails, BlueprintEditor.Pin());
+}
+
+} //namespace
+
+#undef LOCTEXT_NAMESPACE 
diff --git a/Engine/Source/Editor/UMGEditor/Private/TabFactory/PreviewDetailsTabSummoner.h b/Engine/Source/Editor/UMGEditor/Private/TabFactory/PreviewDetailsTabSummoner.h
new file mode 100644
--- /dev/null
+++ b/Engine/Source/Editor/UMGEditor/Private/TabFactory/PreviewDetailsTabSummoner.h
@@ -0,0 +1,27 @@
+// Copyright Epic Games, Inc. All Rights Reserved.
+
+#pragma once
+
+#include "WorkflowOrientedApp/WorkflowTabFactory.h"
+
+class FWidgetBlueprintEditor;
+class SWidget;
+
+namespace UE::UMG::Editor
+{
+
+struct FPreviewDetailsTabSummoner : public FWorkflowTabFactory
+{
+public:
+	static const FName TabID;
+	
+public:
+	FPreviewDetailsTabSummoner(TSharedPtr<FWidgetBlueprintEditor> InBlueprintEditor);
+	
+	virtual TSharedRef<SWidget> CreateTabBody(const FWorkflowTabSpawnInfo& Info) const override;
+
+protected:
+	TWeakPtr<FWidgetBlueprintEditor> BlueprintEditor;
+};
+
+} //namespace
\ No newline at end of file
diff --git a/Engine/Source/Editor/UMGEditor/Private/TabFactory/PreviewTabSummoner.cpp b/Engine/Source/Editor/UMGEditor/Private/TabFactory/PreviewTabSummoner.cpp
--- a/Engine/Source/Editor/UMGEditor/Private/TabFactory/PreviewTabSummoner.cpp
+++ b/Engine/Source/Editor/UMGEditor/Private/TabFactory/PreviewTabSummoner.cpp
@@ -2,14 +2,19 @@
 
 #include "TabFactory/PreviewTabSummoner.h"
 
-#include "Debug/SDebugPreview.h"
+#include "Preview/SWidgetPreview.h"
 #include "UMGStyle.h"
+#include "WidgetBlueprintEditor.h"
+#include "Widgets/SWidget.h"
 
 #define LOCTEXT_NAMESPACE "UMG"
 
-const FName FPreviewTabSummoner::TabID("WidgetPreview");
+namespace UE::UMG::Editor
+{
+
+const FName FWidgetPreviewTabSummoner::TabID("WidgetPreview");
 
-FPreviewTabSummoner::FPreviewTabSummoner(TSharedPtr<FWidgetBlueprintEditor> InBlueprintEditor)
+FWidgetPreviewTabSummoner::FWidgetPreviewTabSummoner(TSharedPtr<FWidgetBlueprintEditor> InBlueprintEditor)
 		: FWorkflowTabFactory(TabID, InBlueprintEditor)
 		, BlueprintEditor(InBlueprintEditor)
 {
@@ -22,12 +27,12 @@ FPreviewTabSummoner::FPreviewTabSummoner(TSharedPtr<FWidgetBlueprintEditor> InBl
 	ViewMenuTooltip = LOCTEXT("DebugPreview_ViewMenu_ToolTip", "Show the Preview");
 }
 
-TSharedRef<SWidget> FPreviewTabSummoner::CreateTabBody(const FWorkflowTabSpawnInfo& Info) const
+TSharedRef<SWidget> FWidgetPreviewTabSummoner::CreateTabBody(const FWorkflowTabSpawnInfo& Info) const
 {
-	TSharedPtr<FWidgetBlueprintEditor> BlueprintEditorPtr = StaticCastSharedPtr<FWidgetBlueprintEditor>(BlueprintEditor.Pin());
-
-	return SNew(UE::UMG::SDebugPreview, BlueprintEditorPtr)
+	return SNew(SWidgetPreview, BlueprintEditor.Pin())
 		.AddMetaData<FTagMetaData>(FTagMetaData(TEXT("Preview")));
 }
 
+} // namespace
+
 #undef LOCTEXT_NAMESPACE 
diff --git a/Engine/Source/Editor/UMGEditor/Private/TabFactory/PreviewTabSummoner.h b/Engine/Source/Editor/UMGEditor/Private/TabFactory/PreviewTabSummoner.h
--- a/Engine/Source/Editor/UMGEditor/Private/TabFactory/PreviewTabSummoner.h
+++ b/Engine/Source/Editor/UMGEditor/Private/TabFactory/PreviewTabSummoner.h
@@ -2,20 +2,25 @@
 
 #pragma once
 
-#include "WidgetBlueprintEditor.h"
-#include "Widgets/SWidget.h"
 #include "WorkflowOrientedApp/WorkflowTabFactory.h"
 
-struct FPreviewTabSummoner : public FWorkflowTabFactory
+class FWidgetBlueprintEditor;
+
+namespace UE::UMG::Editor
+{
+
+struct FWidgetPreviewTabSummoner : public FWorkflowTabFactory
 {
 public:
 	static const FName TabID;
 	
 public:
-	FPreviewTabSummoner(TSharedPtr<FWidgetBlueprintEditor> InBlueprintEditor);
+	FWidgetPreviewTabSummoner(TSharedPtr<FWidgetBlueprintEditor> InBlueprintEditor);
 	
 	virtual TSharedRef<SWidget> CreateTabBody(const FWorkflowTabSpawnInfo& Info) const override;
 	
 protected:
 	TWeakPtr<FWidgetBlueprintEditor> BlueprintEditor;
 };
+
+} // namespace
\ No newline at end of file
diff --git a/Engine/Source/Editor/UMGEditor/Private/WidgetBlueprintEditor.cpp b/Engine/Source/Editor/UMGEditor/Private/WidgetBlueprintEditor.cpp
--- a/Engine/Source/Editor/UMGEditor/Private/WidgetBlueprintEditor.cpp
+++ b/Engine/Source/Editor/UMGEditor/Private/WidgetBlueprintEditor.cpp
@@ -38,9 +38,9 @@
 #include "Blueprint/WidgetTree.h"
 #include "WidgetBlueprintEditorUtils.h"
 #include "WorkflowOrientedApp/ApplicationMode.h"
-#include "BlueprintModes/WidgetDebugApplicationMode.h"
 #include "BlueprintModes/WidgetDesignerApplicationMode.h"
 #include "BlueprintModes/WidgetGraphApplicationMode.h"
+#include "BlueprintModes/WidgetPreviewApplicationMode.h"
 #include "WidgetModeManager.h"
 
 #include "WidgetBlueprintEditorToolbar.h"
@@ -67,6 +67,7 @@
 #include "GameProjectGenerationModule.h"
 #include "Tools/ToolCompatible.h"
 
+#include "Preview/PreviewMode.h"
 #include "Palette/SPaletteViewModel.h"
 #include "Library/SLibraryViewModel.h"
 
@@ -546,7 +547,7 @@ void FWidgetBlueprintEditor::RegisterApplicationModes(const TArray<UBlueprint*>&
 {
 	//Super::RegisterApplicationModes(InBlueprints, bShouldOpenInDefaultsMode);
 
-	if ( InBlueprints.Num() == 1 )
+	if (InBlueprints.Num() == 1)
 	{
 		TSharedPtr<FWidgetBlueprintEditor> ThisPtr(SharedThis(this));
 
@@ -555,12 +556,13 @@ void FWidgetBlueprintEditor::RegisterApplicationModes(const TArray<UBlueprint*>&
 		TempModeList.Add(MakeShared<FWidgetDesignerApplicationMode>(ThisPtr));
 		TempModeList.Add(MakeShared<FWidgetGraphApplicationMode>(ThisPtr));
 
-		if (FWidgetBlueprintApplicationModes::IsDebugModeEnabled())
+		if (FWidgetBlueprintApplicationModes::IsPreviewModeEnabled())
 		{
-			TempModeList.Add(MakeShared<FWidgetDebugApplicationMode>(ThisPtr));
+			TempModeList.Add(MakeShared<UE::UMG::Editor::FWidgetPreviewApplicationMode>(ThisPtr));
+			PreviewMode = MakeShared<UE::UMG::Editor::FPreviewMode>();
 		}
 
-		for ( TSharedRef<FApplicationMode>& AppMode : TempModeList )
+		for (TSharedRef<FApplicationMode>& AppMode : TempModeList)
 		{
 			AddApplicationMode(AppMode->GetModeName(), AppMode);
 		}
diff --git a/Engine/Source/Editor/UMGEditor/Private/WidgetBlueprintEditorToolbar.cpp b/Engine/Source/Editor/UMGEditor/Private/WidgetBlueprintEditorToolbar.cpp
--- a/Engine/Source/Editor/UMGEditor/Private/WidgetBlueprintEditorToolbar.cpp
+++ b/Engine/Source/Editor/UMGEditor/Private/WidgetBlueprintEditorToolbar.cpp
@@ -120,20 +120,20 @@ void FWidgetBlueprintEditorToolbar::FillWidgetBlueprintEditorModesToolbar(FToolB
 		// Right side padding
 		BlueprintEditorPtr->AddToolbarWidget(SNew(SSpacer).Size(FVector2D(10.0f, 1.0f)));
 
-		if (FWidgetBlueprintApplicationModes::IsDebugModeEnabled())
+		if (FWidgetBlueprintApplicationModes::IsPreviewModeEnabled())
 		{
 			BlueprintEditorPtr->AddToolbarWidget(
-				SNew(SModeWidget, FWidgetBlueprintApplicationModes::GetLocalizedMode(FWidgetBlueprintApplicationModes::DebugMode), FWidgetBlueprintApplicationModes::DebugMode)
+				SNew(SModeWidget, FWidgetBlueprintApplicationModes::GetLocalizedMode(FWidgetBlueprintApplicationModes::PreviewMode), FWidgetBlueprintApplicationModes::PreviewMode)
 				.OnGetActiveMode(GetActiveMode)
 				.OnSetActiveMode(SetActiveMode)
 				.CanBeSelected(BlueprintEditorPtr.Get(), &FBlueprintEditor::IsEditingSingleBlueprint)
 				.ToolTip(IDocumentation::Get()->CreateToolTip(
-					LOCTEXT("DebugModeButtonTooltip", "Switch to Debugging Mode"),
+					LOCTEXT("PreviewModeButtonTooltip", "Switch to Preview Mode"),
 					nullptr,
 					TEXT("Shared/Editors/BlueprintEditor"),
 					TEXT("DebugMode")))
 				.IconImage(FAppStyle::GetBrush("BlueprintDebugger.TabIcon"))
-				.AddMetaData<FTagMetaData>(FTagMetaData(TEXT("DebugMode")))
+				.AddMetaData<FTagMetaData>(FTagMetaData(TEXT("PreviewMode")))
 			);
 		
 			// Right side padding
diff --git a/Engine/Source/Editor/UMGEditor/Public/BlueprintModes/WidgetBlueprintApplicationModes.h b/Engine/Source/Editor/UMGEditor/Public/BlueprintModes/WidgetBlueprintApplicationModes.h
--- a/Engine/Source/Editor/UMGEditor/Public/BlueprintModes/WidgetBlueprintApplicationModes.h
+++ b/Engine/Source/Editor/UMGEditor/Public/BlueprintModes/WidgetBlueprintApplicationModes.h
@@ -15,11 +15,15 @@ struct UMGEDITOR_API FWidgetBlueprintApplicationModes
 	// Mode constants
 	static const FName DesignerMode;
 	static const FName GraphMode;
+	UE_DEPRECATED(5.3, "DebugMode is deprecated. Use PreviewMode instead.")
 	static const FName DebugMode;
+	static const FName PreviewMode;
 
 	static FText GetLocalizedMode(const FName InMode);
 
+	UE_DEPRECATED(5.3, "IsDebugModeEnabled is deprecated. Use IsPreviewModeEnabled instead.")
 	static bool IsDebugModeEnabled();
+	static bool IsPreviewModeEnabled();	
 
 private:
 	FWidgetBlueprintApplicationModes() = delete;
diff --git a/Engine/Source/Editor/UMGEditor/Public/Preview/PreviewMode.h b/Engine/Source/Editor/UMGEditor/Public/Preview/PreviewMode.h
new file mode 100644
--- /dev/null
+++ b/Engine/Source/Editor/UMGEditor/Public/Preview/PreviewMode.h
@@ -0,0 +1,45 @@
+// Copyright Epic Games, Inc. All Rights Reserved.
+
+#pragma once
+
+#include "CoreMinimal.h"
+#include "UObject/WeakObjectPtr.h"
+
+class UObject;
+class UUserWidget;
+
+namespace UE::UMG::Editor
+{
+
+	class UMGEDITOR_API FPreviewMode
+	{
+	public:
+		void SetSelectedObject(TArray<TWeakObjectPtr<UObject>> Objects);
+		void SetSelectedObject(TArrayView<UObject*> Objects);
+		TArray<UObject*> GetSelectedObjectList() const;
+
+		FSimpleMulticastDelegate& OnSelectedObjectChanged()
+		{
+			return SelectedObjectChangedDelegate;
+		}
+
+		void SetPreviewWidget(UUserWidget* Widget);
+
+		UUserWidget* GetPreviewWidget() const
+		{
+			return PreviewWidget.Get();
+		}
+
+		FSimpleMulticastDelegate& OnPreviewWidgetChanged()
+		{
+			return PreviewWidgetChangedDelegate;
+		}
+
+	private:
+		TArray<TWeakObjectPtr<UObject>> SelectedObjects;
+		TWeakObjectPtr<UUserWidget> PreviewWidget;
+		FSimpleMulticastDelegate SelectedObjectChangedDelegate;
+		FSimpleMulticastDelegate PreviewWidgetChangedDelegate;
+	};
+
+} // namespace UE::UMG
\ No newline at end of file
diff --git a/Engine/Source/Editor/UMGEditor/Public/WidgetBlueprintEditor.h b/Engine/Source/Editor/UMGEditor/Public/WidgetBlueprintEditor.h
--- a/Engine/Source/Editor/UMGEditor/Public/WidgetBlueprintEditor.h
+++ b/Engine/Source/Editor/UMGEditor/Public/WidgetBlueprintEditor.h
@@ -25,6 +25,7 @@ class UWidgetAnimation;
 class UWidgetBlueprint;
 class FPaletteViewModel;
 class FLibraryViewModel;
+namespace UE::UMG::Editor { class FPreviewMode; }
 
 struct FNamedSlotSelection
 {
@@ -257,6 +258,9 @@ public:
 
 	void CreateEditorModeManager() override;
 
+	/** Get the relative info for the Debug mode. */
+	TSharedPtr<UE::UMG::Editor::FPreviewMode> GetPreviewMode() const { return PreviewMode; }
+
 public:
 	/** Fires whenever a new widget is being hovered over */
 	FOnHoveredWidgetSet OnHoveredWidgetSet;
@@ -532,6 +536,8 @@ private:
 	/** ViewModel used by the Library View */
 	TSharedPtr<FLibraryViewModel> LibraryViewModel;
 
+	TSharedPtr<UE::UMG::Editor::FPreviewMode> PreviewMode;
+
 	/** When true the sequencer selection is being updated from changes to the external selection. */
 	bool bUpdatingSequencerSelection;
 
