UMG: Add callback when a new WidgetBlueprint is created from the factory. It's an entry point to add widgets/extensions from a plugin when a widget is created.

diff --git a/Engine/Source/Editor/UMGEditor/Private/UMGEditorModule.cpp b/Engine/Source/Editor/UMGEditor/Private/UMGEditorModule.cpp
--- a/Engine/Source/Editor/UMGEditor/Private/UMGEditorModule.cpp
+++ b/Engine/Source/Editor/UMGEditor/Private/UMGEditorModule.cpp
@@ -318,6 +318,11 @@ public:
 		return CustomPropertyTypeLayout;
 	}
 
+	virtual FOnWidgetBlueprintCreated& OnWidgetBlueprintCreated() override
+	{
+		return BlueprintCreatedEvent;
+	}
+
 private:
 	void RegisterAssetTypeAction(IAssetTools& AssetTools, TSharedRef<IAssetTypeActions> Action)
 	{
@@ -400,7 +405,9 @@ private:
 	FOnRegisterTabs RegisterTabsForEditor;
 
 	/** Support layout extensions */
-	FOnRegisterLayoutExtensions	RegisterLayoutExtensions;
+	FOnRegisterLayoutExtensions RegisterLayoutExtensions;
+	/** OnWidgetBlueprintCreated event */
+	FOnWidgetBlueprintCreated BlueprintCreatedEvent;
 
 	/** */
 	TArray<FCustomPropertyTypeLayout> CustomPropertyTypeLayout;
diff --git a/Engine/Source/Editor/UMGEditor/Private/WidgetBlueprintFactory.cpp b/Engine/Source/Editor/UMGEditor/Private/WidgetBlueprintFactory.cpp
--- a/Engine/Source/Editor/UMGEditor/Private/WidgetBlueprintFactory.cpp
+++ b/Engine/Source/Editor/UMGEditor/Private/WidgetBlueprintFactory.cpp
@@ -8,6 +8,7 @@
 #include "WidgetBlueprint.h"
 #include "Kismet2/KismetEditorUtilities.h"
 #include "Modules/ModuleManager.h"
+#include "UMGEditorModule.h"
 
 #include "Blueprint/WidgetTree.h"
 #include "UMGEditorProjectSettings.h"
@@ -194,6 +195,15 @@ UObject* UWidgetBlueprintFactory::FactoryCreateNew(UClass* Class, UObject* InPar
 			}
 		}
 
+		{
+			IUMGEditorModule::FWidgetBlueprintCreatedArgs Args;
+			Args.ParentClass = CurrentParentClass;
+			Args.Blueprint = NewBP;
+
+			IUMGEditorModule& UMGEditor = FModuleManager::LoadModuleChecked<IUMGEditorModule>("UMGEditor");
+			UMGEditor.OnWidgetBlueprintCreated().Broadcast(Args);
+		}
+
 		return NewBP;
 	}
 }
diff --git a/Engine/Source/Editor/UMGEditor/Public/UMGEditorModule.h b/Engine/Source/Editor/UMGEditor/Public/UMGEditorModule.h
--- a/Engine/Source/Editor/UMGEditor/Public/UMGEditorModule.h
+++ b/Engine/Source/Editor/UMGEditor/Public/UMGEditorModule.h
@@ -66,4 +66,15 @@ public:
 
 	/** Remove an instance customization to the widget property view. */
 	virtual TArrayView<const FCustomPropertyTypeLayout> GetAllInstancedCustomPropertyTypeLayout() const = 0;
+
+	/** Arguments for the OnBlueprintCreated callback. */
+	struct FWidgetBlueprintCreatedArgs
+	{
+		UClass* ParentClass = nullptr;
+		UWidgetBlueprint* Blueprint = nullptr;
+	};
+	DECLARE_EVENT_OneParam(IUMGEditorModule, FOnWidgetBlueprintCreated, FWidgetBlueprintCreatedArgs);
+
+	/** Callback when a WidgetBlueprint is created by the factory. */
+	virtual FOnWidgetBlueprintCreated& OnWidgetBlueprintCreated() = 0;
 };
