UMG: fix widgets disappearing temporarily when dragged out of the designer view. Make sure widgets are put back to their original position once they're dropped out of the designer view.

diff --git a/Engine/Source/Editor/UMGEditor/Private/Designer/SDesignerView.cpp b/Engine/Source/Editor/UMGEditor/Private/Designer/SDesignerView.cpp
--- a/Engine/Source/Editor/UMGEditor/Private/Designer/SDesignerView.cpp
+++ b/Engine/Source/Editor/UMGEditor/Private/Designer/SDesignerView.cpp
@@ -2517,6 +2517,25 @@ FReply SDesignerView::OnDragDetected(const FGeometry& MyGeometry, const FPointer
 		ClearExtensionWidgets();
 
 		TSharedRef<FSelectedWidgetDragDropOp> DragOp = FSelectedWidgetDragDropOp::New(BlueprintEditor.Pin(), this, DraggingWidgets);
+		TWeakPtr<SDesignerView> WeakDesignerView = SharedThis(this);
+		DragOp->OnDragDropEnded.AddLambda([WeakDesignerView, WeakDragOp = DragOp.ToWeakPtr()]()
+			{
+				if (TSharedPtr<SDesignerView> DesignerViewPtr = WeakDesignerView.Pin())
+				{
+					if (TSharedPtr<FWidgetBlueprintEditor> BlueprintEditorPtr = DesignerViewPtr->BlueprintEditor.Pin())
+					{
+						if (DesignerViewPtr->DropPreviews.Num() == 0)
+						{
+							DesignerViewPtr->bMovingExistingWidget = false;
+							if (WeakDragOp.IsValid() && WeakDragOp.Pin()->DraggedWidgets.Num() > 0)
+							{
+								BlueprintEditorPtr->RefreshPreview();
+							}
+						}
+					}
+				}
+			});
+
 		return FReply::Handled().BeginDragDrop(DragOp);
 	}
 
diff --git a/Engine/Source/Editor/UMGEditor/Private/DragDrop/SelectedWidgetDragDropOp.cpp b/Engine/Source/Editor/UMGEditor/Private/DragDrop/SelectedWidgetDragDropOp.cpp
--- a/Engine/Source/Editor/UMGEditor/Private/DragDrop/SelectedWidgetDragDropOp.cpp
+++ b/Engine/Source/Editor/UMGEditor/Private/DragDrop/SelectedWidgetDragDropOp.cpp
@@ -16,6 +16,12 @@ FSelectedWidgetDragDropOp::~FSelectedWidgetDragDropOp()
 	}
 }
 
+void FSelectedWidgetDragDropOp::OnDrop(bool bDropWasHandled, const FPointerEvent& MouseEvent)
+{
+	FDecoratedDragDropOp::OnDrop(bDropWasHandled, MouseEvent);
+	OnDragDropEnded.Broadcast();
+}
+
 TSharedRef<FSelectedWidgetDragDropOp> FSelectedWidgetDragDropOp::New(TSharedPtr<FWidgetBlueprintEditor> Editor, IUMGDesigner* InDesigner, const TArray<FDraggingWidgetReference>& InWidgets)
 {
 	TSharedRef<FSelectedWidgetDragDropOp> Operation = MakeShareable(new FSelectedWidgetDragDropOp());
diff --git a/Engine/Source/Editor/UMGEditor/Private/DragDrop/SelectedWidgetDragDropOp.h b/Engine/Source/Editor/UMGEditor/Private/DragDrop/SelectedWidgetDragDropOp.h
--- a/Engine/Source/Editor/UMGEditor/Private/DragDrop/SelectedWidgetDragDropOp.h
+++ b/Engine/Source/Editor/UMGEditor/Private/DragDrop/SelectedWidgetDragDropOp.h
@@ -20,8 +20,13 @@ class FSelectedWidgetDragDropOp : public FDecoratedDragDropOp
 public:
 	DRAG_DROP_OPERATOR_TYPE(FSelectedWidgetDragDropOp, FDecoratedDragDropOp);
 
+	DECLARE_MULTICAST_DELEGATE(FOnDragDropEnded);
+	FOnDragDropEnded OnDragDropEnded;
+
 	virtual ~FSelectedWidgetDragDropOp();
 
+	virtual void OnDrop(bool bDropWasHandled, const FPointerEvent& MouseEvent) override;
+
 	struct FDraggingWidgetReference
 	{
 		FWidgetReference Widget;
