UMG: Add flag to call UserWidget::OnInitialized. It would be preferable to not have that flag, like PreConstruct, Construct and Destruct but there is too much deprecation to handle.

diff --git a/Engine/Source/Editor/UMGEditor/Private/WidgetBlueprint.cpp b/Engine/Source/Editor/UMGEditor/Private/WidgetBlueprint.cpp
--- a/Engine/Source/Editor/UMGEditor/Private/WidgetBlueprint.cpp
+++ b/Engine/Source/Editor/UMGEditor/Private/WidgetBlueprint.cpp
@@ -568,6 +568,7 @@ bool FWidgetAnimation_DEPRECATED::SerializeFromMismatchedTag(struct FPropertyTag
 
 UWidgetBlueprint::UWidgetBlueprint(const FObjectInitializer& ObjectInitializer)
 	: Super(ObjectInitializer)
+	, bCanCallInitializedWithoutPlayerContext(false)
 	, TickFrequency(EWidgetTickFrequency::Auto)
 {
 }
diff --git a/Engine/Source/Editor/UMGEditor/Private/WidgetBlueprintCompiler.cpp b/Engine/Source/Editor/UMGEditor/Private/WidgetBlueprintCompiler.cpp
--- a/Engine/Source/Editor/UMGEditor/Private/WidgetBlueprintCompiler.cpp
+++ b/Engine/Source/Editor/UMGEditor/Private/WidgetBlueprintCompiler.cpp
@@ -1124,6 +1124,8 @@ void FWidgetBlueprintCompilerContext::FinishCompilingClass(UClass* Class)
 		}
 	}
 
+	BPGClass->bCanCallInitializedWithoutPlayerContext = WidgetBP->bCanCallInitializedWithoutPlayerContext;
+
 	Super::FinishCompilingClass(Class);
 
 	CA_ASSUME(BPGClass);
diff --git a/Engine/Source/Editor/UMGEditor/Private/WidgetBlueprintFactory.cpp b/Engine/Source/Editor/UMGEditor/Private/WidgetBlueprintFactory.cpp
--- a/Engine/Source/Editor/UMGEditor/Private/WidgetBlueprintFactory.cpp
+++ b/Engine/Source/Editor/UMGEditor/Private/WidgetBlueprintFactory.cpp
@@ -195,6 +195,8 @@ UObject* UWidgetBlueprintFactory::FactoryCreateNew(UClass* Class, UObject* InPar
 			}
 		}
 
+		NewBP->bCanCallInitializedWithoutPlayerContext = GetDefault<UUMGEditorProjectSettings>()->bCanCallInitializedWithoutPlayerContext;
+
 		{
 			IUMGEditorModule::FWidgetBlueprintCreatedArgs Args;
 			Args.ParentClass = CurrentParentClass;
diff --git a/Engine/Source/Editor/UMGEditor/Public/WidgetBlueprint.h b/Engine/Source/Editor/UMGEditor/Public/WidgetBlueprint.h
--- a/Engine/Source/Editor/UMGEditor/Public/WidgetBlueprint.h
+++ b/Engine/Source/Editor/UMGEditor/Public/WidgetBlueprint.h
@@ -263,6 +263,9 @@ public:
 	UPROPERTY(AssetRegistrySearchable, AssetRegistrySearchable)
 	FString PaletteCategory;
 
+	/** Run the initialize event on widget that doesn't have a player context. */
+	UPROPERTY(EditAnywhere, Category="Widget")
+	bool bCanCallInitializedWithoutPlayerContext;
 #endif
 
 public:
diff --git a/Engine/Source/Editor/UMGEditor/Public/WidgetEditingProjectSettings.h b/Engine/Source/Editor/UMGEditor/Public/WidgetEditingProjectSettings.h
--- a/Engine/Source/Editor/UMGEditor/Public/WidgetEditingProjectSettings.h
+++ b/Engine/Source/Editor/UMGEditor/Public/WidgetEditingProjectSettings.h
@@ -130,11 +130,11 @@ public:
 	virtual void PostInitProperties() override;
 
 public:
-	UPROPERTY(EditAnywhere, config, Category = Compiler)
+	UPROPERTY(EditAnywhere, config, Category = "Compiler")
 	FWidgetCompilerOptions DefaultCompilerOptions;
 
 protected:
-	UPROPERTY(EditAnywhere, config, Category = Compiler)
+	UPROPERTY(EditAnywhere, config, Category = "Compiler")
 	TArray<FDirectoryWidgetCompilerOptions> DirectoryCompilerOptions;
 
 public:
@@ -154,56 +154,60 @@ public:
 public:
 
 	/** Enables a dialog that lets you select a root widget before creating a widget blueprint */
-	UPROPERTY(EditAnywhere, config, Category = Designer)
+	UPROPERTY(EditAnywhere, config, Category = "Designer")
 	bool bUseWidgetTemplateSelector;
 
 	/** This list populates the common class section of the root widget selection dialog */
-	UPROPERTY(EditAnywhere, config, Category = Designer, meta = (EditCondition = "bUseWidgetTemplateSelector"))
+	UPROPERTY(EditAnywhere, config, Category = "Designer", meta = (EditCondition = "bUseWidgetTemplateSelector"))
 	TArray<TSoftClassPtr<UPanelWidget>> CommonRootWidgetClasses;
 
 	/** The panel widget to place at the root of all newly constructed widget blueprints. Can be empty. */
-	UPROPERTY(EditAnywhere, config, Category = Designer)
+	UPROPERTY(EditAnywhere, config, Category = "Designer")
 	TSubclassOf<UPanelWidget> DefaultRootWidget;
 
+	/** Set true to filter all categories and widgets out in the palette, selectively enabling them later via permission lists. */
+	UPROPERTY(EditAnywhere, config, Category = "Designer")
+	bool bUseEditorConfigPaletteFiltering;
+
+	/** Enables a dialog that lets you select the parent class in a tree view. */
+	UPROPERTY(EditAnywhere, config, Category = "Designer")
+	bool bUseUserWidgetParentClassViewerSelector;
+
+	/** Enables a dialog that lets you select the parent class in a default view. */
+	UPROPERTY(EditAnywhere, config, Category = "Designer")
+	bool bUseUserWidgetParentDefaultClassViewerSelector;
+
 	/** Set true to enable the Is Variable checkbox in the UMG editor DetailView. */
-	UPROPERTY(EditAnywhere, Category = "Designer | Window")
+	UPROPERTY(EditAnywhere, Category = "Designer|Window")
 	bool bEnableMakeVariable;
 
 	/** Set true to hide widget animation related elements in the UMG editor. */
-	UPROPERTY(EditAnywhere, Category = "Designer | Window")
+	UPROPERTY(EditAnywhere, Category = "Designer|Window")
 	bool bEnableWidgetAnimationEditor;
 
 	/** Set true to enabled the Palette window in the UMG editor. */
-	UPROPERTY(EditAnywhere, Category = "Designer | Window")
+	UPROPERTY(EditAnywhere, Category = "Designer|Window")
 	bool bEnablePaletteWindow;
 
 	/** Set true to enabled the LIbrary window in the UMG editor. */
-	UPROPERTY(EditAnywhere, Category = "Designer | Window")
+	UPROPERTY(EditAnywhere, Category = "Designer|Window")
 	bool bEnableLibraryWindow;
 
 	/** Set true to enabled the Widget Hierarchy window in the UMG editor. */
-	UPROPERTY(EditAnywhere, Category = "Designer | Window")
+	UPROPERTY(EditAnywhere, Category = "Designer|Window")
 	bool bEnableHierarchyWindow;
 	
 	/** Set true to enabled the Bind Widget window in the UMG editor. */
-	UPROPERTY(EditAnywhere, Category = "Designer | Window")
+	UPROPERTY(EditAnywhere, Category = "Designer|Window")
 	bool bEnableBindWidgetWindow;
 
 	/** Set true to enabled the Navigation Simulation window in the UMG editor. */
-	UPROPERTY(EditAnywhere, Category = "Designer | Window")
+	UPROPERTY(EditAnywhere, Category = "Designer|Window")
 	bool bEnableNavigationSimulationWindow;
-
-	/** Set true to filter all categories and widgets out in the palette, selectively enabling them later via permission lists. */
-	UPROPERTY(EditAnywhere, config, Category = Designer)
-	bool bUseEditorConfigPaletteFiltering;
-
-	/** Enables a dialog that lets you select the parent class in a tree view. */
-	UPROPERTY(EditAnywhere, config, Category = Designer)
-	bool bUseUserWidgetParentClassViewerSelector;
-
-	/** Enables a dialog that lets you select the parent class in a default view. */
-	UPROPERTY(EditAnywhere, config, Category = Designer)
-	bool bUseUserWidgetParentDefaultClassViewerSelector;
+	
+	/** The default value of bCanCallInitializedWithoutPlayerContext. */
+	UPROPERTY(EditAnywhere, Category = "Class Settings")
+	bool bCanCallInitializedWithoutPlayerContext;
 
 	/**
 	 * The list of parent classes to choose from for newly constructed widget blueprints.
diff --git a/Engine/Source/Runtime/UMG/Private/UserWidget.cpp b/Engine/Source/Runtime/UMG/Private/UserWidget.cpp
--- a/Engine/Source/Runtime/UMG/Private/UserWidget.cpp
+++ b/Engine/Source/Runtime/UMG/Private/UserWidget.cpp
@@ -153,7 +153,9 @@ bool UUserWidget::Initialize()
 			InitializeNamedSlots();
 		}
 
-		if (!IsDesignTime() && PlayerContext.IsValid())
+		// For backward compatibility, run the initialize event on widget that doesn't have a player context only when the class authorized it.
+		bool bClassWantsToRunInitialized = BGClass && BGClass->bCanCallInitializedWithoutPlayerContext;
+		if (!IsDesignTime() && (PlayerContext.IsValid() || bClassWantsToRunInitialized))
 		{
 			NativeOnInitialized();
 		}
diff --git a/Engine/Source/Runtime/UMG/Private/WidgetBlueprintGeneratedClass.cpp b/Engine/Source/Runtime/UMG/Private/WidgetBlueprintGeneratedClass.cpp
--- a/Engine/Source/Runtime/UMG/Private/WidgetBlueprintGeneratedClass.cpp
+++ b/Engine/Source/Runtime/UMG/Private/WidgetBlueprintGeneratedClass.cpp
@@ -123,6 +123,7 @@ namespace
 
 UWidgetBlueprintGeneratedClass::UWidgetBlueprintGeneratedClass()
 {
+	bCanCallInitializedWithoutPlayerContext = false;
 #if WITH_EDITORONLY_DATA
 	{
 		static const FAutoRegisterTextReferenceCollectorCallback AutomaticRegistrationOfTextReferenceCollector(UWidgetBlueprintGeneratedClass::StaticClass(), &CollectWidgetBlueprintGeneratedClassTextReferences);
diff --git a/Engine/Source/Runtime/UMG/Public/Blueprint/WidgetBlueprintGeneratedClass.h b/Engine/Source/Runtime/UMG/Public/Blueprint/WidgetBlueprintGeneratedClass.h
--- a/Engine/Source/Runtime/UMG/Public/Blueprint/WidgetBlueprintGeneratedClass.h
+++ b/Engine/Source/Runtime/UMG/Public/Blueprint/WidgetBlueprintGeneratedClass.h
@@ -99,6 +99,10 @@ public:
 #endif
 
 public:
+	/** */
+	UPROPERTY()
+	uint32 bCanCallInitializedWithoutPlayerContext : 1;
+
 	UPROPERTY()
 	TArray< FDelegateRuntimeBinding > Bindings;
 
