UMG: Fix bad use of a customizer for an external object inside the UMG Designer detail view.

diff --git a/Engine/Source/Editor/UMGEditor/Private/Details/DetailWidgetExtensionHandler.cpp b/Engine/Source/Editor/UMGEditor/Private/Details/DetailWidgetExtensionHandler.cpp
--- a/Engine/Source/Editor/UMGEditor/Private/Details/DetailWidgetExtensionHandler.cpp
+++ b/Engine/Source/Editor/UMGEditor/Private/Details/DetailWidgetExtensionHandler.cpp
@@ -28,14 +28,17 @@ namespace Private
 		}
 
 		// Check to see if there's a delegate for the parent property.
-		FProperty* ParentProperty = ParentPropertyHandle->GetProperty();
-		FString DelegateName = ParentProperty->GetName() + "Delegate";
-
-		if (UClass* ContainerClass = ParentProperty->GetOwner<UClass>())
+		if (ParentPropertyHandle)
 		{
-			FDelegateProperty* DelegateProperty = FindFProperty<FDelegateProperty>(ContainerClass, FName(*DelegateName));
-			return DelegateProperty;
-		} 
+			FProperty* ParentProperty = ParentPropertyHandle->GetProperty();
+			FString DelegateName = ParentProperty->GetName() + "Delegate";
+
+			if (UClass* ContainerClass = ParentProperty->GetOwner<UClass>())
+			{
+				FDelegateProperty* DelegateProperty = FindFProperty<FDelegateProperty>(ContainerClass, FName(*DelegateName));
+				return DelegateProperty;
+			}
+		}
 
 		return nullptr;
 	}
