UMG: Always allow widget from the same plugin.

diff --git a/Engine/Source/Editor/UMGEditor/Private/WidgetBlueprintEditorUtils.cpp b/Engine/Source/Editor/UMGEditor/Private/WidgetBlueprintEditorUtils.cpp
--- a/Engine/Source/Editor/UMGEditor/Private/WidgetBlueprintEditorUtils.cpp
+++ b/Engine/Source/Editor/UMGEditor/Private/WidgetBlueprintEditorUtils.cpp
@@ -5,6 +5,7 @@
 #include "Components/PanelWidget.h"
 #include "Components/ContentWidget.h"
 #include "Engine/Texture2D.h"
+#include "Interfaces/IPluginManager.h"
 #include "UObject/UObjectHash.h"
 #include "UObject/UObjectIterator.h"
 #include "Internationalization/TextPackageNamespaceUtil.h"
@@ -12,6 +13,7 @@
 #include "UObject/TopLevelAssetPath.h"
 #include "Blueprint/WidgetTree.h"
 #include "Misc/ConfigCacheIni.h"
+#include "Misc/PathViews.h"
 #include "Modules/ModuleManager.h"
 #include "MovieScene.h"
 #include "UMGEditorProjectSettings.h"
@@ -2235,7 +2237,19 @@ bool IsUsableWidgetClass(const FString& WidgetPathName, const FAssetData& Widget
 				return GlobalClassFilter->IsUnloadedClassAllowed(FClassViewerInitializationOptions(), UnloadedBlueprint, ClassViewerModule->CreateFilterFuncs());
 			}
 		}
-		return false;
+
+		auto IsPathUnderMountPoints = [](FStringView Path)
+		{
+			const TSet<FString>& MountPoints = IPluginManager::Get().GetBuiltInPluginNames();
+			if (MountPoints.Num() > 0)
+			{
+				const FStringView MountPoint = FPathViews::GetMountPointNameFromPath(Path);
+				return MountPoints.ContainsByHash(GetTypeHash(MountPoint), MountPoint);
+			}
+			return false;
+		};
+
+		return IsPathUnderMountPoints(WidgetPathName);
 	}
 	else
 	{
diff --git a/Engine/Source/Editor/UMGEditor/UMGEditor.Build.cs b/Engine/Source/Editor/UMGEditor/UMGEditor.Build.cs
--- a/Engine/Source/Editor/UMGEditor/UMGEditor.Build.cs
+++ b/Engine/Source/Editor/UMGEditor/UMGEditor.Build.cs
@@ -39,6 +39,7 @@ public class UMGEditor : ModuleRules
 				"BlueprintGraph",
 				"GraphEditor",
 				"Kismet",  // for FWorkflowCentricApplication
+				"Projects",
 				"PropertyPath",
 				"PropertyEditor",
 				"UMG",
