Remove uclasses from the widget palette when it's a built in asset and not in the specific allowed list.
Non-built in assets are user generated assets and should be allowed.

diff --git a/Engine/Source/Editor/UMGEditor/Private/WidgetBlueprintEditorUtils.cpp b/Engine/Source/Editor/UMGEditor/Private/WidgetBlueprintEditorUtils.cpp
--- a/Engine/Source/Editor/UMGEditor/Private/WidgetBlueprintEditorUtils.cpp
+++ b/Engine/Source/Editor/UMGEditor/Private/WidgetBlueprintEditorUtils.cpp
@@ -2218,11 +2218,6 @@ bool IsUsableWidgetClass(const FString& WidgetPathName, const FAssetData& Widget
 
 	if (UMGEditorProjectSettings->bUseEditorConfigPaletteFiltering)
 	{
-		if (UMGEditorProjectSettings->GetAllowedPaletteWidgets().PassesFilter(WidgetPathName))
-		{
-			return true;
-		}
-
 		FClassViewerModule* ClassViewerModule = FModuleManager::GetModulePtr<FClassViewerModule>("ClassViewer");
 		const TSharedPtr<IClassViewerFilter> GlobalClassFilter = ClassViewerModule ? ClassViewerModule->GetGlobalClassViewerFilter() : TSharedPtr<IClassViewerFilter>();
 		if (UMGEditorProjectSettings->GetAllowedPaletteCategories().PassesFilter(Category) && GlobalClassFilter.IsValid())
@@ -2239,17 +2234,34 @@ bool IsUsableWidgetClass(const FString& WidgetPathName, const FAssetData& Widget
 		}
 
 		auto IsPathUnderMountPoints = [](FStringView Path)
-		{
-			const TSet<FString>& MountPoints = IPluginManager::Get().GetBuiltInPluginNames();
-			if (MountPoints.Num() > 0)
 			{
-				const FStringView MountPoint = FPathViews::GetMountPointNameFromPath(Path);
-				return MountPoints.ContainsByHash(GetTypeHash(MountPoint), MountPoint);
-			}
+				static const FString EnginePath = TEXT("Engine");
+				static const FString GamePath = TEXT("Game");
+
+				const TSet<FString>& MountPoints = IPluginManager::Get().GetBuiltInPluginNames();
+				if (MountPoints.Num() > 0)
+				{
+					const FStringView MountPoint = FPathViews::GetMountPointNameFromPath(Path);
+					return MountPoints.ContainsByHash(GetTypeHash(MountPoint), MountPoint)
+						|| MountPoint.Equals(EnginePath, ESearchCase::IgnoreCase)
+						|| MountPoint.Equals(GamePath, ESearchCase::IgnoreCase);
+				}
+				return false;
+			};
+
+		const bool bPassesAllowedPalletteFilter = UMGEditorProjectSettings->GetAllowedPaletteWidgets().PassesFilter(WidgetPathName);
+		if (FPackageName::IsScriptPackage(WidgetPathName))
+		{
+			return bPassesAllowedPalletteFilter;
+		}
+
+		const bool bPathUnderMountPoints = IsPathUnderMountPoints(WidgetPathName);
+		if (bPathUnderMountPoints && !bPassesAllowedPalletteFilter)
+		{
 			return false;
-		};
+		}
 
-		return IsPathUnderMountPoints(WidgetPathName);
+		return true;
 	}
 	else
 	{
