Fix NamedSlot failing to compile if parent widgets has any BindWidget

diff --git a/Engine/Source/Editor/UMGEditor/Private/WidgetBlueprintCompiler.cpp b/Engine/Source/Editor/UMGEditor/Private/WidgetBlueprintCompiler.cpp
--- a/Engine/Source/Editor/UMGEditor/Private/WidgetBlueprintCompiler.cpp
+++ b/Engine/Source/Editor/UMGEditor/Private/WidgetBlueprintCompiler.cpp
@@ -488,15 +488,40 @@ void FWidgetBlueprintCompilerContext::CreateClassVariablesFromBlueprint()
 		Widgets = WidgetBPToScan->GetAllSourceWidgets();
 		if (Widgets.Num() != 0)
 		{
-			// We found widgets.
+			// We found widgets. Stop search, but still check if we have a parent for bind widget validation
+			UWidgetBlueprint* ParentWidgetBP = WidgetBPToScan->ParentClass && WidgetBPToScan->ParentClass->ClassGeneratedBy
+				? Cast<UWidgetBlueprint>(WidgetBPToScan->ParentClass->ClassGeneratedBy)
+				: nullptr;
+
+			if (ParentWidgetBP)
+			{
+				TArray<UWidget*> ParentOwnedWidgets = ParentWidgetBP->GetAllSourceWidgets();
+				ParentOwnedWidgets.Sort([](const UWidget& Lhs, const UWidget& Rhs) { return Rhs.GetFName().LexicalLess(Lhs.GetFName()); });
+
+				for (UWidget* ParentOwnedWidget : ParentOwnedWidgets)
+				{
+					// Look in the Parent class properties to find a property with the BindWidget meta tag of the same name and Type.
+					FObjectPropertyBase* ExistingProperty = CastField<FObjectPropertyBase>(ParentClass->FindPropertyByName(ParentOwnedWidget->GetFName()));
+					if (ExistingProperty &&
+						FWidgetBlueprintEditorUtils::IsBindWidgetProperty(ExistingProperty) &&
+						ParentOwnedWidget->IsA(ExistingProperty->PropertyClass))
+					{
+						ParentWidgetToBindWidgetMap.Add(ParentOwnedWidget, ExistingProperty);
+					}
+				}
+			}
+
 			break;
 		}
+
 		// We don't want to create variables for widgets that are in a parent blueprint. They will be created at the Parent compilation.
 		// But we want them to be added to the Member variable map for validation of the BindWidget property
 		bSkipVariableCreation = true;
 		
 		// Get the parent WidgetBlueprint
-		WidgetBPToScan = WidgetBPToScan->ParentClass && WidgetBPToScan->ParentClass->ClassGeneratedBy ? Cast<UWidgetBlueprint>(WidgetBPToScan->ParentClass->ClassGeneratedBy):nullptr;
+		WidgetBPToScan = WidgetBPToScan->ParentClass && WidgetBPToScan->ParentClass->ClassGeneratedBy 
+			? Cast<UWidgetBlueprint>(WidgetBPToScan->ParentClass->ClassGeneratedBy)
+			: nullptr;
 	}
 
 	// Sort the widgets alphabetically
@@ -1055,6 +1080,13 @@ void FWidgetBlueprintCompilerContext::FinishCompilingClass(UClass* Class)
 				const FText IncorrectWidgetTypeError = LOCTEXT("IncorrectWidgetTypes", "The widget @@ is of type @@, but the bind widget property is of type @@.");
 
 				UWidget* const* Widget = WidgetToMemberVariableMap.FindKey(WidgetProperty);
+
+				// If at first we don't find the binding, search the parent binding map
+				if (!Widget)
+				{
+					Widget = ParentWidgetToBindWidgetMap.FindKey(WidgetProperty);
+				}
+
 				if (!Widget)
 				{
 					if (bIsOptional)
@@ -1208,6 +1240,7 @@ void FWidgetBlueprintCompilerContext::OnPostCDOCompiled(const UObject::FPostCDOC
 
 	WidgetToMemberVariableMap.Empty();
 	WidgetAnimToMemberVariableMap.Empty();
+	ParentWidgetToBindWidgetMap.Empty();
 
 	UWidgetBlueprintGeneratedClass* WidgetClass = NewWidgetBlueprintClass;
 	UWidgetBlueprint* WidgetBP = WidgetBlueprint();
diff --git a/Engine/Source/Editor/UMGEditor/Public/WidgetBlueprintCompiler.h b/Engine/Source/Editor/UMGEditor/Public/WidgetBlueprintCompiler.h
--- a/Engine/Source/Editor/UMGEditor/Public/WidgetBlueprintCompiler.h
+++ b/Engine/Source/Editor/UMGEditor/Public/WidgetBlueprintCompiler.h
@@ -127,6 +127,9 @@ protected:
 	// Map of properties created for widgets; to aid in debug data generation
 	TMap<UWidget*, FProperty*> WidgetToMemberVariableMap;
 
+	// Map of properties created in parent widget for bind widget validation
+	TMap<UWidget*, FProperty*> ParentWidgetToBindWidgetMap;
+
 	// Map of properties created for widget animations; to aid in debug data generation
 	TMap<UWidgetAnimation*, FProperty*> WidgetAnimToMemberVariableMap;
 
