UMG: Add a message when a multicast delegate is available but doesn't show up because the widget is not a variable.

diff --git a/Engine/Source/Editor/UMGEditor/Private/Customizations/UMGDetailCustomizations.cpp b/Engine/Source/Editor/UMGEditor/Private/Customizations/UMGDetailCustomizations.cpp
--- a/Engine/Source/Editor/UMGEditor/Private/Customizations/UMGDetailCustomizations.cpp
+++ b/Engine/Source/Editor/UMGEditor/Private/Customizations/UMGDetailCustomizations.cpp
@@ -822,68 +822,76 @@ void FBlueprintWidgetCustomization::CreateMulticastEventCustomization(IDetailLay
 		PropertyTooltip = FText::FromString(DelegateProperty->GetName());
 	}
 
+	static FText EventCategoryText = LOCTEXT("Events", "Events");
+	IDetailCategoryBuilder& EventCategory = DetailLayout.EditCategory(TEXT("Events"), EventCategoryText, ECategoryPriority::Uncommon);
 	FObjectProperty* ComponentProperty = FindFProperty<FObjectProperty>(BlueprintObj->SkeletonGeneratedClass, ThisComponentName);
-
-	if ( !ComponentProperty )
+	if (ComponentProperty)
 	{
-		return;
-	}
-
-	FName PropertyName = ComponentProperty->GetFName();
-	FName EventName = DelegateProperty->GetFName();
-	FText EventText = DelegateProperty->GetDisplayNameText();
-
-	IDetailCategoryBuilder& EventCategory = DetailLayout.EditCategory(TEXT("Events"), LOCTEXT("Events", "Events"), ECategoryPriority::Uncommon);
+		FName PropertyName = ComponentProperty->GetFName();
+		FName EventName = DelegateProperty->GetFName();
+		FText EventText = DelegateProperty->GetDisplayNameText();
 
-	EventCategory.AddCustomRow(EventText)
-		.WholeRowContent()
-		[
-			SNew(SHorizontalBox)
-			.ToolTipText(DelegateProperty->GetToolTipText())
-			+ SHorizontalBox::Slot()
-			.AutoWidth()
-			.VAlign(VAlign_Center)
-			.Padding(0, 0, 5, 0)
+		EventCategory.AddCustomRow(EventText)
+			.WholeRowContent()
 			[
-				SNew(SImage)
-				.Image(FAppStyle::Get().GetBrush("GraphEditor.Event_16x"))
-			]
+				SNew(SHorizontalBox)
+				.ToolTipText(DelegateProperty->GetToolTipText())
+				+ SHorizontalBox::Slot()
+				.AutoWidth()
+				.VAlign(VAlign_Center)
+				.Padding(0, 0, 5, 0)
+				[
+					SNew(SImage)
+					.Image(FAppStyle::Get().GetBrush("GraphEditor.Event_16x"))
+				]
 
-			+ SHorizontalBox::Slot()
-			.VAlign(VAlign_Center)
-			.HAlign(HAlign_Left)
-			[
-				SNew(STextBlock)
-				.Font(IDetailLayoutBuilder::GetDetailFont())
-				.Text(EventText)
-			]
+				+ SHorizontalBox::Slot()
+				.VAlign(VAlign_Center)
+				.HAlign(HAlign_Left)
+				[
+					SNew(STextBlock)
+					.Font(IDetailLayoutBuilder::GetDetailFont())
+					.Text(EventText)
+				]
 
-			+ SHorizontalBox::Slot()
-			.HAlign(HAlign_Left)
-			.VAlign(VAlign_Center)
-			.Padding(0)
-			[
-				SNew(SButton)
-				.ContentPadding(FMargin(3.0, 2.0))
-				.OnClicked(this, &FBlueprintWidgetCustomization::HandleAddOrViewEventForVariable, EventName, PropertyName, MakeWeakObjectPtr(PropertyClass))
+				+ SHorizontalBox::Slot()
+				.HAlign(HAlign_Left)
+				.VAlign(VAlign_Center)
+				.Padding(0)
 				[
-					SNew(SWidgetSwitcher)
-					.WidgetIndex(this, &FBlueprintWidgetCustomization::HandleAddOrViewIndexForButton, EventName, PropertyName)
-					+ SWidgetSwitcher::Slot()
+					SNew(SButton)
+					.ContentPadding(FMargin(3.0, 2.0))
+					.OnClicked(this, &FBlueprintWidgetCustomization::HandleAddOrViewEventForVariable, EventName, PropertyName, MakeWeakObjectPtr(PropertyClass))
 					[
-						SNew(SImage)
-						.ColorAndOpacity(FSlateColor::UseForeground())
-						.Image(FAppStyle::Get().GetBrush("Icons.SelectInViewport"))
-					]
-					+ SWidgetSwitcher::Slot()
-					[
-						SNew(SImage)
-						.ColorAndOpacity(FSlateColor::UseForeground())
-						.Image(FAppStyle::Get().GetBrush("Icons.Plus"))
+						SNew(SWidgetSwitcher)
+						.WidgetIndex(this, &FBlueprintWidgetCustomization::HandleAddOrViewIndexForButton, EventName, PropertyName)
+						+ SWidgetSwitcher::Slot()
+						[
+							SNew(SImage)
+							.ColorAndOpacity(FSlateColor::UseForeground())
+							.Image(FAppStyle::Get().GetBrush("Icons.SelectInViewport"))
+						]
+						+ SWidgetSwitcher::Slot()
+						[
+							SNew(SImage)
+							.ColorAndOpacity(FSlateColor::UseForeground())
+							.Image(FAppStyle::Get().GetBrush("Icons.Plus"))
+						]
 					]
 				]
-			]
-		];
+			];
+	}
+	else if (!bCreateMulticastEventCustomizationErrorAdded)
+	{
+		bCreateMulticastEventCustomizationErrorAdded = true;
+		EventCategory.AddCustomRow(FText::GetEmpty())
+			.WholeRowContent()
+			[
+				SNew(STextBlock)
+					.Font(IDetailLayoutBuilder::GetDetailFont())
+					.Text(LOCTEXT("EventAvailableButNotVariable", "To see available events, enable the Is Variable setting for this widget."))
+			];
+	}
 }
 
 void FBlueprintWidgetCustomization::CustomizeDetails( IDetailLayoutBuilder& DetailLayout )
@@ -946,6 +954,7 @@ void FBlueprintWidgetCustomization::PerformBindingCustomization(IDetailLayoutBui
 {
 	static const FName IsBindableEventName(TEXT("IsBindableEvent"));
 
+	bCreateMulticastEventCustomizationErrorAdded = false;
 	if ( Widgets.Num() == 1 )
 	{
 		UWidget* Widget = Widgets[0];
diff --git a/Engine/Source/Editor/UMGEditor/Private/Customizations/UMGDetailCustomizations.h b/Engine/Source/Editor/UMGEditor/Private/Customizations/UMGDetailCustomizations.h
--- a/Engine/Source/Editor/UMGEditor/Private/Customizations/UMGDetailCustomizations.h
+++ b/Engine/Source/Editor/UMGEditor/Private/Customizations/UMGDetailCustomizations.h
@@ -54,8 +54,9 @@ private:
 	void PerformAccessibilityCustomization(IDetailLayoutBuilder& DetailLayout);
 	void CustomizeAccessibilityProperty(IDetailLayoutBuilder& DetailLayout, const FName& BehaviorPropertyName, const FName& TextPropertyName);
 	void PerformCustomizationExtenders(IDetailLayoutBuilder& DetailLayout, const TArrayView<UWidget*> Widgets);
-private:
 
+private:
 	TWeakPtr<FWidgetBlueprintEditor> Editor;
 	TWeakObjectPtr<UWidgetBlueprint> Blueprint;
+	bool bCreateMulticastEventCustomizationErrorAdded = false;
 };
