UMG Designer: fix the lock widget functionality in hierachy.
Expected behaviour:
- users can select a locked widget in the hierarchy and see its details panel in read-only mode.
- users cannot cut, paste into, delete, or rename the locked widget or use Wrap With and Replace With on it (they can copy and duplicate)
- users cannot drag a locked widget or drop onto a locked widget

diff --git a/Engine/Source/Editor/UMGEditor/Private/Details/SWidgetDetailsView.cpp b/Engine/Source/Editor/UMGEditor/Private/Details/SWidgetDetailsView.cpp
--- a/Engine/Source/Editor/UMGEditor/Private/Details/SWidgetDetailsView.cpp
+++ b/Engine/Source/Editor/UMGEditor/Private/Details/SWidgetDetailsView.cpp
@@ -67,10 +67,13 @@ void SWidgetDetailsView::Construct(const FArguments& InArgs, TSharedPtr<FWidgetB
 	// Notify us of object selection changes so we can update the package re-mapping
 	PropertyView->SetOnObjectArrayChanged(FOnObjectArrayChanged::CreateSP(this, &SWidgetDetailsView::OnPropertyViewObjectArrayChanged));
 
+	PropertyView->SetIsPropertyEditingEnabledDelegate(FIsPropertyEditingEnabled::CreateSP(this, &SWidgetDetailsView::IsDetailsPanelEditingAllowed));
+
 	TSharedRef<SWidget> IsVariableCheckbox =
 		SNew(SCheckBox)
 		.IsChecked(this, &SWidgetDetailsView::GetIsVariable)
 		.OnCheckStateChanged(this, &SWidgetDetailsView::HandleIsVariableChanged)
+		.IsEnabled(this, &SWidgetDetailsView::IsDetailsPanelEditingAllowed)
 		.Padding(FMargin(3.0f, 1.0f, 18.0f, 1.0f))
 		[
 			SNew(STextBlock)
@@ -153,7 +156,7 @@ void SWidgetDetailsView::Construct(const FArguments& InArgs, TSharedPtr<FWidgetB
 							.SelectAllTextWhenFocused(true)
 							.HintText(LOCTEXT("Name", "Name"))
 							.Text(this, &SWidgetDetailsView::GetNameText)
-							.IsEnabled(this, &SWidgetDetailsView::IsSingleObjectSelected)
+							.IsEnabled(this, &SWidgetDetailsView::IsWidgetNameFieldEnabled)
 							.OnTextChanged(this, &SWidgetDetailsView::HandleNameTextChanged)
 							.OnTextCommitted(this, &SWidgetDetailsView::HandleNameTextCommitted)
 						]
@@ -339,6 +342,26 @@ void SWidgetDetailsView::ClearFocusIfOwned()
 	}
 }
 
+bool SWidgetDetailsView::IsDetailsPanelEditingAllowed() const 
+{
+	return !SelectedObjects.ContainsByPredicate([](const TWeakObjectPtr<UObject> WeakObj)
+		{
+			if (UObject* ObjPtr = WeakObj.Get())
+			{
+				if (UWidget* Widget = Cast<UWidget>(ObjPtr))
+				{
+					return Widget->IsLockedInDesigner();
+				}
+			}
+			return false;
+		});
+}
+
+bool SWidgetDetailsView::IsWidgetNameFieldEnabled() const
+{
+	return IsSingleObjectSelected() && IsDetailsPanelEditingAllowed();
+}
+
 bool SWidgetDetailsView::IsPropertyVisible(const FPropertyAndParent& PropertyAndParent) const
 {
 	const FProperty& Property = PropertyAndParent.Property;
diff --git a/Engine/Source/Editor/UMGEditor/Private/Details/SWidgetDetailsView.h b/Engine/Source/Editor/UMGEditor/Private/Details/SWidgetDetailsView.h
--- a/Engine/Source/Editor/UMGEditor/Private/Details/SWidgetDetailsView.h
+++ b/Engine/Source/Editor/UMGEditor/Private/Details/SWidgetDetailsView.h
@@ -46,6 +46,10 @@ private:
 	/** Handles the callback from the property detail view confirming the list of objects being edited has changed */
 	void OnPropertyViewObjectArrayChanged(const FString& InTitle, const TArray<UObject*>& InObjects);
 
+	bool IsDetailsPanelEditingAllowed() const;
+
+	bool IsWidgetNameFieldEnabled() const;
+
 	void ClearFocusIfOwned();
 
 	bool IsPropertyVisible(const FPropertyAndParent& PropertyAndParent) const;
diff --git a/Engine/Source/Editor/UMGEditor/Private/Hierarchy/SHierarchyViewItem.cpp b/Engine/Source/Editor/UMGEditor/Private/Hierarchy/SHierarchyViewItem.cpp
--- a/Engine/Source/Editor/UMGEditor/Private/Hierarchy/SHierarchyViewItem.cpp
+++ b/Engine/Source/Editor/UMGEditor/Private/Hierarchy/SHierarchyViewItem.cpp
@@ -203,8 +203,16 @@ TOptional<EItemDropZone> ProcessHierarchyDragDrop(const FDragDropEvent& DragDrop
 			DecoratedDragDropOp->ResetToDefaultToolTip();
 		}
 
+		// Are we adding to a locked widget?
+		if ( TargetItem.IsValid() && TargetItem.GetPreview()->IsLockedInDesigner() )
+		{
+			if (DecoratedDragDropOp.IsValid())
+			{
+				DecoratedDragDropOp->CurrentHoverText = LOCTEXT("LockedWidget", "Widget is locked.");
+			}
+		}
 		// Are we adding to the root?
-		if ( !TargetItem.IsValid() && Blueprint->WidgetTree->RootWidget == nullptr )
+		else if ( !TargetItem.IsValid() && Blueprint->WidgetTree->RootWidget == nullptr )
 		{
 			// TODO UMG Allow showing a preview of this.
 			if ( bIsDrop )
@@ -307,6 +315,13 @@ TOptional<EItemDropZone> ProcessHierarchyDragDrop(const FDragDropEvent& DragDrop
 				return TOptional<EItemDropZone>();
 			}
 
+			if (TargetItem.GetPreview()->IsLockedInDesigner())
+			{
+				HierarchyDragDropOp->CurrentIconBrush = FAppStyle::GetBrush(TEXT("Graph.ConnectorFeedback.Error"));
+				HierarchyDragDropOp->CurrentHoverText = LOCTEXT("LockedWidget", "Widget is locked.");
+				return TOptional<EItemDropZone>();
+			}
+
 			UPanelWidget* NewParent = Cast<UPanelWidget>(TargetItem.GetTemplate());
 			if (!NewParent)
 			{
@@ -491,7 +506,7 @@ TOptional<EItemDropZone> FHierarchyModel::HandleCanAcceptDrop(const FDragDropEve
 
 FReply FHierarchyModel::HandleDragDetected(const FGeometry& MyGeometry, const FPointerEvent& MouseEvent)
 {
-	if (!IsRoot())
+	if (!IsRoot() && !IsLockedInDesigner())
 	{
 		TArray<FWidgetReference> DraggedItems;
 
@@ -1399,7 +1414,7 @@ void FHierarchyWidget::UpdateSelection()
 
 bool FHierarchyWidget::CanRename() const
 {
-	return true;
+	return !IsLockedInDesigner();
 }
 
 void FHierarchyWidget::RequestBeginRename()
diff --git a/Engine/Source/Editor/UMGEditor/Private/Hierarchy/SHierarchyViewItem.h b/Engine/Source/Editor/UMGEditor/Private/Hierarchy/SHierarchyViewItem.h
--- a/Engine/Source/Editor/UMGEditor/Private/Hierarchy/SHierarchyViewItem.h
+++ b/Engine/Source/Editor/UMGEditor/Private/Hierarchy/SHierarchyViewItem.h
@@ -68,7 +68,7 @@ public:
 	virtual void SetIsVisible(bool IsVisible) { }
 
 	virtual bool CanControlLockedInDesigner() const { return false; }
-	virtual bool IsLockedInDesigner() { return false; }
+	virtual bool IsLockedInDesigner() const { return false; }
 	virtual void SetIsLockedInDesigner(bool NewIsLocked, bool bRecursive)
 	{
 		if (bRecursive)
@@ -291,7 +291,7 @@ public:
 		return true;
 	}
 
-	virtual bool IsLockedInDesigner() override
+	virtual bool IsLockedInDesigner() const override
 	{
 		UWidget* TemplateWidget = Item.GetTemplate();
 		if (TemplateWidget)
diff --git a/Engine/Source/Editor/UMGEditor/Private/WidgetBlueprintEditor.cpp b/Engine/Source/Editor/UMGEditor/Private/WidgetBlueprintEditor.cpp
--- a/Engine/Source/Editor/UMGEditor/Private/WidgetBlueprintEditor.cpp
+++ b/Engine/Source/Editor/UMGEditor/Private/WidgetBlueprintEditor.cpp
@@ -769,7 +769,7 @@ void FWidgetBlueprintEditor::OnObjectsReplaced(const TMap<UObject*, UObject*>& R
 bool FWidgetBlueprintEditor::CanDeleteSelectedWidgets()
 {
 	TSet<FWidgetReference> Widgets = GetSelectedWidgets();
-	return Widgets.Num() > 0;
+	return Widgets.Num() > 0 && !FWidgetBlueprintEditorUtils::IsAnySelectedWidgetLocked(Widgets);
 }
 
 void FWidgetBlueprintEditor::DeleteSelectedWidgets()
@@ -797,7 +797,7 @@ void FWidgetBlueprintEditor::CopySelectedWidgets()
 bool FWidgetBlueprintEditor::CanCutSelectedWidgets()
 {
 	TSet<FWidgetReference> Widgets = GetSelectedWidgets();
-	return Widgets.Num() > 0;
+	return Widgets.Num() > 0 && !FWidgetBlueprintEditorUtils::IsAnySelectedWidgetLocked(Widgets);
 }
 
 void FWidgetBlueprintEditor::CutSelectedWidgets()
@@ -814,6 +814,12 @@ const UWidgetAnimation* FWidgetBlueprintEditor::RefreshCurrentAnimation()
 bool FWidgetBlueprintEditor::CanPasteWidgets()
 {
 	TSet<FWidgetReference> Widgets = GetSelectedWidgets();
+
+	if (FWidgetBlueprintEditorUtils::IsAnySelectedWidgetLocked(Widgets))
+	{
+		return false;
+	}
+
 	if ( Widgets.Num() == 1 )
 	{
 		// Always return true here now since we want to support pasting widgets as siblings
diff --git a/Engine/Source/Editor/UMGEditor/Private/WidgetBlueprintEditorUtils.cpp b/Engine/Source/Editor/UMGEditor/Private/WidgetBlueprintEditorUtils.cpp
--- a/Engine/Source/Editor/UMGEditor/Private/WidgetBlueprintEditorUtils.cpp
+++ b/Engine/Source/Editor/UMGEditor/Private/WidgetBlueprintEditorUtils.cpp
@@ -472,19 +472,22 @@ void FWidgetBlueprintEditorUtils::CreateWidgetContextMenu(FMenuBuilder& MenuBuil
 				)
 			);
 
-		MenuBuilder.AddSubMenu(
-			LOCTEXT("WidgetTree_WrapWith", "Wrap With..."),
-			LOCTEXT("WidgetTree_WrapWithToolTip", "Wraps the currently selected widgets inside of another container widget"),
-			FNewMenuDelegate::CreateStatic(&FWidgetBlueprintEditorUtils::BuildWrapWithMenu, BlueprintEditor, BP, Widgets)
-			);
-
-		if ( Widgets.Num() == 1 )
+		if (!FWidgetBlueprintEditorUtils::IsAnySelectedWidgetLocked(Widgets))
 		{
 			MenuBuilder.AddSubMenu(
-				LOCTEXT("WidgetTree_ReplaceWith", "Replace With..."),
-				LOCTEXT("WidgetTree_ReplaceWithToolTip", "Replaces the currently selected widget, with another widget"),
-				FNewMenuDelegate::CreateStatic(&FWidgetBlueprintEditorUtils::BuildReplaceWithMenu, BlueprintEditor, BP, Widgets)
+				LOCTEXT("WidgetTree_WrapWith", "Wrap With..."),
+				LOCTEXT("WidgetTree_WrapWithToolTip", "Wraps the currently selected widgets inside of another container widget"),
+				FNewMenuDelegate::CreateStatic(&FWidgetBlueprintEditorUtils::BuildWrapWithMenu, BlueprintEditor, BP, Widgets)
+			);
+
+			if (Widgets.Num() == 1)
+			{
+				MenuBuilder.AddSubMenu(
+					LOCTEXT("WidgetTree_ReplaceWith", "Replace With..."),
+					LOCTEXT("WidgetTree_ReplaceWithToolTip", "Replaces the currently selected widget, with another widget"),
+					FNewMenuDelegate::CreateStatic(&FWidgetBlueprintEditorUtils::BuildReplaceWithMenu, BlueprintEditor, BP, Widgets)
 				);
+			}
 		}
 	}
 	MenuBuilder.EndSection();
@@ -1502,6 +1505,19 @@ TArray<UWidget*> FWidgetBlueprintEditorUtils::DuplicateWidgets(TSharedRef<FWidge
 
 	return DuplicatedWidgets;
 }
+
+bool FWidgetBlueprintEditorUtils::IsAnySelectedWidgetLocked(TSet<FWidgetReference> SelectedWidgets)
+{
+	for (const FWidgetReference& Widget : SelectedWidgets)
+	{
+		if (Widget.GetPreview()->IsLockedInDesigner())
+		{
+			return true;
+		}
+	}
+	return false;
+}
+
 UWidget* FWidgetBlueprintEditorUtils::GetWidgetTemplateFromDragDrop(UWidgetBlueprint* Blueprint, UWidgetTree* RootWidgetTree, TSharedPtr<FDragDropOperation>& DragDropOp)
 {
 	UWidget* Widget = nullptr;
diff --git a/Engine/Source/Editor/UMGEditor/Public/WidgetBlueprintEditorUtils.h b/Engine/Source/Editor/UMGEditor/Public/WidgetBlueprintEditorUtils.h
--- a/Engine/Source/Editor/UMGEditor/Public/WidgetBlueprintEditorUtils.h
+++ b/Engine/Source/Editor/UMGEditor/Public/WidgetBlueprintEditorUtils.h
@@ -49,6 +49,8 @@ public:
 
 	static TArray<UWidget*> DuplicateWidgets(TSharedRef<FWidgetBlueprintEditor> BlueprintEditor, UWidgetBlueprint* BP, TSet<FWidgetReference> Widgets);
 
+	static bool IsAnySelectedWidgetLocked(TSet<FWidgetReference> SelectedWidgets);
+
 	static UWidget* GetWidgetTemplateFromDragDrop(UWidgetBlueprint* Blueprint, UWidgetTree* RootWidgetTree, TSharedPtr<FDragDropOperation>& DragDropOp);
 
 	static bool IsBindWidgetProperty(const FProperty* InProperty);
