UMG: fix single-cast delegates are not bind-able from UMG details panel if the project setting "Property Binding Rule" is set to "Prevent" (introduced by CL 24102693).
This does not affect multicast delegates as they're customized differently and those events are still bindable regardless of  "Property Binding Rule" setting. We want the same behaviour for single-cast events.

diff --git a/Engine/Source/Editor/UMGEditor/Private/Customizations/UMGDetailCustomizations.cpp b/Engine/Source/Editor/UMGEditor/Private/Customizations/UMGDetailCustomizations.cpp
--- a/Engine/Source/Editor/UMGEditor/Private/Customizations/UMGDetailCustomizations.cpp
+++ b/Engine/Source/Editor/UMGEditor/Private/Customizations/UMGDetailCustomizations.cpp
@@ -101,7 +101,7 @@ private:
 	TSharedPtr<FEdGraphSchemaAction> Action;
 };
 
-TSharedRef<SWidget> FBlueprintWidgetCustomization::MakePropertyBindingWidget(TWeakPtr<FWidgetBlueprintEditor> InEditor, UFunction* SignatureFunction, TSharedRef<IPropertyHandle> InPropertyHandle, bool bInGeneratePureBindings, bool bAllowOldBinding)
+TSharedRef<SWidget> FBlueprintWidgetCustomization::MakePropertyBindingWidget(TWeakPtr<FWidgetBlueprintEditor> InEditor, UFunction* SignatureFunction, TSharedRef<IPropertyHandle> InPropertyHandle, bool bInGeneratePureBindings, bool bAllowDetailsPanelLegacyBinding)
 {
 	if (!IModularFeatures::Get().IsModularFeatureAvailable("PropertyAccessEditor"))
 	{
@@ -606,7 +606,7 @@ TSharedRef<SWidget> FBlueprintWidgetCustomization::MakePropertyBindingWidget(TWe
 	});
 
 	Args.bGeneratePureBindings = bInGeneratePureBindings;
-	Args.bAllowNewBindings = bAllowOldBinding;
+	Args.bAllowNewBindings = bAllowDetailsPanelLegacyBinding;
 
 	IPropertyAccessEditor& PropertyAccessEditor = IModularFeatures::Get().GetModularFeature<IPropertyAccessEditor>("PropertyAccessEditor");
 	return PropertyAccessEditor.MakePropertyBindingWidget(InEditor.Pin()->GetBlueprintObj(), Args);
@@ -726,7 +726,7 @@ void FBlueprintWidgetCustomization::CreateEventCustomization( IDetailLayoutBuild
 		.MinDesiredWidth(200)
 		.MaxDesiredWidth(250)
 		[
-			MakePropertyBindingWidget(Editor.Pin(), Property->SignatureFunction, DelegatePropertyHandle, false, BlueprintObj->ArePropertyBindingsAllowed())
+			MakePropertyBindingWidget(Editor.Pin(), Property->SignatureFunction, DelegatePropertyHandle, false, true)
 		];
 }
 
diff --git a/Engine/Source/Editor/UMGEditor/Private/Customizations/UMGDetailCustomizations.h b/Engine/Source/Editor/UMGEditor/Private/Customizations/UMGDetailCustomizations.h
--- a/Engine/Source/Editor/UMGEditor/Private/Customizations/UMGDetailCustomizations.h
+++ b/Engine/Source/Editor/UMGEditor/Private/Customizations/UMGDetailCustomizations.h
@@ -33,7 +33,7 @@ public:
 	virtual void CustomizeDetails( IDetailLayoutBuilder& DetailLayout ) override;
 
 	/** Make a property binding widget */
-	static TSharedRef<SWidget> MakePropertyBindingWidget(TWeakPtr<FWidgetBlueprintEditor> InEditor, UFunction* SignatureFunction, TSharedRef<IPropertyHandle> InDelegatePropertyHandle, bool bInGeneratePureBindings, bool bAllowOldBinding);
+	static TSharedRef<SWidget> MakePropertyBindingWidget(TWeakPtr<FWidgetBlueprintEditor> InEditor, UFunction* SignatureFunction, TSharedRef<IPropertyHandle> InDelegatePropertyHandle, bool bInGeneratePureBindings, bool bAllowDetailsPanelLegacyBinding);
 
 	/** Whether the property (or its parent property, if this is an array/struct element) currently has bindings */
 	static bool HasPropertyBindings(TWeakPtr<FWidgetBlueprintEditor> InEditor, const TSharedRef<IPropertyHandle>& InPropertyHandle);
