Added "Find References" to UMG animation right-click menu.

diff --git a/Engine/Source/Editor/UMGEditor/Private/TabFactory/AnimationTabSummoner.cpp b/Engine/Source/Editor/UMGEditor/Private/TabFactory/AnimationTabSummoner.cpp
--- a/Engine/Source/Editor/UMGEditor/Private/TabFactory/AnimationTabSummoner.cpp
+++ b/Engine/Source/Editor/UMGEditor/Private/TabFactory/AnimationTabSummoner.cpp
@@ -33,6 +33,9 @@
 #include "Misc/TextFilter.h"
 #include "Kismet2/Kismet2NameValidators.h"
 #include "SPositiveActionButton.h"
+#include "GraphEditorActions.h"
+#include "BlueprintModes/WidgetBlueprintApplicationModes.h"
+#include "Engine/MemberReference.h"
 
 #define LOCTEXT_NAMESPACE "UMG"
 
@@ -681,7 +684,7 @@ private:
 		}
 	}
 
-	TSharedPtr<SWidget> OnContextMenuOpening() const
+	TSharedPtr<SWidget> OnContextMenuOpening()
 	{
 		FMenuBuilder MenuBuilder( true, CommandList.ToSharedRef() );
 
@@ -690,7 +693,7 @@ private:
 			MenuBuilder.AddMenuEntry(FGenericCommands::Get().Rename);
 			MenuBuilder.AddMenuEntry(FGenericCommands::Get().Duplicate);
 			MenuBuilder.AddMenuSeparator();
-
+			MenuBuilder.AddMenuEntry(FGraphEditorCommands::Get().FindReferences);
 			MenuBuilder.AddMenuEntry(FGenericCommands::Get().Delete);
 		}
 		MenuBuilder.EndSection();
@@ -698,6 +701,26 @@ private:
 		return MenuBuilder.MakeWidget();
 	}
 
+	void FindReferencesToSelectedAnimation()
+	{
+		if (TSharedPtr<FWidgetBlueprintEditor> WidgetEditor = BlueprintEditor.Pin())
+		{
+			TArray< TSharedPtr<FWidgetAnimationListItem> > SelectedAnimations = AnimationListView->GetSelectedItems();
+			if (SelectedAnimations.Num() == 1)
+			{
+				TSharedPtr<FWidgetAnimationListItem> SelectedAnimation = SelectedAnimations[0];
+				const FString VariableName = SelectedAnimation->Animation->GetName();
+				
+				FMemberReference MemberReference;
+				MemberReference.SetSelfMember(*VariableName);
+				const FString SearchTerm = MemberReference.GetReferenceSearchString(WidgetEditor->GetWidgetBlueprintObj()->SkeletonGeneratedClass);
+
+				WidgetEditor->SetCurrentMode(FWidgetBlueprintApplicationModes::GraphMode);
+				WidgetEditor->SummonSearchUI(true, SearchTerm);
+			}
+		}
+	}
+
 	TSharedRef<ITableRow> OnGenerateWidgetForMovieScene( TSharedPtr<FWidgetAnimationListItem> InListItem, const TSharedRef< STableViewBase >& InOwnerTableView )
 	{
 		return SNew( SWidgetAnimationListItem, InOwnerTableView, BlueprintEditor.Pin(), InListItem );
@@ -724,6 +747,12 @@ private:
 			FExecuteAction::CreateSP(this, &SUMGAnimationList::OnRenameAnimation),
 			FCanExecuteAction::CreateSP(this, &SUMGAnimationList::CanExecuteContextMenuAction)
 			);
+		
+		CommandList->MapAction(
+			FGraphEditorCommands::Get().FindReferences,
+			FExecuteAction::CreateSP(this, &SUMGAnimationList::FindReferencesToSelectedAnimation),
+			FCanExecuteAction::CreateSP(this, &SUMGAnimationList::CanExecuteContextMenuAction)
+			);
 
 		CommandList->MapAction(FUMGEditorCommands::Get().OpenAnimDrawer,
 			FExecuteAction::CreateSP(this, &SUMGAnimationList::ToggleAnimDrawer)
