UMG: Hierarchy paste operation improvements:
- Disable Paste action if the clipboard text contains no widget
- Log a message when a paste operation fails
- Ask for confirmation before deleting content of single-slot widget that we are pasting into
- Avoid renaming a pasted widget if it is the same widget that we just deleted in the previous scenario

diff --git a/Engine/Source/Editor/UMGEditor/Private/WidgetBlueprintEditor.cpp b/Engine/Source/Editor/UMGEditor/Private/WidgetBlueprintEditor.cpp
--- a/Engine/Source/Editor/UMGEditor/Private/WidgetBlueprintEditor.cpp
+++ b/Engine/Source/Editor/UMGEditor/Private/WidgetBlueprintEditor.cpp
@@ -836,6 +836,11 @@ bool FWidgetBlueprintEditor::CanPasteWidgets()
 		return false;
 	}
 
+	if (!FWidgetBlueprintEditorUtils::DoesClipboardTextContainWidget(GetWidgetBlueprintObj()))
+	{
+		return false;
+	}
+
 	if ( Widgets.Num() == 1 )
 	{
 		// Always return true here now since we want to support pasting widgets as siblings
diff --git a/Engine/Source/Editor/UMGEditor/Private/WidgetBlueprintEditorUtils.cpp b/Engine/Source/Editor/UMGEditor/Private/WidgetBlueprintEditorUtils.cpp
--- a/Engine/Source/Editor/UMGEditor/Private/WidgetBlueprintEditorUtils.cpp
+++ b/Engine/Source/Editor/UMGEditor/Private/WidgetBlueprintEditorUtils.cpp
@@ -6,6 +6,7 @@
 #include "Components/ContentWidget.h"
 #include "Engine/Texture2D.h"
 #include "Interfaces/IPluginManager.h"
+#include "UObject/GCObjectScopeGuard.h"
 #include "UObject/UObjectHash.h"
 #include "UObject/UObjectIterator.h"
 #include "Internationalization/TextPackageNamespaceUtil.h"
@@ -1507,6 +1508,7 @@ TArray<UWidget*> FWidgetBlueprintEditorUtils::DuplicateWidgets(TSharedRef<FWidge
 		DuplicatedWidgets = PasteWidgetsInternal(BlueprintEditor, BP, ExportedText, ParentWidgetRef, SlotName, FVector2D::ZeroVector, true, TransactionSuccesful);
 		if (!TransactionSuccesful)
 		{
+			BlueprintEditor->LogSimpleMessage(LOCTEXT("PasteWidgetsCancel", "Paste operation on widget cancelled."));
 			Transaction.Cancel();
 		}
 	}
@@ -1656,66 +1658,90 @@ TArray<UWidget*> FWidgetBlueprintEditorUtils::PasteWidgets(TSharedRef<FWidgetBlu
 	TArray<UWidget*> PastedWidgets = PasteWidgetsInternal(BlueprintEditor, BP, TextToImport, ParentWidgetRef, SlotName, PasteLocation, false, bTransactionSuccessful);
 	if (!bTransactionSuccessful)
 	{
+		BlueprintEditor->LogSimpleMessage(LOCTEXT("PasteWidgetsCancel", "Paste operation on widget cancelled."));
 		Transaction.Cancel();
 	}
 	return PastedWidgets;
 }
 
+bool FWidgetBlueprintEditorUtils::DisplayPasteWarningAndEarlyExit()
+{
+	const FText DeleteConfirmationPrompt = LOCTEXT("DeleteConfirmationPrompt", "Pasting in a single-slot widget will erase its content. Do you wish to proceed?");
+	const FText DeleteConfirmationTitle = LOCTEXT("DeleteConfirmationTitle", "Delete widget");
+
+	// Warn the user that this may result in data loss
+	FSuppressableWarningDialog::FSetupInfo Info(DeleteConfirmationPrompt, DeleteConfirmationTitle, TEXT("Paste_Warning"));
+	Info.ConfirmText = LOCTEXT("DeleteConfirmation_Yes", "Yes");
+	Info.CancelText = LOCTEXT("DeleteConfirmation_No", "No");
+
+	FSuppressableWarningDialog DeleteChildWidgetWarningDialog(Info);
+	return DeleteChildWidgetWarningDialog.ShowModal() == FSuppressableWarningDialog::Cancel;
+}
+
 TArray<UWidget*> FWidgetBlueprintEditorUtils::PasteWidgetsInternal(TSharedRef<FWidgetBlueprintEditor> BlueprintEditor, UWidgetBlueprint* BP, const FString& TextToImport, FWidgetReference ParentWidgetRef, FName SlotName, FVector2D PasteLocation, bool bForceSibling, bool& bTransactionSuccessful)
 {
-	// Import the nodes
-	TSet<UWidget*> PastedWidgets;
-	TMap<FName, UWidgetSlotPair*> PastedExtraSlotData;
-	FWidgetBlueprintEditorUtils::ImportWidgetsFromText(BP, TextToImport, /*out*/ PastedWidgets, /*out*/ PastedExtraSlotData);
+	// Do an intial text processing to make sure we have any widgets to paste
+	UPackage* TempPackage = nullptr;
+	FWidgetObjectTextFactory Factory = ProcessImportedText(BP, TextToImport, TempPackage);
+	TGCObjectScopeGuard<UPackage> TempPackageGCGuard(TempPackage);
+	const bool bHasPastedWidget = Factory.NewWidgetMap.Num() > 0;
 
 	// Ignore an empty set of widget paste data.
-	if ( PastedWidgets.Num() == 0 )
+	if (!bHasPastedWidget)
 	{
 		bTransactionSuccessful = false;
 		return TArray<UWidget*>();
 	}
 
-	// If we're pasting into a content widget of the same type, treat it as a sibling duplication
-	UWidget* FirstPastedWidget = *PastedWidgets.CreateIterator();
-	if (FirstPastedWidget->IsA(UContentWidget::StaticClass()) &&
-		ParentWidgetRef.IsValid() &&
-		FirstPastedWidget->GetClass() == ParentWidgetRef.GetTemplate()->GetClass())
-	{
-		UPanelWidget* TargetParentWidget = ParentWidgetRef.GetTemplate()->GetParent();
-		if (TargetParentWidget && TargetParentWidget->CanAddMoreChildren())
-		{
-			bForceSibling = true;
-		}
-	}
-
 	TArray<UWidget*> RootPasteWidgets;
-	for ( UWidget* NewWidget : PastedWidgets )
+	TMap<FName, UWidgetSlotPair*> PastedExtraSlotData;
+	TSet<UWidget*> PastedWidgets;
+
+	auto ImportWidgets = [&]()
 	{
-		// Widgets with a null parent mean that they were the root most widget of their selection set when
-		// they were copied and thus we need to paste only the root most widgets.  All their children will be added
-		// automatically.
-		if ( NewWidget->GetParent() == nullptr )
+		FWidgetBlueprintEditorUtils::ImportWidgetsFromText(BP, TextToImport, /*out*/ PastedWidgets, /*out*/ PastedExtraSlotData);
+
+		for (UWidget* NewWidget : PastedWidgets)
 		{
-			// Check to see if this widget is content of another widget holding it in a named slot.
-			bool bIsNamedSlot = false;
-			for (UWidget* ContainerWidget : PastedWidgets)
+			// Widgets with a null parent mean that they were the root most widget of their selection set when
+			// they were copied and thus we need to paste only the root most widgets.  All their children will be added
+			// automatically.
+			if (NewWidget->GetParent() == nullptr)
 			{
-				if (INamedSlotInterface* NamedSlotContainer = Cast<INamedSlotInterface>(ContainerWidget))
+				// Check to see if this widget is content of another widget holding it in a named slot.
+				bool bIsNamedSlot = false;
+				for (UWidget* ContainerWidget : PastedWidgets)
 				{
-					if (NamedSlotContainer->ContainsContent(NewWidget))
+					if (INamedSlotInterface* NamedSlotContainer = Cast<INamedSlotInterface>(ContainerWidget))
 					{
-						bIsNamedSlot = true;
-						break;
+						if (NamedSlotContainer->ContainsContent(NewWidget))
+						{
+							bIsNamedSlot = true;
+							break;
+						}
 					}
 				}
-			}
 
-			// It's a Root widget only if it's not not in a named slot.
-			if (!bIsNamedSlot)
-			{
-				RootPasteWidgets.Add(NewWidget);
+				// It's a Root widget only if it's not not in a named slot.
+				if (!bIsNamedSlot)
+				{
+					RootPasteWidgets.Add(NewWidget);
+				}
 			}
 		}
+	};
+
+	// If we're pasting into a content widget of the same type, treat it as a sibling duplication
+	UWidget* FirstPastedWidget = Factory.NewWidgetMap.CreateIterator()->Value;
+	if (FirstPastedWidget->IsA(UContentWidget::StaticClass()) &&
+		ParentWidgetRef.IsValid() &&
+		FirstPastedWidget->GetClass() == ParentWidgetRef.GetTemplate()->GetClass())
+	{
+		UPanelWidget* TargetParentWidget = ParentWidgetRef.GetTemplate()->GetParent();
+		if (TargetParentWidget && TargetParentWidget->CanAddMoreChildren())
+		{
+			bForceSibling = true;
+		}
 	}
 
 	if ( SlotName == NAME_None )
@@ -1751,15 +1777,6 @@ TArray<UWidget*> FWidgetBlueprintEditorUtils::PasteWidgetsInternal(TSharedRef<FW
 			}
 		}
 
-		// If there isn't a root widget and we're copying multiple root widgets, then we need to add a container root
-		// to hold the pasted data since multiple root widgets isn't permitted.
-		if ( !ParentWidget && RootPasteWidgets.Num() > 1 )
-		{
-			ParentWidget = BP->WidgetTree->ConstructWidget<UCanvasPanel>(UCanvasPanel::StaticClass());
-			BP->WidgetTree->Modify();
-			BP->WidgetTree->RootWidget = ParentWidget;
-		}
-
 		if ( ParentWidget )
 		{
 			// If parent widget can only have one child and that slot is already occupied, we will remove its contents so the pasted widgets can be inserted in their place
@@ -1767,8 +1784,14 @@ TArray<UWidget*> FWidgetBlueprintEditorUtils::PasteWidgetsInternal(TSharedRef<FW
 			if (!ParentWidget->CanHaveMultipleChildren() && ParentWidget->GetChildrenCount() > 0)
 			{
 				// We do not Remove child if there is nothing to paste.
-				if (RootPasteWidgets.Num() > 0)
+				if ( bHasPastedWidget )
 				{
+					if (FWidgetBlueprintEditorUtils::DisplayPasteWarningAndEarlyExit())
+					{
+						bTransactionSuccessful = false;
+						return TArray<UWidget*>();
+					}
+
 					// Delete the singular child
 					ChildWidgetToDelete = ParentWidget->GetAllChildren()[0];
 					ChildWidgetToDelete->SetFlags(RF_Transactional);
@@ -1780,6 +1803,26 @@ TArray<UWidget*> FWidgetBlueprintEditorUtils::PasteWidgetsInternal(TSharedRef<FW
 				}
 			}
 
+			if (ChildWidgetToDelete)
+			{
+				DeleteWidgets(BlueprintEditor, BP, { BlueprintEditor->GetReferenceFromTemplate(ChildWidgetToDelete) });
+			}
+		}
+
+		ImportWidgets();
+
+		// If there isn't a root widget and we're copying multiple root widgets, then we need to add a container root
+		// to hold the pasted data since multiple root widgets isn't permitted.
+		if (!ParentWidget && RootPasteWidgets.Num() > 1)
+		{
+			ParentWidget = BP->WidgetTree->ConstructWidget<UCanvasPanel>(UCanvasPanel::StaticClass());
+			BP->WidgetTree->Modify();
+			BP->WidgetTree->RootWidget = ParentWidget;
+		}
+
+		if (ParentWidget)
+		{
+
 			// A bit of a hack, but we can look at the widget's slot properties to determine if it is a canvas slot. If so, we'll try and maintain the relative positions
 			bool bShouldReproduceOffsets = true;
 			static const FName LayoutDataLabel = FName(TEXT("LayoutData"));
@@ -1858,11 +1901,6 @@ TArray<UWidget*> FWidgetBlueprintEditorUtils::PasteWidgetsInternal(TSharedRef<FW
 				}
 			}
 
-			if (ChildWidgetToDelete)
-			{
-				DeleteWidgets(BlueprintEditor, BP, { BlueprintEditor->GetReferenceFromTemplate(ChildWidgetToDelete) });
-			}
-
 			FBlueprintEditorUtils::MarkBlueprintAsStructurallyModified(BP);
 		}
 		else
@@ -1883,6 +1921,8 @@ TArray<UWidget*> FWidgetBlueprintEditorUtils::PasteWidgetsInternal(TSharedRef<FW
 	}
 	else
 	{
+		ImportWidgets();
+
 		if ( RootPasteWidgets.Num() > 1 )
 		{
 			FNotificationInfo Info(LOCTEXT("NamedSlotsOnlyHoldOneWidget", "Can't paste content, a slot can only hold one widget at the root."));
@@ -1919,12 +1959,11 @@ TArray<UWidget*> FWidgetBlueprintEditorUtils::PasteWidgetsInternal(TSharedRef<FW
 	return RootPasteWidgets;
 }
 
-void FWidgetBlueprintEditorUtils::ImportWidgetsFromText(UWidgetBlueprint* BP, const FString& TextToImport, /*out*/ TSet<UWidget*>& ImportedWidgetSet, /*out*/ TMap<FName, UWidgetSlotPair*>& PastedExtraSlotData)
+FWidgetObjectTextFactory FWidgetBlueprintEditorUtils::ProcessImportedText(UWidgetBlueprint* BP, const FString& TextToImport, /*out*/ UPackage*& TempPackage)
 {
 	// We create our own transient package here so that we can deserialize the data in isolation and ensure unreferenced
 	// objects not part of the deserialization set are unresolved.
-	UPackage* TempPackage = NewObject<UPackage>(nullptr, TEXT("/Engine/UMG/Editor/Transient"), RF_Transient);
-	TempPackage->AddToRoot();
+	TempPackage = NewObject<UPackage>(nullptr, TEXT("/Engine/UMG/Editor/Transient"), RF_Transient);
 
 	// Force the transient package to have the same namespace as the final widget blueprint package.
 	// This ensures any text properties serialized from the buffer will be keyed correctly for the target package.
@@ -1941,6 +1980,14 @@ void FWidgetBlueprintEditorUtils::ImportWidgetsFromText(UWidgetBlueprint* BP, co
 	// Turn the text buffer into objects
 	FWidgetObjectTextFactory Factory;
 	Factory.ProcessBuffer(TempPackage, RF_Transactional, TextToImport);
+	return Factory;
+}
+
+void FWidgetBlueprintEditorUtils::ImportWidgetsFromText(UWidgetBlueprint* BP, const FString& TextToImport, /*out*/ TSet<UWidget*>& ImportedWidgetSet, /*out*/ TMap<FName, UWidgetSlotPair*>& PastedExtraSlotData)
+{
+	UPackage* TempPackage = nullptr;
+	FWidgetObjectTextFactory Factory = ProcessImportedText(BP, TextToImport, TempPackage);
+	TGCObjectScopeGuard<UPackage> TempPackageGCGuard(TempPackage);
 
 	PastedExtraSlotData = Factory.MissingSlotData;
 
@@ -1995,9 +2042,6 @@ void FWidgetBlueprintEditorUtils::ImportWidgetsFromText(UWidgetBlueprint* BP, co
 			Widget->Rename(*WidgetOldName, BP->WidgetTree);
 		}
 	}
-
-	// Remove the temp package from the root now that it has served its purpose.
-	TempPackage->RemoveFromRoot();
 }
 
 void FWidgetBlueprintEditorUtils::ExportPropertiesToText(UObject* Object, TMap<FName, FString>& ExportedProperties)
@@ -2042,6 +2086,15 @@ void FWidgetBlueprintEditorUtils::ImportPropertiesFromText(UObject* Object, cons
 	}
 }
 
+bool FWidgetBlueprintEditorUtils::DoesClipboardTextContainWidget(UWidgetBlueprint* BP)
+{
+	FString TextToImport;
+	FPlatformApplicationMisc::ClipboardPaste(TextToImport);
+	UPackage* TempPackage = nullptr;
+	FWidgetObjectTextFactory Factory = ProcessImportedText(BP, TextToImport, TempPackage);
+	return Factory.NewWidgetMap.Num() > 0;
+}
+
 bool FWidgetBlueprintEditorUtils::IsBindWidgetProperty(const FProperty* InProperty)
 {
 	bool bIsOptional;
diff --git a/Engine/Source/Editor/UMGEditor/Public/WidgetBlueprintEditorUtils.h b/Engine/Source/Editor/UMGEditor/Public/WidgetBlueprintEditorUtils.h
--- a/Engine/Source/Editor/UMGEditor/Public/WidgetBlueprintEditorUtils.h
+++ b/Engine/Source/Editor/UMGEditor/Public/WidgetBlueprintEditorUtils.h
@@ -15,6 +15,7 @@ class UWidgetSlotPair;
 class UWidgetTree;
 class SWindow;
 class UWidgetEditingProjectSettings;
+class FWidgetObjectTextFactory;
 
 //////////////////////////////////////////////////////////////////////////
 // FWidgetBlueprintEditorUtils
@@ -83,6 +84,8 @@ public:
 	/** Attempts to import any property in the map and apply it to a property with the same name on the object. */
 	static void ImportPropertiesFromText(UObject* Object, const TMap<FName, FString>& ExportedProperties);
 
+	static bool DoesClipboardTextContainWidget(UWidgetBlueprint* BP);
+
 	static TScriptInterface<INamedSlotInterface> FindNamedSlotHostForContent(UWidget* WidgetTemplate, UWidgetTree* WidgetTree);
 
 	static UWidget* FindNamedSlotHostWidgetForContent(UWidget* WidgetTemplate, UWidgetTree* WidgetTree);
@@ -127,6 +130,8 @@ private:
 
 	static TArray<UWidget*> PasteWidgetsInternal(TSharedRef<FWidgetBlueprintEditor> BlueprintEditor, UWidgetBlueprint* BP, const FString& TextToImport, FWidgetReference ParentWidget, FName SlotName, FVector2D PasteLocation, bool bForceSibling, bool& TransactionSuccesful);
 
+	static bool DisplayPasteWarningAndEarlyExit();
+
 	static void ExecuteOpenSelectedWidgetsForEdit( TSet<FWidgetReference> SelectedWidgets );
 
 	static bool FindAndRemoveNamedSlotContent(UWidget* WidgetTemplate, UWidgetTree* WidgetTree);
@@ -160,4 +165,6 @@ private:
 	static TOptional<FWidgetThumbnailProperties> DrawSWidgetInRenderTargetInternal(UUserWidget* WidgetInstance, FRenderTarget* RenderTarget2D, UTextureRenderTarget2D* TextureRenderTarget,FVector2D ThumbnailSize, bool bIsForThumbnail, TOptional<FVector2D> ThumbnailCustomSize, EThumbnailPreviewSizeMode ThumbnailSizeMode);
 	
 	static bool IsDesiredFocusWiget(TSharedRef<FWidgetBlueprintEditor> BlueprintEditor, UWidget* Widget);
+
+	static FWidgetObjectTextFactory ProcessImportedText(UWidgetBlueprint* BP, const FString& TextToImport, /*out*/ UPackage*& TempPackage);
 };
