UMG: PR #11248: filter classes in the parent widget picker to blueprint types only.

diff --git a/Engine/Source/Editor/UMGEditor/Private/WidgetBlueprintFactory.cpp b/Engine/Source/Editor/UMGEditor/Private/WidgetBlueprintFactory.cpp
--- a/Engine/Source/Editor/UMGEditor/Private/WidgetBlueprintFactory.cpp
+++ b/Engine/Source/Editor/UMGEditor/Private/WidgetBlueprintFactory.cpp
@@ -67,6 +67,7 @@ bool UWidgetBlueprintFactory::ConfigureProperties()
 		Options.Mode = EClassViewerMode::ClassPicker;
 		Options.bShowNoneOption = false;
 		Options.bExpandAllNodes = true;
+		Options.bIsBlueprintBaseOnly = true;
 		Options.bShowDefaultClasses = GetDefault<UUMGEditorProjectSettings>()->bUseUserWidgetParentDefaultClassViewerSelector;
 		Options.bShowClassesViewer = GetDefault<UUMGEditorProjectSettings>()->bUseUserWidgetParentClassViewerSelector;
 
