MVVM: support widget drag-drop from hierarchy view to binding panel.

diff --git a/Engine/Plugins/Runtime/ModelViewViewModel/Source/ModelViewViewModelEditor/Private/Widgets/SMVVMViewBindingPanel.cpp b/Engine/Plugins/Runtime/ModelViewViewModel/Source/ModelViewViewModelEditor/Private/Widgets/SMVVMViewBindingPanel.cpp
--- a/Engine/Plugins/Runtime/ModelViewViewModel/Source/ModelViewViewModelEditor/Private/Widgets/SMVVMViewBindingPanel.cpp
+++ b/Engine/Plugins/Runtime/ModelViewViewModel/Source/ModelViewViewModelEditor/Private/Widgets/SMVVMViewBindingPanel.cpp
@@ -11,8 +11,8 @@
 #include "MVVMWidgetBlueprintExtension_View.h"
 #include "WidgetBlueprintEditor.h"
 #include "WidgetBlueprintToolMenuContext.h"
-#include "WidgetReference.h"
 
+#include "Hierarchy/HierarchyWidgetDragDropOp.h"
 #include "IStructureDetailsView.h"
 #include "PropertyEditorModule.h"
 #include "StatusBarSubsystem.h"
@@ -225,7 +225,7 @@ void SBindingsPanel::Tick(const FGeometry& AllottedGeometry, const double InCurr
 	Super::Tick(AllottedGeometry, InCurrentTime, InDeltaTime);
 }
 
-void SBindingsPanel::AddDefaultBinding()
+void SBindingsPanel::AddBindingToWidgetList(const TSet<FWidgetReference>& WidgetsToAddBinding)
 {
 	if (!CanAddBinding())
 	{
@@ -241,7 +241,7 @@ void SBindingsPanel::AddDefaultBinding()
 
 			if (AddBindingMode == EAddBindingMode::Selected)
 			{
-				for (const FWidgetReference& WidgetReference : BlueprintEditor->GetSelectedWidgets())
+				for (const FWidgetReference& WidgetReference : WidgetsToAddBinding)
 				{
 					if (WidgetReference.IsValid() && WidgetReference.GetTemplate())
 					{
@@ -277,6 +277,14 @@ void SBindingsPanel::AddDefaultBinding()
 	}
 }
 
+void SBindingsPanel::AddDefaultBinding()
+{
+	if (TSharedPtr<FWidgetBlueprintEditor> BlueprintEditor = WeakBlueprintEditor.Pin())
+	{
+		AddBindingToWidgetList(BlueprintEditor->GetSelectedWidgets());
+	}
+}
+
 bool SBindingsPanel::CanAddBinding() const
 {
 	UMVVMWidgetBlueprintExtension_View* MVVMExtensionPtr = MVVMExtension.Get();
@@ -480,6 +488,65 @@ void SBindingsPanel::RegisterSettingsMenu()
 		}));
 }
 
+FReply SBindingsPanel::OnDragOver(const FGeometry& MyGeometry, const FDragDropEvent& DragDropEvent)
+{
+	if (TSharedPtr<FDragDropOperation> DragDropOp = DragDropEvent.GetOperation())
+	{
+		if (TSharedPtr<FHierarchyWidgetDragDropOp> HierarchyDragDropOp = DragDropEvent.GetOperationAs<FHierarchyWidgetDragDropOp>())
+		{
+			if (HierarchyDragDropOp->HasOriginatedFrom(WeakBlueprintEditor.Pin()))
+			{
+				if (CanAddBinding())
+				{
+					HierarchyDragDropOp->CurrentIconBrush = FAppStyle::GetBrush(TEXT("Graph.ConnectorFeedback.OK"));
+					return FReply::Handled();
+				}
+			}
+			HierarchyDragDropOp->CurrentIconBrush = FAppStyle::GetBrush(TEXT("Graph.ConnectorFeedback.Error"));
+		}
+		else
+		{
+			TSharedPtr<FDecoratedDragDropOp> DecoratedDragDropOp = nullptr;
+			if (DragDropOp->IsOfType<FDecoratedDragDropOp>())
+			{
+				DecoratedDragDropOp = StaticCastSharedPtr<FDecoratedDragDropOp>(DragDropOp);
+				DecoratedDragDropOp->ResetToDefaultToolTip();
+				DecoratedDragDropOp->CurrentIconBrush = FAppStyle::GetBrush(TEXT("Graph.ConnectorFeedback.Error"));
+			}
+		}
+	}
+	return FReply::Unhandled();
+}
+
+void SBindingsPanel::OnDragLeave(const FDragDropEvent& DragDropEvent)
+{
+	if (TSharedPtr<FDecoratedDragDropOp> DecoratedDragDropOp = DragDropEvent.GetOperationAs<FDecoratedDragDropOp>())
+	{
+		DecoratedDragDropOp->ResetToDefaultToolTip();
+	}
+}
+
+FReply SBindingsPanel::OnDrop(const FGeometry& MyGeometry, const FDragDropEvent& DragDropEvent)
+{
+	if (TSharedPtr<FHierarchyWidgetDragDropOp> HierarchyDragDropOp = DragDropEvent.GetOperationAs<FHierarchyWidgetDragDropOp>())
+	{
+		if (HierarchyDragDropOp->HasOriginatedFrom(WeakBlueprintEditor.Pin()))
+		{
+			if (CanAddBinding())
+			{
+				TSet<FWidgetReference> DraggedWidgetSet;
+				for (const FWidgetReference& WidgetRef : HierarchyDragDropOp->GetWidgetReferences())
+				{
+					DraggedWidgetSet.Add(WidgetRef);
+				}
+
+				AddBindingToWidgetList(DraggedWidgetSet);
+				return FReply::Handled();
+			}
+		}
+	}
+	return FReply::Unhandled();
+}
 
 TSharedRef<SWidget> SBindingsPanel::GenerateEditViewWidget()
 {
diff --git a/Engine/Plugins/Runtime/ModelViewViewModel/Source/ModelViewViewModelEditor/Private/Widgets/SMVVMViewBindingPanel.h b/Engine/Plugins/Runtime/ModelViewViewModel/Source/ModelViewViewModelEditor/Private/Widgets/SMVVMViewBindingPanel.h
--- a/Engine/Plugins/Runtime/ModelViewViewModel/Source/ModelViewViewModelEditor/Private/Widgets/SMVVMViewBindingPanel.h
+++ b/Engine/Plugins/Runtime/ModelViewViewModel/Source/ModelViewViewModelEditor/Private/Widgets/SMVVMViewBindingPanel.h
@@ -2,7 +2,7 @@
 
 #pragma once
 
-
+#include "WidgetReference.h"
 #include "Widgets/SCompoundWidget.h"
 
 
@@ -37,6 +37,9 @@ public:
 	virtual FReply OnKeyDown(const FGeometry& MyGeometry, const FKeyEvent& InKeyEvent) override;
 	virtual bool SupportsKeyboardFocus() const override;
 	virtual void Tick(const FGeometry& AllottedGeometry, const double InCurrentTime, const float InDeltaTime) override;
+	virtual FReply OnDragOver(const FGeometry& MyGeometry, const FDragDropEvent& DragDropEvent) override;
+	virtual void OnDragLeave(const FDragDropEvent& DragDropEvent) override;
+	virtual FReply OnDrop(const FGeometry& MyGeometry, const FDragDropEvent& DragDropEvent) override;
 	//~ End SWidget Interface
 
 	void OnBindingListSelectionChanged(TConstArrayView<FMVVMBlueprintViewBinding*> Selection);
@@ -57,6 +60,7 @@ private:
 
 	TSharedRef<SWidget> GenerateEditViewWidget();
 
+	void AddBindingToWidgetList(const TSet<FWidgetReference>& WidgetsToAddBinding);
 	void AddDefaultBinding();
 	bool CanAddBinding() const;
 	FText GetAddBindingText() const;
diff --git a/Engine/Source/Editor/UMGEditor/Private/Hierarchy/HierarchyWidgetDragDropOp.cpp b/Engine/Source/Editor/UMGEditor/Private/Hierarchy/HierarchyWidgetDragDropOp.cpp
new file mode 100644
--- /dev/null
+++ b/Engine/Source/Editor/UMGEditor/Private/Hierarchy/HierarchyWidgetDragDropOp.cpp
@@ -0,0 +1,56 @@
+// Copyright Epic Games, Inc. All Rights Reserved.
+
+#include "Hierarchy/HierarchyWidgetDragDropOp.h"
+
+#include "Blueprint/WidgetTree.h"
+#include "WidgetBlueprint.h"
+#include "WidgetBlueprintEditor.h"
+
+#define LOCTEXT_NAMESPACE "UMG"
+
+TSharedRef<FHierarchyWidgetDragDropOp> FHierarchyWidgetDragDropOp::New(UWidgetBlueprint* Blueprint, const TArray<FWidgetReference>& InWidgets)
+{
+	check(InWidgets.Num() > 0);
+
+	TSharedRef<FHierarchyWidgetDragDropOp> Operation = MakeShareable(new FHierarchyWidgetDragDropOp());
+
+	// Set the display text and the transaction name based on whether we're dragging a single or multiple widgets
+	if (InWidgets.Num() == 1)
+	{
+		Operation->CurrentHoverText = Operation->DefaultHoverText = InWidgets[0].GetTemplate()->GetLabelText();
+	}
+	else
+	{
+		Operation->CurrentHoverText = Operation->DefaultHoverText = LOCTEXT("Designer_DragMultipleWidgets", "Multiple Widgets");
+	}
+
+	for (const auto& Widget : InWidgets)
+	{
+		Operation->WidgetReferences.Add(Widget);
+	}
+
+	Operation->Construct();
+	
+	Blueprint->WidgetTree->SetFlags(RF_Transactional);
+	Blueprint->WidgetTree->Modify();
+
+	return Operation;
+}
+
+const TArrayView<const FWidgetReference> FHierarchyWidgetDragDropOp::GetWidgetReferences() const
+{
+	return MakeArrayView(WidgetReferences);
+}
+
+bool FHierarchyWidgetDragDropOp::HasOriginatedFrom(const TSharedPtr<FWidgetBlueprintEditor>& BlueprintEditor) const
+{
+	for (const FWidgetReference& Widget : WidgetReferences)
+	{
+		if (Widget.GetWidgetEditor() != BlueprintEditor)
+		{
+			return false;
+		}
+	}
+	return true;
+}
+#undef LOCTEXT_NAMESPACE
diff --git a/Engine/Source/Editor/UMGEditor/Private/Hierarchy/SHierarchyViewItem.cpp b/Engine/Source/Editor/UMGEditor/Private/Hierarchy/SHierarchyViewItem.cpp
--- a/Engine/Source/Editor/UMGEditor/Private/Hierarchy/SHierarchyViewItem.cpp
+++ b/Engine/Source/Editor/UMGEditor/Private/Hierarchy/SHierarchyViewItem.cpp
@@ -25,6 +25,7 @@
 #include "DragAndDrop/ClassDragDropOp.h"
 #include "DragDrop/WidgetTemplateDragDropOp.h"
 #include "DragDrop/SelectedWidgetDragDropOp.h"
+#include "Hierarchy/HierarchyWidgetDragDropOp.h"
 
 #include "WidgetTemplate.h"
 
@@ -42,27 +43,16 @@
 /**
 *
 */
-class FHierarchyWidgetDragDropOp : public FDecoratedDragDropOp
+class FHierarchyWidgetDragDropOpImpl : public FHierarchyWidgetDragDropOp
 {
 public:
-	DRAG_DROP_OPERATOR_TYPE(FHierarchyWidgetDragDropOp, FDecoratedDragDropOp)
+	DRAG_DROP_OPERATOR_TYPE(FHierarchyWidgetDragDropOpImpl, FHierarchyWidgetDragDropOp)
 
-	virtual ~FHierarchyWidgetDragDropOp();
+	FHierarchyWidgetDragDropOpImpl(FHierarchyWidgetDragDropOp& HierarchyWidgetDragDropOp);
+	virtual ~FHierarchyWidgetDragDropOpImpl();
 
 	virtual void OnDrop(bool bDropWasHandled, const FPointerEvent& MouseEvent) override;
 
-	bool HasOriginatedFrom(const TSharedPtr<FWidgetBlueprintEditor>& BlueprintEditor)
-	{
-		for (const FItem& Item : DraggedWidgets)
-		{
-			if (Item.Widget.GetWidgetEditor() != BlueprintEditor)
-			{
-				return false;
-			}
-		}
-		return true;
-	}
-
 	struct FItem
 	{
 		/** The slot properties for the old slot the widget was in, is used to attempt to reapply the same layout information */
@@ -81,14 +71,20 @@ public:
 	FScopedTransaction* Transaction;
 
 	/** Constructs a new drag/drop operation */
-	static TSharedRef<FHierarchyWidgetDragDropOp> New(UWidgetBlueprint* Blueprint, const TArray<FWidgetReference>& InWidgets);
+	static TSharedRef<FHierarchyWidgetDragDropOpImpl> New(UWidgetBlueprint* Blueprint, const TArray<FWidgetReference>& InWidgets);
 };
 
-TSharedRef<FHierarchyWidgetDragDropOp> FHierarchyWidgetDragDropOp::New(UWidgetBlueprint* Blueprint, const TArray<FWidgetReference>& InWidgets)
+FHierarchyWidgetDragDropOpImpl::FHierarchyWidgetDragDropOpImpl(FHierarchyWidgetDragDropOp& HierarchyWidgetDragDropOp) 
+	: FHierarchyWidgetDragDropOp(HierarchyWidgetDragDropOp)
+{
+}
+
+TSharedRef<FHierarchyWidgetDragDropOpImpl> FHierarchyWidgetDragDropOpImpl::New(UWidgetBlueprint* Blueprint, const TArray<FWidgetReference>& InWidgets)
 {
 	check(InWidgets.Num() > 0);
 
-	TSharedRef<FHierarchyWidgetDragDropOp> Operation = MakeShareable(new FHierarchyWidgetDragDropOp());
+	TSharedRef<FHierarchyWidgetDragDropOp> HierarchyWidgetDragDropOp = FHierarchyWidgetDragDropOp::New(Blueprint, InWidgets);
+	TSharedRef<FHierarchyWidgetDragDropOpImpl> Operation = MakeShareable(new FHierarchyWidgetDragDropOpImpl(*HierarchyWidgetDragDropOp));
 
 	// Set the display text and the transaction name based on whether we're dragging a single or multiple widgets
 	if (InWidgets.Num() == 1)
@@ -131,12 +127,12 @@ TSharedRef<FHierarchyWidgetDragDropOp> FHierarchyWidgetDragDropOp::New(UWidgetBl
 	return Operation;
 }
 
-FHierarchyWidgetDragDropOp::~FHierarchyWidgetDragDropOp()
+FHierarchyWidgetDragDropOpImpl::~FHierarchyWidgetDragDropOpImpl()
 {
 	delete Transaction;
 }
 
-void FHierarchyWidgetDragDropOp::OnDrop(bool bDropWasHandled, const FPointerEvent& MouseEvent)
+void FHierarchyWidgetDragDropOpImpl::OnDrop(bool bDropWasHandled, const FPointerEvent& MouseEvent)
 {
 	if ( !bDropWasHandled )
 	{
@@ -150,7 +146,7 @@ TOptional<EItemDropZone> ProcessHierarchyDragDrop(const FDragDropEvent& DragDrop
 {
 	UWidget* TargetTemplate = TargetItem.GetTemplate();
 
-	if (TSharedPtr<FHierarchyWidgetDragDropOp> HierarchyDragDropOp = DragDropEvent.GetOperationAs<FHierarchyWidgetDragDropOp>())
+	if (TSharedPtr<FHierarchyWidgetDragDropOpImpl> HierarchyDragDropOp = DragDropEvent.GetOperationAs<FHierarchyWidgetDragDropOpImpl>())
 	{
 		if (!HierarchyDragDropOp->HasOriginatedFrom(BlueprintEditor))
 		{
@@ -194,7 +190,7 @@ TOptional<EItemDropZone> ProcessHierarchyDragDrop(const FDragDropEvent& DragDrop
 
 	// Is this a drag/drop op to create a new widget in the tree?
 	TSharedPtr<FDragDropOperation> DragDropOp = DragDropEvent.GetOperation();
-	if (DragDropOp.IsValid() && !DragDropOp->IsOfType<FHierarchyWidgetDragDropOp>())
+	if (DragDropOp.IsValid() && !DragDropOp->IsOfType<FHierarchyWidgetDragDropOpImpl>())
 	{
 		TSharedPtr<FDecoratedDragDropOp> DecoratedDragDropOp = nullptr;
 		if (DragDropOp->IsOfType<FDecoratedDragDropOp>())
@@ -295,7 +291,7 @@ TOptional<EItemDropZone> ProcessHierarchyDragDrop(const FDragDropEvent& DragDrop
 		return TOptional<EItemDropZone>();
 	}
 
-	TSharedPtr<FHierarchyWidgetDragDropOp> HierarchyDragDropOp = DragDropEvent.GetOperationAs<FHierarchyWidgetDragDropOp>();
+	TSharedPtr<FHierarchyWidgetDragDropOpImpl> HierarchyDragDropOp = DragDropEvent.GetOperationAs<FHierarchyWidgetDragDropOpImpl>();
 	if ( HierarchyDragDropOp.IsValid() )
 	{
 		HierarchyDragDropOp->ResetToDefaultToolTip();
@@ -304,7 +300,7 @@ TOptional<EItemDropZone> ProcessHierarchyDragDrop(const FDragDropEvent& DragDrop
 		// the null case and we should be adding it as the root widget.
 		if ( TargetItem.IsValid() )
 		{
-			const bool bIsDraggedObject = HierarchyDragDropOp->DraggedWidgets.ContainsByPredicate([TargetItem](const FHierarchyWidgetDragDropOp::FItem& DraggedItem)
+			const bool bIsDraggedObject = HierarchyDragDropOp->DraggedWidgets.ContainsByPredicate([TargetItem](const FHierarchyWidgetDragDropOpImpl::FItem& DraggedItem)
 			{
 				return DraggedItem.Widget == TargetItem;
 			});
@@ -534,7 +530,7 @@ FReply FHierarchyModel::HandleDragDetected(const FGeometry& MyGeometry, const FP
 
 		if (DraggedItems.Num() > 0)
 		{
-			return FReply::Handled().BeginDragDrop(FHierarchyWidgetDragDropOp::New(BlueprintEditor.Pin()->GetWidgetBlueprintObj(), DraggedItems));
+			return FReply::Handled().BeginDragDrop(FHierarchyWidgetDragDropOpImpl::New(BlueprintEditor.Pin()->GetWidgetBlueprintObj(), DraggedItems));
 		}
 	}
 
@@ -909,7 +905,7 @@ TOptional<EItemDropZone> FNamedSlotModelBase::HandleCanAcceptDrop(const FDragDro
 	UWidgetBlueprint* Blueprint = BlueprintEditor.Pin()->GetWidgetBlueprintObj();
 
 	TSharedPtr<FDragDropOperation> DragDropOp = DragDropEvent.GetOperation();
-	if (DragDropOp.IsValid() && !DragDropOp->IsOfType<FHierarchyWidgetDragDropOp>())
+	if (DragDropOp.IsValid() && !DragDropOp->IsOfType<FHierarchyWidgetDragDropOpImpl>())
 	{
 		TSharedPtr<FDecoratedDragDropOp> DecoratedDragDropOp = nullptr;
 		if (DragDropOp->IsOfType<FDecoratedDragDropOp>())
@@ -943,7 +939,7 @@ TOptional<EItemDropZone> FNamedSlotModelBase::HandleCanAcceptDrop(const FDragDro
 		}
 	}
 
-	TSharedPtr<FHierarchyWidgetDragDropOp> HierarchyDragDropOp = DragDropEvent.GetOperationAs<FHierarchyWidgetDragDropOp>();
+	TSharedPtr<FHierarchyWidgetDragDropOpImpl> HierarchyDragDropOp = DragDropEvent.GetOperationAs<FHierarchyWidgetDragDropOpImpl>();
 	if (HierarchyDragDropOp.IsValid() && HierarchyDragDropOp->DraggedWidgets.Num() == 1)
 	{
 		HierarchyDragDropOp->ResetToDefaultToolTip();
@@ -997,7 +993,7 @@ FReply FNamedSlotModelBase::HandleAcceptDrop(FDragDropEvent const& DragDropEvent
 	}
 
 	TSharedPtr<FDragDropOperation> DragDropOp = DragDropEvent.GetOperation();
-	if (DragDropOp.IsValid() && !DragDropOp->IsOfType<FHierarchyWidgetDragDropOp>())
+	if (DragDropOp.IsValid() && !DragDropOp->IsOfType<FHierarchyWidgetDragDropOpImpl>())
 	{
 		FScopedTransaction Transaction(LOCTEXT("AddWidgetFromTemplate", "Add Widget"));
 
@@ -1016,7 +1012,7 @@ FReply FNamedSlotModelBase::HandleAcceptDrop(FDragDropEvent const& DragDropEvent
 		}
 	}
 
-	TSharedPtr<FHierarchyWidgetDragDropOp> HierarchyDragDropOp = DragDropEvent.GetOperationAs<FHierarchyWidgetDragDropOp>();
+	TSharedPtr<FHierarchyWidgetDragDropOpImpl> HierarchyDragDropOp = DragDropEvent.GetOperationAs<FHierarchyWidgetDragDropOpImpl>();
 	if (HierarchyDragDropOp.IsValid() && HierarchyDragDropOp->DraggedWidgets.Num() == 1)
 	{
 		UWidgetBlueprint* Blueprint = BlueprintEditor.Pin()->GetWidgetBlueprintObj();
diff --git a/Engine/Source/Editor/UMGEditor/Private/Hierarchy/SHierarchyViewItem.h b/Engine/Source/Editor/UMGEditor/Private/Hierarchy/SHierarchyViewItem.h
--- a/Engine/Source/Editor/UMGEditor/Private/Hierarchy/SHierarchyViewItem.h
+++ b/Engine/Source/Editor/UMGEditor/Private/Hierarchy/SHierarchyViewItem.h
@@ -13,8 +13,6 @@
 #include "WidgetReference.h"
 #include "WidgetBlueprintEditor.h"
 
-class FWidgetBlueprintEditor;
-
 class FHierarchyModel : public TSharedFromThis < FHierarchyModel >
 {
 public:
diff --git a/Engine/Source/Editor/UMGEditor/Public/Hierarchy/HierarchyWidgetDragDropOp.h b/Engine/Source/Editor/UMGEditor/Public/Hierarchy/HierarchyWidgetDragDropOp.h
new file mode 100644
--- /dev/null
+++ b/Engine/Source/Editor/UMGEditor/Public/Hierarchy/HierarchyWidgetDragDropOp.h
@@ -0,0 +1,29 @@
+// Copyright Epic Games, Inc. All Rights Reserved.
+
+#pragma once
+
+#include "CoreMinimal.h"
+#include "DragAndDrop/DecoratedDragDropOp.h"
+#include "WidgetReference.h"
+
+class FWidgetBlueprintEditor;
+class UWidgetBlueprint;
+
+/**
+* This drag drop operation allows widgets from the hierarchy rows to be dragged and dropped around the editor.
+*/
+class FHierarchyWidgetDragDropOp : public FDecoratedDragDropOp
+{
+public:
+	DRAG_DROP_OPERATOR_TYPE(FHierarchyWidgetDragDropOp, FDecoratedDragDropOp)
+
+	bool UMGEDITOR_API HasOriginatedFrom(const TSharedPtr<FWidgetBlueprintEditor>& BlueprintEditor) const;
+	
+	const TArrayView<const FWidgetReference> UMGEDITOR_API GetWidgetReferences() const;
+
+	/** Constructs a new drag/drop operation */
+	static TSharedRef<FHierarchyWidgetDragDropOp> New(UWidgetBlueprint* Blueprint, const TArray<FWidgetReference>& InWidgets);
+
+private:
+	TArray<FWidgetReference> WidgetReferences;
+};
\ No newline at end of file
