UMG Designer Truncation Fixes step 1

diff --git a/Engine/Source/Editor/UMGEditor/Private/BlueprintModes/WidgetDesignerApplicationMode.cpp b/Engine/Source/Editor/UMGEditor/Private/BlueprintModes/WidgetDesignerApplicationMode.cpp
--- a/Engine/Source/Editor/UMGEditor/Private/BlueprintModes/WidgetDesignerApplicationMode.cpp
+++ b/Engine/Source/Editor/UMGEditor/Private/BlueprintModes/WidgetDesignerApplicationMode.cpp
@@ -41,7 +41,7 @@ FWidgetDesignerApplicationMode::FWidgetDesignerApplicationMode(TSharedPtr<FWidge
 	FDisplayMetrics DisplayMetrics;
 	FSlateApplication::Get().GetDisplayMetrics(DisplayMetrics);
 
-	const float DPIScale = FPlatformApplicationMisc::GetDPIScaleFactorAtPoint(DisplayMetrics.PrimaryDisplayWorkAreaRect.Left, DisplayMetrics.PrimaryDisplayWorkAreaRect.Top);
+	const double DPIScale = FPlatformApplicationMisc::GetDPIScaleFactorAtPoint(static_cast<float>(DisplayMetrics.PrimaryDisplayWorkAreaRect.Left), static_cast<float>(DisplayMetrics.PrimaryDisplayWorkAreaRect.Top));
 
 	const float CenterScale = 0.4f;
 	const FVector2D DisplaySize(
diff --git a/Engine/Source/Editor/UMGEditor/Private/Designer/SDesignSurface.cpp b/Engine/Source/Editor/UMGEditor/Private/Designer/SDesignSurface.cpp
--- a/Engine/Source/Editor/UMGEditor/Private/Designer/SDesignSurface.cpp
+++ b/Engine/Source/Editor/UMGEditor/Private/Designer/SDesignSurface.cpp
@@ -263,12 +263,12 @@ FReply SDesignSurface::OnMouseMove(const FGeometry& MyGeometry, const FPointerEv
 		const bool bShouldZoom = bIsRightMouseButtonDown && (bIsLeftMouseButtonDown || bIsMiddleMouseButtonDown || ModifierKeysState.IsAltDown() || FSlateApplication::Get().IsUsingTrackpad());
 		if ( bShouldZoom )
 		{
-			const float MouseZoomScaling = 0.04f;
+			const double MouseZoomScaling = 0.04f;
 			FReply ReplyState = FReply::Handled();
 
 			TotalMouseDelta += CursorDelta.X + CursorDelta.Y;
 
-			const int32 ZoomLevelDelta = FMath::RoundToInt(TotalMouseDelta * MouseZoomScaling);
+			const int32 ZoomLevelDelta = FMath::RoundToInt32(TotalMouseDelta * MouseZoomScaling);
 
 			// Get rid of mouse movement that's been 'used up' by zooming
 			if (ZoomLevelDelta != 0)
@@ -313,10 +313,10 @@ FReply SDesignSurface::OnMouseWheel(const FGeometry& MyGeometry, const FPointerE
 FReply SDesignSurface::OnTouchGesture(const FGeometry& MyGeometry, const FPointerEvent& GestureEvent)
 {
 	const EGestureEvent GestureType = GestureEvent.GetGestureType();
-	const FVector2D& GestureDelta = GestureEvent.GetGestureDelta();
+	const FVector2D GestureDelta = GestureEvent.GetGestureDelta();
 	if ( GestureType == EGestureEvent::Magnify )
 	{
-		TotalGestureMagnify += GestureDelta.X;
+		TotalGestureMagnify += static_cast<float>(GestureDelta.X);
 		if ( FMath::Abs(TotalGestureMagnify) > 0.07f )
 		{
 			// We want to zoom into this point; i.e. keep it the same fraction offset into the panel
@@ -377,7 +377,7 @@ void SDesignSurface::ChangeZoomLevel(int32 ZoomLevelDelta, const FVector2D& Widg
 		// If they are already zoomed in past 1:1, user may zoom freely
 		( ZoomLevel > DefaultZoomLevel );
 
-	const float OldZoomLevel = ZoomLevel;
+	const int32 OldZoomLevel = ZoomLevel;
 
 	if ( bAllowFullZoomRange )
 	{
@@ -570,7 +570,7 @@ void SDesignSurface::PaintBackgroundAsLines(const FSlateBrush* BackgroundImage,
 
 	const float GridCellSize = NominalGridSize * ZoomFactor * Inflation;
 
-	FVector2D LocalGridOrigin = AllottedGeometry.AbsoluteToLocal(GridOrigin);
+	FVector2f LocalGridOrigin = AllottedGeometry.AbsoluteToLocal(GridOrigin);
 
 	float ImageOffsetX = LocalGridOrigin.X - ((GridCellSize*RulePeriod) * FMath::Max(FMath::CeilToInt(LocalGridOrigin.X / (GridCellSize*RulePeriod)), 0));
 	float ImageOffsetY = LocalGridOrigin.Y - ((GridCellSize*RulePeriod) * FMath::Max(FMath::CeilToInt(LocalGridOrigin.Y / (GridCellSize*RulePeriod)), 0));
diff --git a/Engine/Source/Editor/UMGEditor/Private/Designer/SDesignerView.cpp b/Engine/Source/Editor/UMGEditor/Private/Designer/SDesignerView.cpp
--- a/Engine/Source/Editor/UMGEditor/Private/Designer/SDesignerView.cpp
+++ b/Engine/Source/Editor/UMGEditor/Private/Designer/SDesignerView.cpp
@@ -989,7 +989,7 @@ FOptionalSize SDesignerView::GetPreviewAreaWidth() const
 	TTuple<FVector2D, FVector2D> AreaAndSize = FWidgetBlueprintEditorUtils::GetWidgetPreviewAreaAndSize(GetDefaultWidget(), CachedPreviewDesiredSize, FVector2D(PreviewWidth, PreviewHeight), GetDefaultWidget()->DesignSizeMode, TOptional<FVector2D>());
 	FVector2D Area = AreaAndSize.Get<0>();
 
-	return Area.X;
+	return static_cast<float>(Area.X);
 }
 
 FOptionalSize SDesignerView::GetPreviewAreaHeight() const
@@ -997,7 +997,7 @@ FOptionalSize SDesignerView::GetPreviewAreaHeight() const
 	TTuple<FVector2D, FVector2D> AreaAndSize = FWidgetBlueprintEditorUtils::GetWidgetPreviewAreaAndSize(GetDefaultWidget(), CachedPreviewDesiredSize, FVector2D(PreviewWidth, PreviewHeight), GetDefaultWidget()->DesignSizeMode, TOptional<FVector2D>());
 	FVector2D Area = AreaAndSize.Get<0>();;
 
-	return Area.Y;
+	return static_cast<float>(Area.Y);
 }
 
 FOptionalSize SDesignerView::GetPreviewSizeWidth() const
@@ -1005,7 +1005,7 @@ FOptionalSize SDesignerView::GetPreviewSizeWidth() const
 	TTuple<FVector2D, FVector2D> AreaAndSize = FWidgetBlueprintEditorUtils::GetWidgetPreviewAreaAndSize(GetDefaultWidget(), CachedPreviewDesiredSize, FVector2D(PreviewWidth, PreviewHeight), GetDefaultWidget()->DesignSizeMode, TOptional<FVector2D>());
 	FVector2D Size = AreaAndSize.Get<1>();
 
-	return Size.X;
+	return static_cast<float>(Size.X);
 }
 
 FOptionalSize SDesignerView::GetPreviewSizeHeight() const
@@ -1013,7 +1013,7 @@ FOptionalSize SDesignerView::GetPreviewSizeHeight() const
 	TTuple<FVector2D, FVector2D> AreaAndSize = FWidgetBlueprintEditorUtils::GetWidgetPreviewAreaAndSize(GetDefaultWidget(), CachedPreviewDesiredSize, FVector2D(PreviewWidth, PreviewHeight), GetDefaultWidget()->DesignSizeMode, TOptional<FVector2D>());
 	FVector2D Size = AreaAndSize.Get<1>();
 
-	return Size.Y;
+	return static_cast<float>(Size.Y);
 }
 
 void SDesignerView::BeginResizingArea()
@@ -2165,19 +2165,19 @@ void SDesignerView::DrawSafeZone(const FOnPaintHandlerParams& PaintArgs)
 		FGeometry PreviewGeometry = PreviewAreaConstraint->GetTickSpaceGeometry();
 		PreviewGeometry.AppendTransform(FSlateLayoutTransform(Inverse(PaintArgs.Args.GetWindowToDesktopTransform())));
 		
-		const float Width = PreviewWidth;
-		const float Height = PreviewHeight;
+		const float Width = static_cast<float>(PreviewWidth);
+		const float Height = static_cast<float>(PreviewHeight);
 		if (PreviewOverrideName.IsEmpty())
 		{
 			FMargin SafeMargin;
 			FSlateApplication::Get().ResetCustomSafeZone();
-			FSlateApplication::Get().GetSafeZoneSize(SafeMargin, FVector2D(Width, Height));
+			FSlateApplication::Get().GetSafeZoneSize(SafeMargin, FVector2f(Width, Height));
 			const float HeightOfSides = Height - SafeMargin.GetTotalSpaceAlong<Orient_Vertical>();
 			// Top bar
 			FSlateDrawElement::MakeBox(
 				PaintArgs.OutDrawElements,
 				PaintArgs.Layer,
-				PreviewGeometry.ToPaintGeometry(FVector2D(Width, SafeMargin.Top), FSlateLayoutTransform()),
+				PreviewGeometry.ToPaintGeometry(FVector2f(Width, SafeMargin.Top), FSlateLayoutTransform()),
 				WhiteBrush,
 				ESlateDrawEffect::None,
 				UnsafeZoneColor
@@ -3518,7 +3518,7 @@ TOptional<int32> SDesignerView::GetCustomResolutionWidth() const
 {
 	if ( UUserWidget* DefaultWidget = GetDefaultWidget() )
 	{
-		return DefaultWidget->DesignTimeSize.X;
+		return FMath::TruncToInt32(DefaultWidget->DesignTimeSize.X);
 	}
 
 	return 1;
@@ -3528,7 +3528,7 @@ TOptional<int32> SDesignerView::GetCustomResolutionHeight() const
 {
 	if ( UUserWidget* DefaultWidget = GetDefaultWidget() )
 	{
-		return DefaultWidget->DesignTimeSize.Y;
+		return  FMath::TruncToInt32(DefaultWidget->DesignTimeSize.Y);
 	}
 
 	return 1;
diff --git a/Engine/Source/Editor/UMGEditor/Private/Designer/SRuler.cpp b/Engine/Source/Editor/UMGEditor/Private/Designer/SRuler.cpp
--- a/Engine/Source/Editor/UMGEditor/Private/Designer/SRuler.cpp
+++ b/Engine/Source/Editor/UMGEditor/Private/Designer/SRuler.cpp
@@ -317,7 +317,7 @@ int32 SRuler::OnPaint(const FPaintArgs& Args, const FGeometry& AllottedGeometry,
 		FLinearColor(FColor(48,48,48))
 	);
 
-	FVector2D LocalOrigin = AllottedGeometry.AbsoluteToLocal(AbsoluteOrigin);
+	FVector2f LocalOrigin = AllottedGeometry.AbsoluteToLocal(AbsoluteOrigin);
 
 	float Scale = SlateToUnitScale < ScrubConstants::MinDisplayTickSpacing ? ScrubConstants::MinDisplayTickSpacing : SlateToUnitScale;
 	
diff --git a/Engine/Source/Editor/UMGEditor/Private/Designer/STransformHandle.cpp b/Engine/Source/Editor/UMGEditor/Private/Designer/STransformHandle.cpp
--- a/Engine/Source/Editor/UMGEditor/Private/Designer/STransformHandle.cpp
+++ b/Engine/Source/Editor/UMGEditor/Private/Designer/STransformHandle.cpp
@@ -155,7 +155,7 @@ bool STransformHandle::CanResize(UPanelSlot* Slot, const FVector2D& Direction) c
 	return Cast<UCanvasPanelSlot>(Slot) != nullptr;
 }
 
-void STransformHandle::Resize(UCanvasPanelSlot* Slot, const FVector2D& Direction, const FVector2D& Amount)
+void STransformHandle::Resize(UCanvasPanelSlot* Slot, const FVector2D& InDirection, const FVector2D& InAmount)
 {
 	if ( Slot == nullptr )
 	{
@@ -164,10 +164,14 @@ void STransformHandle::Resize(UCanvasPanelSlot* Slot, const FVector2D& Direction
 
 	FMargin Offsets = StartingOffsets;
 	const FAnchorData LayoutData = Slot->GetLayout();
+	FVector2f Alignment = UE::Slate::CastToVector2f(LayoutData.Alignment);
 
-	FVector2D Movement = Amount * Direction;
-	FVector2D PositionMovement = Movement * ( FVector2D(1.0f, 1.0f) - LayoutData.Alignment );
-	FVector2D SizeMovement = Movement;
+	FVector2f Direction = UE::Slate::CastToVector2f(InDirection);
+	FVector2f Amount = UE::Slate::CastToVector2f(InAmount);	
+
+	FVector2f Movement = Amount * Direction;
+	FVector2f PositionMovement = Movement * ( FVector2f(1.0f, 1.0f) - Alignment);
+	FVector2f SizeMovement = Movement;
 
 	if ( Direction.X < 0 )
 	{
@@ -203,7 +207,7 @@ void STransformHandle::Resize(UCanvasPanelSlot* Slot, const FVector2D& Direction
 		}
 		else
 		{
-			Offsets.Left += (Movement * LayoutData.Alignment).X;
+			Offsets.Left += (Movement * Alignment).X;
 			Offsets.Right += Amount.X * Direction.X;
 		}
 	}
@@ -216,7 +220,7 @@ void STransformHandle::Resize(UCanvasPanelSlot* Slot, const FVector2D& Direction
 		}
 		else
 		{
-			Offsets.Top += ( Movement * LayoutData.Alignment ).Y;
+			Offsets.Top += ( Movement * Alignment ).Y;
 			Offsets.Bottom += Amount.Y * Direction.Y;
 		}
 	}
diff --git a/Engine/Source/Editor/UMGEditor/Private/Extensions/CanvasSlotExtension.cpp b/Engine/Source/Editor/UMGEditor/Private/Extensions/CanvasSlotExtension.cpp
--- a/Engine/Source/Editor/UMGEditor/Private/Extensions/CanvasSlotExtension.cpp
+++ b/Engine/Source/Editor/UMGEditor/Private/Extensions/CanvasSlotExtension.cpp
@@ -76,14 +76,14 @@ private:
 /////////////////////////////////////////////////////
 // FCanvasSlotExtension
 
-const float SnapDistance = 7;
+const double SnapDistance = 7.0;
 
-static float DistancePointToLine2D(const FVector2D& LinePointA, const FVector2D& LinePointB, const FVector2D& PointC)
+static double DistancePointToLine2D(const FVector2D& LinePointA, const FVector2D& LinePointB, const FVector2D& PointC)
 {
 	FVector2D AB = LinePointB - LinePointA;
 	FVector2D AC = PointC - LinePointA;
 
-	float Distance = FVector2D::CrossProduct(AB, AC) / FVector2D::Distance(LinePointA, LinePointB);
+	double Distance = FVector2D::CrossProduct(AB, AC) / FVector2D::Distance(LinePointA, LinePointB);
 	return FMath::Abs(Distance);
 }
 
@@ -358,12 +358,12 @@ FReply FCanvasSlotExtension::HandleAnchorEndDrag(const FGeometry& Geometry, cons
 
 void FCanvasSlotExtension::ProximitySnapValue(float SnapFrequency, float SnapProximity, FVector2D::FReal& Value)
 {
-	float MajorAnchorDiv = Value / SnapFrequency;
-	float SubAnchorLinePos = MajorAnchorDiv - FMath::RoundToInt(MajorAnchorDiv);
+	double MajorAnchorDiv = Value / SnapFrequency;
+	double SubAnchorLinePos = MajorAnchorDiv - FMath::RoundToDouble(MajorAnchorDiv);
 
 	if ( FMath::Abs(SubAnchorLinePos) <= SnapProximity )
 	{
-		Value = FMath::RoundToInt(MajorAnchorDiv) * SnapFrequency;
+		Value = FMath::RoundToDouble(MajorAnchorDiv) * SnapFrequency;
 	}
 }
 
@@ -572,7 +572,7 @@ void FCanvasSlotExtension::PaintDragPercentages(const TSet< FWidgetReference >&
 				const FVector2D AnchorMin = LayoutData.Anchors.Minimum;
 				const FVector2D AnchorMax = LayoutData.Anchors.Maximum;
 
-				auto DrawSegment =[&] (FVector2D Offset, FVector2D Start, FVector2D End, float Value, FVector2D TextTransform, bool InHorizontalLine) {
+				auto DrawSegment =[&] (FVector2D Offset, FVector2D Start, FVector2D End, double Value, FVector2D TextTransform, bool InHorizontalLine) {
 					PaintLineWithText(
 						Start + Offset,
 						End + Offset,
@@ -748,7 +748,7 @@ void FCanvasSlotExtension::PaintCollisionLines(const TSet< FWidgetReference >& S
 										FVector2D CollisionPoint = MySegmentBase;
 
 										//TODO Collide against all sides of the arranged geometry.
-										float Distance = DistancePointToLine2D(PointA, PointB, CollisionPoint);
+										double Distance = DistancePointToLine2D(PointA, PointB, CollisionPoint);
 										if ( Distance <= SnapDistance )
 										{
 											FVector2D FarthestPoint = FVector2D::Distance(PointA, CollisionPoint) > FVector2D::Distance(PointB, CollisionPoint) ? PointA : PointB;
diff --git a/Engine/Source/Editor/UMGEditor/Private/TabFactory/AnimationTabSummoner.cpp b/Engine/Source/Editor/UMGEditor/Private/TabFactory/AnimationTabSummoner.cpp
--- a/Engine/Source/Editor/UMGEditor/Private/TabFactory/AnimationTabSummoner.cpp
+++ b/Engine/Source/Editor/UMGEditor/Private/TabFactory/AnimationTabSummoner.cpp
@@ -394,7 +394,7 @@ public:
 		[
 			SNew(SBorder)
 			.BorderImage(FAppStyle::GetBrush("ToolPanel.GroupBorder"))
-			.Padding( FMargin(bIsDrawerTab ? 8.0 : 2.0, 2.0) )
+			.Padding( FMargin(bIsDrawerTab ? 8.0f : 2.0f, 2.0f) )
 			[
 				SNew(SOverlay)
 				+SOverlay::Slot()
diff --git a/Engine/Source/Editor/UMGEditor/Private/WidgetBlueprintEditorUtils.cpp b/Engine/Source/Editor/UMGEditor/Private/WidgetBlueprintEditorUtils.cpp
--- a/Engine/Source/Editor/UMGEditor/Private/WidgetBlueprintEditorUtils.cpp
+++ b/Engine/Source/Editor/UMGEditor/Private/WidgetBlueprintEditorUtils.cpp
@@ -2611,7 +2611,7 @@ float FWidgetBlueprintEditorUtils::GetWidgetPreviewDPIScale(UUserWidget* UserWid
 		}
 	}
 
-	return GetDefault<UUserInterfaceSettings>()->GetDPIScaleBasedOnSize(FIntPoint(PreviewSize.X, PreviewSize.Y));
+	return GetDefault<UUserInterfaceSettings>()->GetDPIScaleBasedOnSize(FIntPoint(FMath::TruncToInt32(PreviewSize.X), FMath::TruncToInt32(PreviewSize.Y)));
 }
 
 
@@ -2809,17 +2809,17 @@ TTuple<float, FVector2D> FWidgetBlueprintEditorUtils::GetThumbnailImageScaleAndO
 	checkf(WidgetSize.X > 0.f && WidgetSize.Y > 0.f, TEXT("The size should have been previously checked to be > 0."));
 
 	float Scale;
-	float XOffset = 0;
-	float YOffset = 0;
+	double XOffset = 0;
+	double YOffset = 0;
 	if (WidgetSize.X > WidgetSize.Y)
 	{
-		Scale = ThumbnailSize.X / WidgetSize.X;
+		Scale = static_cast<float>(ThumbnailSize.X / WidgetSize.X);
 		WidgetSize *= Scale;
 		YOffset = (ThumbnailSize.Y - WidgetSize.Y) / 2.f;
 	}
 	else
 	{
-		Scale = ThumbnailSize.Y / WidgetSize.Y;
+		Scale = static_cast<float>(ThumbnailSize.Y / WidgetSize.Y);
 		WidgetSize *= Scale;
 		XOffset = (ThumbnailSize.X - WidgetSize.X) / 2.f;
 	}
diff --git a/Engine/Source/Editor/UMGEditor/Public/Designer/SDesignSurface.h b/Engine/Source/Editor/UMGEditor/Public/Designer/SDesignSurface.h
--- a/Engine/Source/Editor/UMGEditor/Public/Designer/SDesignSurface.h
+++ b/Engine/Source/Editor/UMGEditor/Public/Designer/SDesignSurface.h
@@ -135,7 +135,7 @@ protected:
 	FVector2D ZoomToFitPadding;
 
 	/** The Y component of mouse drag (used when zooming) */
-	float TotalMouseDelta;
+	double TotalMouseDelta;
 
 	/** Offset in the panel the user started the LMB+RMB zoom from */
 	FVector2D ZoomStartOffset;
