PR #10919 Fix crash that can occur when placing widgets in inherited WidgetTrees

#tests It no longer crashes when creating a widget in a user widget with a parent already using the same name that would be used by this new widget. Found another bug while testing it (renaming of a widget in the parent with a child already using this name UE-193778), that will be fixed in another CL..

diff --git a/Engine/Source/Editor/UMGEditor/Private/Templates/WidgetTemplateClass.cpp b/Engine/Source/Editor/UMGEditor/Private/Templates/WidgetTemplateClass.cpp
--- a/Engine/Source/Editor/UMGEditor/Private/Templates/WidgetTemplateClass.cpp
+++ b/Engine/Source/Editor/UMGEditor/Private/Templates/WidgetTemplateClass.cpp
@@ -11,6 +11,7 @@
 #include "UObject/CoreRedirects.h"
 #include "WidgetBlueprintEditorUtils.h"
 
+#include "Blueprint/WidgetBlueprintGeneratedClass.h"
 #include "Blueprint/WidgetTree.h"
 #include "Styling/SlateIconFinder.h"
 
@@ -164,6 +165,21 @@ UWidget* FWidgetTemplateClass::CreateNamed(class UWidgetTree* Tree, FName NameOv
 		return nullptr;
 	}
 
+	// For inherited widget trees, we need to make sure the name is unique for the parent tree
+	// We do this even if NameOverride is not set as the default name could exist in the parent tree
+	if (Tree->RootWidget == nullptr)
+	{
+		if (const UWidgetBlueprint* WidgetBP = Cast<UWidgetBlueprint>(Tree->GetOuter()))
+		{
+			const UWidgetBlueprintGeneratedClass* BaseWidgetClass = Cast<UWidgetBlueprintGeneratedClass>(WidgetBP->GeneratedClass);
+			const UWidgetBlueprintGeneratedClass* ParentWidgetClass = BaseWidgetClass ? BaseWidgetClass->FindWidgetTreeOwningClass() : nullptr;
+			if (ParentWidgetClass != BaseWidgetClass && ParentWidgetClass != nullptr)
+			{
+				NameOverride = MakeUniqueObjectName(ParentWidgetClass->GetWidgetTreeArchetype(), WidgetClass.Get(), NameOverride);
+			}
+		}
+	}
+
 	if (NameOverride != NAME_None)
 	{
 		UObject* ExistingObject = StaticFindObject(UObject::StaticClass(), Tree, *NameOverride.ToString());
