UMG Designer and UI Framework Truncation Warning fixes

diff --git a/Engine/Plugins/Experimental/UIFramework/Source/Private/Widgets/UIFImageBlock.cpp b/Engine/Plugins/Experimental/UIFramework/Source/Private/Widgets/UIFImageBlock.cpp
--- a/Engine/Plugins/Experimental/UIFramework/Source/Private/Widgets/UIFImageBlock.cpp
+++ b/Engine/Plugins/Experimental/UIFramework/Source/Private/Widgets/UIFImageBlock.cpp
@@ -96,8 +96,8 @@ void UUIFrameworkImageBlock::LocalOnUMGWidgetCreated()
 	{
 		if (UTexture2D* Texture = Cast<UTexture2D>(Resource))
 		{
-			TmpBrush.ImageSize.X = Texture->GetSizeX();
-			TmpBrush.ImageSize.Y = Texture->GetSizeY();
+			TmpBrush.ImageSize.X = static_cast<float>(Texture->GetSizeX());
+			TmpBrush.ImageSize.Y = static_cast<float>(Texture->GetSizeY());
 		}
 	}
 	TmpBrush.SetResourceObject(Resource);
diff --git a/Engine/Source/Editor/UMGEditor/Private/Customizations/CanvasSlotCustomization.cpp b/Engine/Source/Editor/UMGEditor/Private/Customizations/CanvasSlotCustomization.cpp
--- a/Engine/Source/Editor/UMGEditor/Private/Customizations/CanvasSlotCustomization.cpp
+++ b/Engine/Source/Editor/UMGEditor/Private/Customizations/CanvasSlotCustomization.cpp
@@ -94,8 +94,8 @@ public:
 					.Padding(0)
 					[
 						SNew(SBox)
-						.WidthOverride(64)
-						.HeightOverride(64)
+						.WidthOverride(64.0f)
+						.HeightOverride(64.0f)
 						.HAlign(HAlign_Center)
 						.VAlign(VAlign_Center)
 						[
@@ -110,8 +110,8 @@ public:
 
 									+ SConstraintCanvas::Slot()
 									.Anchors(Anchors)
-									.Offset(FMargin(0, 0, Anchors.IsStretchedHorizontal() ? 0 : 15, Anchors.IsStretchedVertical() ? 0 : 15))
-									.Alignment(FVector2D(Anchors.IsStretchedHorizontal() ? 0 : Anchors.Minimum.X, Anchors.IsStretchedVertical() ? 0 : Anchors.Minimum.Y))
+									.Offset(FMargin(0.0f, 0.0f, Anchors.IsStretchedHorizontal() ? 0.0f : 15.0f, Anchors.IsStretchedVertical() ? 0.0f : 15.0f))
+									.Alignment(FVector2D(Anchors.IsStretchedHorizontal() ? 0.0f : Anchors.Minimum.X, Anchors.IsStretchedVertical() ? 0.0f : Anchors.Minimum.Y))
 									[
 										SNew(SImage)
 										.Image(FAppStyle::Get().GetBrush("UMGEditor.AnchoredWidget"))
diff --git a/Engine/Source/Editor/UMGEditor/Private/Customizations/UMGDetailCustomizations.cpp b/Engine/Source/Editor/UMGEditor/Private/Customizations/UMGDetailCustomizations.cpp
--- a/Engine/Source/Editor/UMGEditor/Private/Customizations/UMGDetailCustomizations.cpp
+++ b/Engine/Source/Editor/UMGEditor/Private/Customizations/UMGDetailCustomizations.cpp
@@ -839,7 +839,7 @@ void FBlueprintWidgetCustomization::CreateMulticastEventCustomization(IDetailLay
 				+ SHorizontalBox::Slot()
 				.AutoWidth()
 				.VAlign(VAlign_Center)
-				.Padding(0, 0, 5, 0)
+				.Padding(0.0f, 0.0f, 5.0f, 0.0f)
 				[
 					SNew(SImage)
 					.Image(FAppStyle::Get().GetBrush("GraphEditor.Event_16x"))
@@ -857,10 +857,10 @@ void FBlueprintWidgetCustomization::CreateMulticastEventCustomization(IDetailLay
 				+ SHorizontalBox::Slot()
 				.HAlign(HAlign_Left)
 				.VAlign(VAlign_Center)
-				.Padding(0)
+				.Padding(0.0f)
 				[
 					SNew(SButton)
-					.ContentPadding(FMargin(3.0, 2.0))
+					.ContentPadding(FMargin(3.0f, 2.0f))
 					.OnClicked(this, &FBlueprintWidgetCustomization::HandleAddOrViewEventForVariable, EventName, PropertyName, MakeWeakObjectPtr(PropertyClass))
 					[
 						SNew(SWidgetSwitcher)
diff --git a/Engine/Source/Editor/UMGEditor/Private/Customizations/WidgetChildTypeCustomization.cpp b/Engine/Source/Editor/UMGEditor/Private/Customizations/WidgetChildTypeCustomization.cpp
--- a/Engine/Source/Editor/UMGEditor/Private/Customizations/WidgetChildTypeCustomization.cpp
+++ b/Engine/Source/Editor/UMGEditor/Private/Customizations/WidgetChildTypeCustomization.cpp
@@ -129,7 +129,7 @@ TSharedRef<SWidget> FWidgetChildTypeCustomization::GetPopupContent()
 
 	return SNew(SBorder)
 		.BorderImage(FAppStyle::Get().GetBrush("Brushes.Recessed"))
-		.Padding(4)
+		.Padding(4.0f)
 		[
 			SNew(SBox)
 			.MinDesiredWidth(MinPopupWidth)
diff --git a/Engine/Source/Editor/UMGEditor/Private/Customizations/WidgetNavigationCustomization.cpp b/Engine/Source/Editor/UMGEditor/Private/Customizations/WidgetNavigationCustomization.cpp
--- a/Engine/Source/Editor/UMGEditor/Private/Customizations/WidgetNavigationCustomization.cpp
+++ b/Engine/Source/Editor/UMGEditor/Private/Customizations/WidgetNavigationCustomization.cpp
@@ -228,7 +228,7 @@ void FWidgetNavigationCustomization::MakeNavRow(TWeakPtr<IPropertyHandle> Proper
 				SNew(SComboButton)
 				.Visibility(this, &FWidgetNavigationCustomization::GetExplictWidgetFieldVisibility, PropertyHandle, Nav)
 				.OnGetMenuContent(this, &FWidgetNavigationCustomization::OnGenerateWidgetList, PropertyHandle, Nav)
-				.ContentPadding(1)
+				.ContentPadding(1.0f)
 				.ButtonContent()
 				[
 					SNew(STextBlock)
diff --git a/Engine/Source/Editor/UMGEditor/Private/Customizations/WidgetTypeCustomization.cpp b/Engine/Source/Editor/UMGEditor/Private/Customizations/WidgetTypeCustomization.cpp
--- a/Engine/Source/Editor/UMGEditor/Private/Customizations/WidgetTypeCustomization.cpp
+++ b/Engine/Source/Editor/UMGEditor/Private/Customizations/WidgetTypeCustomization.cpp
@@ -80,11 +80,11 @@ TSharedRef<SWidget> FWidgetTypeCustomization::GetPopupContent()
 			]
 
 			+ SVerticalBox::Slot()
-			.Padding(0, 2.0f, 0, 0)
+			.Padding(0.0f, 2.0f, 0.0f, 0.0f)
 			.AutoHeight()
 			[
 				SNew(SBox)
-				.MaxDesiredHeight(300)
+				.MaxDesiredHeight(300.0f)
 				[
 					WidgetListView.ToSharedRef()
 				]
diff --git a/Engine/Source/Editor/UMGEditor/Private/Designer/SDesignerView.cpp b/Engine/Source/Editor/UMGEditor/Private/Designer/SDesignerView.cpp
--- a/Engine/Source/Editor/UMGEditor/Private/Designer/SDesignerView.cpp
+++ b/Engine/Source/Editor/UMGEditor/Private/Designer/SDesignerView.cpp
@@ -420,7 +420,7 @@ TSharedRef<SWidget> SDesignerView::CreateOverlayUI()
 
 	// Outline and text for important state.
 	+ SOverlay::Slot()
-	.Padding(0)
+	.Padding(0.0f)
 	.VAlign(VAlign_Fill)
 	.HAlign(HAlign_Fill)
 	[
@@ -429,7 +429,7 @@ TSharedRef<SWidget> SDesignerView::CreateOverlayUI()
 
 		// Top-right corner text indicating PIE is active
 		+ SOverlay::Slot()
-		.Padding(0)
+		.Padding(0.0f)
 		.VAlign(VAlign_Fill)
 		.HAlign(HAlign_Fill)
 		[
@@ -440,7 +440,7 @@ TSharedRef<SWidget> SDesignerView::CreateOverlayUI()
 
 		// Top-right corner text indicating PIE is active
 		+ SOverlay::Slot()
-		.Padding(20)
+		.Padding(20.0f)
 		.VAlign(VAlign_Top)
 		.HAlign(HAlign_Right)
 		[
@@ -461,7 +461,7 @@ TSharedRef<SWidget> SDesignerView::CreateOverlayUI()
 		+ SHorizontalBox::Slot()
 		.AutoWidth()
 		.VAlign(VAlign_Center)
-		.Padding(6, 2, 0, 0)
+		.Padding(6.0f, 2.0f, 0.0f, 0.0f)
 		[
 			SNew(STextBlock)
 			.TextStyle(FAppStyle::Get(), "Graph.ZoomText")
@@ -473,7 +473,7 @@ TSharedRef<SWidget> SDesignerView::CreateOverlayUI()
 		+ SHorizontalBox::Slot()
 		.AutoWidth()
 		.VAlign(VAlign_Center)
-		.Padding(40, 2, 0, 0)
+		.Padding(40.0f, 2.0f, 0.0f, 0.0f)
 		[
 			SNew(STextBlock)
 			.TextStyle(FAppStyle::Get(), "Graph.ZoomText")
@@ -486,7 +486,7 @@ TSharedRef<SWidget> SDesignerView::CreateOverlayUI()
 		+ SHorizontalBox::Slot()
 		.AutoWidth()
 		.VAlign(VAlign_Center)
-		.Padding(40, 2, 0, 0)
+		.Padding(40.0f, 2.0f, 0.0f, 0.0f)
 		[
 			SNew(STextBlock)
 			.TextStyle(FAppStyle::Get(), "Graph.ZoomText")
@@ -500,7 +500,7 @@ TSharedRef<SWidget> SDesignerView::CreateOverlayUI()
 		.FillWidth(1.0f)
 		[
 			SNew(SSpacer)
-			.Size(FVector2D(1, 1))
+			.Size(FVector2D(1.0f, 1.0f))
 		]
 
 		+ SHorizontalBox::Slot()
@@ -597,7 +597,7 @@ TSharedRef<SWidget> SDesignerView::CreateOverlayUI()
 		]
 		+ SHorizontalBox::Slot()
 		.AutoWidth()
-		.Padding(FMargin(2, 0))
+		.Padding(FMargin(2.0f, 0.0f))
 		.VAlign(VAlign_Center)
 		[
 			SNew(STextBlock)
@@ -607,7 +607,7 @@ TSharedRef<SWidget> SDesignerView::CreateOverlayUI()
 		]
 		+ SHorizontalBox::Slot()
 		.AutoWidth()
-		.Padding(FMargin(2, 0))
+		.Padding(FMargin(2.0f, 0.0f))
 		.VAlign(VAlign_Center)
 		[
 			SNew(SNumericEntryBox<int32>)
@@ -619,12 +619,12 @@ TSharedRef<SWidget> SDesignerView::CreateOverlayUI()
 			.Value(this, &SDesignerView::GetCustomResolutionWidth)
 			.OnValueChanged(this, &SDesignerView::OnCustomResolutionWidthChanged)
 			.Visibility(this, &SDesignerView::GetCustomResolutionEntryVisibility)
-			.MinDesiredValueWidth(50)
+			.MinDesiredValueWidth(50.0f)
 			.ToolTipText(LOCTEXT("CustomSize_WidthTooltip", "1+\tSets the width of the widget in the designer.\n0\tThe width will match the desired width of the widget."))
 		]
 		+ SHorizontalBox::Slot()
 		.AutoWidth()
-		.Padding(FMargin(2, 0))
+		.Padding(FMargin(2.0f, 0.0f))
 		.VAlign(VAlign_Center)
 		[
 			SNew(STextBlock)
@@ -634,7 +634,7 @@ TSharedRef<SWidget> SDesignerView::CreateOverlayUI()
 		]
 		+ SHorizontalBox::Slot()
 		.AutoWidth()
-		.Padding(FMargin(2, 0))
+		.Padding(FMargin(2.0f, 0.0f))
 		.VAlign(VAlign_Center)
 		[
 			SNew(SNumericEntryBox<int32>)
@@ -663,7 +663,7 @@ TSharedRef<SWidget> SDesignerView::CreateOverlayUI()
 			.BorderBackgroundColor(FLinearColor(0.10, 0.10, 0.10, 0.75))
 			.HAlign(HAlign_Center)
 			.VAlign(VAlign_Center)
-			.Padding(FMargin(0, 5))
+			.Padding(FMargin(0.0f, 5.0f))
 			.Visibility(this, &SDesignerView::GetInfoBarVisibility)
 			[
 				SNew(STextBlock)
@@ -681,7 +681,7 @@ TSharedRef<SWidget> SDesignerView::CreateOverlayUI()
 		SNew(SHorizontalBox)
 		+ SHorizontalBox::Slot()
 		.AutoWidth()
-		.Padding(6, 0, 0, 2)
+		.Padding(6.0f, 0.0f, 0.0f, 2.0f)
 		[
 			SNew(SVerticalBox)
 			+ SVerticalBox::Slot()
@@ -713,7 +713,7 @@ TSharedRef<SWidget> SDesignerView::CreateOverlayUI()
 		+ SHorizontalBox::Slot()
 		.FillWidth(1.0f)
 		.HAlign(HAlign_Right)
-		.Padding(0, 0, 6, 2)
+		.Padding(0.0f, 0.0f, 6.0f, 2.0f)
 		[
 			SNew(SHorizontalBox)
 			+ SHorizontalBox::Slot()
@@ -728,7 +728,7 @@ TSharedRef<SWidget> SDesignerView::CreateOverlayUI()
 
 			+ SHorizontalBox::Slot()
 			.AutoWidth()
-			.Padding(6, 0, 0, 0)
+			.Padding(6.0f, 0.0f, 0.0f, 0.0f)
 			.VAlign(VAlign_Bottom)
 			[
 				SNew(SButton)
diff --git a/Engine/Source/Editor/UMGEditor/Private/Details/SFunctionSelector.cpp b/Engine/Source/Editor/UMGEditor/Private/Details/SFunctionSelector.cpp
--- a/Engine/Source/Editor/UMGEditor/Private/Details/SFunctionSelector.cpp
+++ b/Engine/Source/Editor/UMGEditor/Private/Details/SFunctionSelector.cpp
@@ -55,14 +55,14 @@ void SFunctionSelector::Construct(const FArguments& InArgs, TSharedRef<FWidgetBl
 		[
 			SNew(SComboButton)
 			.OnGetMenuContent(this, &SFunctionSelector::OnGenerateDelegateMenu)
-			.ContentPadding(1)
+			.ContentPadding(1.0f)
 			.ButtonContent()
 			[
 				SNew(SHorizontalBox)
 				+ SHorizontalBox::Slot()
 				.AutoWidth()
 				.VAlign(VAlign_Center)
-				.Padding(4, 1, 0, 0)
+				.Padding(4.0f, 1.0f, 0.0f, 0.0f)
 				[
 					SNew(STextBlock)
 					.Text(this, &SFunctionSelector::GetCurrentBindingText)
diff --git a/Engine/Source/Editor/UMGEditor/Private/Hierarchy/SHierarchyView.cpp b/Engine/Source/Editor/UMGEditor/Private/Hierarchy/SHierarchyView.cpp
--- a/Engine/Source/Editor/UMGEditor/Private/Hierarchy/SHierarchyView.cpp
+++ b/Engine/Source/Editor/UMGEditor/Private/Hierarchy/SHierarchyView.cpp
@@ -57,7 +57,7 @@ void SHierarchyView::Construct(const FArguments& InArgs, TSharedPtr<FWidgetBluep
 			SNew(SVerticalBox)
 
 			+ SVerticalBox::Slot()
-			.Padding(4)
+			.Padding(4.0f)
 			.AutoHeight()
 			[
 				SAssignNew(SearchBoxPtr, SSearchBox)
diff --git a/Engine/Source/Editor/UMGEditor/Private/Hierarchy/SReadOnlyHierarchyView.cpp b/Engine/Source/Editor/UMGEditor/Private/Hierarchy/SReadOnlyHierarchyView.cpp
--- a/Engine/Source/Editor/UMGEditor/Private/Hierarchy/SReadOnlyHierarchyView.cpp
+++ b/Engine/Source/Editor/UMGEditor/Private/Hierarchy/SReadOnlyHierarchyView.cpp
@@ -41,7 +41,7 @@ void SReadOnlyHierarchyView::Construct(const FArguments& InArgs, const UWidgetBl
 	if (InArgs._ShowSearch)
 	{
 		ContentBox->AddSlot()
-			.Padding(2)
+			.Padding(2.0f)
 			.AutoHeight()
 			[
 				SAssignNew(SearchBox, SSearchBox)
diff --git a/Engine/Source/Editor/UMGEditor/Private/Library/SLibraryView.cpp b/Engine/Source/Editor/UMGEditor/Private/Library/SLibraryView.cpp
--- a/Engine/Source/Editor/UMGEditor/Private/Library/SLibraryView.cpp
+++ b/Engine/Source/Editor/UMGEditor/Private/Library/SLibraryView.cpp
@@ -216,7 +216,7 @@ void SLibraryView::Construct(const FArguments& InArgs, TSharedPtr<FWidgetBluepri
 		SNew(SVerticalBox)
 
 		+ SVerticalBox::Slot()
-		.Padding(4)
+		.Padding(4.0f)
 		.AutoHeight()
 		[
 			SNew(SHorizontalBox)
@@ -230,10 +230,10 @@ void SLibraryView::Construct(const FArguments& InArgs, TSharedPtr<FWidgetBluepri
 			]
 			+ SHorizontalBox::Slot()
 			.AutoWidth()
-			.Padding(2, 0, 0, 0)
+			.Padding(2.0f, 0.0f, 0.0f, 0.0f)
 			[
 				SNew(SComboButton)
-				.ContentPadding(0)
+				.ContentPadding(0.0f)
 				.ComboButtonStyle(&FAppStyle::Get().GetWidgetStyle<FComboButtonStyle>("SimpleComboButton"))
 				.HasDownArrow(false)
 				.AddMetaData<FTagMetaData>(FTagMetaData(TEXT("ViewOptions")))
@@ -255,15 +255,15 @@ void SLibraryView::Construct(const FArguments& InArgs, TSharedPtr<FWidgetBluepri
 		[
 			SNew(SBorder)
 			.BorderImage(FAppStyle::Get().GetBrush("Brushes.Recessed"))
-			.Padding(0)
+			.Padding(0.0f)
 			[
 				SNew(SScrollBox)
-				.ScrollBarPadding(FMargin(2.0, 0))
+				.ScrollBarPadding(FMargin(2.0f, 0.0f))
 				+SScrollBox::Slot()
 				[
 					SNew(SBorder)
 					.BorderImage(FAppStyle::Get().GetBrush("Brushes.Recessed"))
-					.Padding(0)
+					.Padding(0.0f)
 					[
 						WidgetTemplatesView.ToSharedRef()
 					]
diff --git a/Engine/Source/Editor/UMGEditor/Private/Navigation/SWidgetDesignerNavigation.cpp b/Engine/Source/Editor/UMGEditor/Private/Navigation/SWidgetDesignerNavigation.cpp
--- a/Engine/Source/Editor/UMGEditor/Private/Navigation/SWidgetDesignerNavigation.cpp
+++ b/Engine/Source/Editor/UMGEditor/Private/Navigation/SWidgetDesignerNavigation.cpp
@@ -262,25 +262,25 @@ void SWidgetDesignerNavigation::Construct(const FArguments& InArgs, TSharedPtr<F
 			SNew(SVerticalBox)
 
 			+ SVerticalBox::Slot()
-			.Padding(4)
+			.Padding(4.0f)
 			.FillHeight(1.f)
 			[
 				NavigationEventSimulationView.ToSharedRef()
 			]
 			
 			+ SVerticalBox::Slot()
-			.Padding(4)
+			.Padding(4.0f)
 			.AutoHeight()
 			[
 				DetailView->GetWidget().ToSharedRef()
 			]
 
 			+ SVerticalBox::Slot()
-			.Padding(4)
+			.Padding(4.0f)
 			.AutoHeight()
 			[
 				SNew(SButton)
-				.ContentPadding(3)
+				.ContentPadding(3.0f)
 				.OnClicked(this, &SWidgetDesignerNavigation::HandleRefreshClicked)
 				.IsEnabled(this, &SWidgetDesignerNavigation::HandleRefreshEnabled)
 				.Content()
diff --git a/Engine/Source/Editor/UnrealEd/Public/Tools/ToolCompatible.inl b/Engine/Source/Editor/UnrealEd/Public/Tools/ToolCompatible.inl
--- a/Engine/Source/Editor/UnrealEd/Public/Tools/ToolCompatible.inl
+++ b/Engine/Source/Editor/UnrealEd/Public/Tools/ToolCompatible.inl
@@ -117,7 +117,7 @@ void TToolCompatibleMixin<ToolCompatibleWidget>::OnMouseEnter(const FGeometry& M
 	ToolCompatibleWidget::OnMouseEnter(MyGeometry, MouseEvent);
 	if (FWidgetModeManager* WidgetModeManager = GetWidgetModeManger())
 	{
-		WidgetModeManager->MouseEnter(nullptr, nullptr, MouseEvent.GetScreenSpacePosition().X, MouseEvent.GetScreenSpacePosition().Y);
+		WidgetModeManager->MouseEnter(nullptr, nullptr, FMath::TruncToInt32(MouseEvent.GetScreenSpacePosition().X), FMath::TruncToInt32(MouseEvent.GetScreenSpacePosition().Y));
 	}
 }
 
