UMG: ensure we do the widget rename operation from the details panel name textbox if the name actually changed.

diff --git a/Engine/Source/Editor/UMGEditor/Private/Details/SWidgetDetailsView.cpp b/Engine/Source/Editor/UMGEditor/Private/Details/SWidgetDetailsView.cpp
--- a/Engine/Source/Editor/UMGEditor/Private/Details/SWidgetDetailsView.cpp
+++ b/Engine/Source/Editor/UMGEditor/Private/Details/SWidgetDetailsView.cpp
@@ -547,7 +547,10 @@ void SWidgetDetailsView::HandleNameTextCommitted(const FText& Text, ETextCommit:
 			if ( HandleVerifyNameTextChanged(Text, DummyText) )
 			{
 				UWidget* Widget = Cast<UWidget>(SelectedObjects[0].Get());
-				FWidgetBlueprintEditorUtils::RenameWidget(BlueprintEditor.Pin().ToSharedRef(), Widget->GetFName(), Text.ToString());
+				if (!Widget->GetLabelText().EqualToCaseIgnored(Text))
+				{
+					FWidgetBlueprintEditorUtils::RenameWidget(BlueprintEditor.Pin().ToSharedRef(), Widget->GetFName(), Text.ToString());
+				}
 			}
 		}
 		IsReentrant = false;
