DetectSlateWidgetLeaks (DataValidator) instanciate abstract blueprints when run, which trigger ensure.
#tests Ran the validation on many WBP without issue. WBP with abstract flag no longer trigger ensure.

diff --git a/Engine/Source/Editor/UMGEditor/Private/WidgetBlueprint.cpp b/Engine/Source/Editor/UMGEditor/Private/WidgetBlueprint.cpp
--- a/Engine/Source/Editor/UMGEditor/Private/WidgetBlueprint.cpp
+++ b/Engine/Source/Editor/UMGEditor/Private/WidgetBlueprint.cpp
@@ -777,6 +777,13 @@ bool UWidgetBlueprint::DetectSlateWidgetLeaks(FDataValidationContext& Context) c
 		return false;
 	}
 
+	// The detection relies on instantiation of the class: don't try to create an abstract class. 
+	// The validation will have to be run on the WBP inheriting from abstract ones.
+	if (GeneratedClass->HasAnyClassFlags(CLASS_Abstract))
+	{
+		return false;
+	}
+
 	UWorld* DummyWorld = NewObject<UWorld>();
 	UUserWidget* TempUserWidget = NewObject<UUserWidget>(DummyWorld, GeneratedClass);
 	TempUserWidget->ClearFlags(RF_Transactional);
