Improvements to BindWidgets window:
- Fixed icons being stretched
- Added colors to icons
- Added jump to native to context menu
- Categories now expanded by default

diff --git a/Engine/Source/Editor/UMGEditor/Private/UMGEditorModule.cpp b/Engine/Source/Editor/UMGEditor/Private/UMGEditorModule.cpp
--- a/Engine/Source/Editor/UMGEditor/Private/UMGEditorModule.cpp
+++ b/Engine/Source/Editor/UMGEditor/Private/UMGEditorModule.cpp
@@ -42,6 +42,7 @@
 #include "Customizations/ListViewBaseDetails.h"
 #include "WidgetBlueprintThumbnailRenderer.h"
 #include "Customizations/WidgetThumbnailCustomization.h"
+#include "Widgets/SBindWidgetView.h"
 
 #define LOCTEXT_NAMESPACE "UMG"
 
@@ -74,6 +75,7 @@ public:
 		if (GIsEditor)
 		{
 			FDesignerCommands::Register();
+			FBindWidgetCommands::Register();
 		}
 
 		MenuExtensibilityManager = MakeShared<FExtensibilityManager>();
diff --git a/Engine/Source/Editor/UMGEditor/Private/Widgets/SBindWidgetView.cpp b/Engine/Source/Editor/UMGEditor/Private/Widgets/SBindWidgetView.cpp
--- a/Engine/Source/Editor/UMGEditor/Private/Widgets/SBindWidgetView.cpp
+++ b/Engine/Source/Editor/UMGEditor/Private/Widgets/SBindWidgetView.cpp
@@ -7,6 +7,7 @@
 
 #include "IDocumentation.h"
 #include "Editor.h"
+#include "SourceCodeNavigation.h"
 #include "Styling/AppStyle.h"
 #include "Styling/SlateIconFinder.h"
 #include "Widgets/SBoxPanel.h"
@@ -22,6 +23,11 @@
 
 #define LOCTEXT_NAMESPACE "UMG.BindWidget"
 
+void FBindWidgetCommands::RegisterCommands() 
+{
+	UI_COMMAND( GotoNativeVarDefinition, "Goto Code Definition", "Goto the native code definition of this variable", EUserInterfaceActionType::Button, FInputChord() );
+}
+
 namespace UE
 {
 namespace UMG
@@ -76,7 +82,9 @@ public:
 
 		SMultiColumnTableRow<FBindWidgetListEntryPtr>::Construct(
 			FSuperRowType::FArguments()
-			.Padding(1.0f),
+			.Padding(1.0f)
+			.ShowSelection(EntryPtr->EntryType == FBindWidgetListEntry::EEntryType::Binding)
+			,
 			OwnerTableView
 		);
 	}
@@ -119,12 +127,12 @@ public:
 				FText Tooltip = FText::GetEmpty();
 				if (EntryPtr->bIsBound && EntryPtr->bIsCorrectClass)
 				{
-					Image = FAppStyle::GetBrush("Icons.Success");
+					Image = FAppStyle::GetBrush("Icons.SuccessWithColor");
 					Tooltip = LOCTEXT("BoundCorrectlyTooltip", "The widget is bound");
 				}
 				else if (EntryPtr->bIsBound) // is not of the correct class
 				{
-					Image = FAppStyle::GetBrush("Icons.Error");
+					Image = FAppStyle::GetBrush("Icons.ErrorWithColor");
 					Tooltip = LOCTEXT("BoundWidgetWrongClassTooltip", "The bound widget is not of the correct type.");
 				}
 				else if (EntryPtr->bIsOptional) // is not of the correct class
@@ -134,11 +142,13 @@ public:
 				}
 				else
 				{
-					Image = FAppStyle::GetBrush("Icons.Error");
+					Image = FAppStyle::GetBrush("Icons.ErrorWithColor");
 					Tooltip = LOCTEXT("BoundNoWidgetTooltip", "The widget is not bound.");
 				}
 
 				return SNew(SBox)
+					.HAlign(HAlign_Center)
+					.VAlign(VAlign_Center)
 					[
 						SNew(SImage)
 						.Image(Image)
@@ -150,6 +160,19 @@ public:
 		return SNullWidget::NullWidget;
 	}
 
+	virtual FReply OnMouseButtonDoubleClick(const FGeometry& InMyGeometry, const FPointerEvent& InMouseEvent) override
+	{
+		if (EntryPtr->EntryType == FBindWidgetListEntry::EEntryType::Binding
+			&& EntryPtr->Property
+			&& FSourceCodeNavigation::CanNavigateToProperty(EntryPtr->Property))
+		{
+			FSourceCodeNavigation::NavigateToProperty(EntryPtr->Property);
+			return FReply::Handled();
+		}
+
+		return FReply::Unhandled();
+	}
+
 private:
 	FText HandleGetTooltip() const
 	{
@@ -194,9 +217,11 @@ public:
 		}
 
 		Super::Construct(Super::FArguments()
-			.SelectionMode(ESelectionMode::None)
+			.SelectionMode(ESelectionMode::Single)
 			.OnGenerateRow(this, &SBindWidgetView::HandleGenerateRow)
 			.OnGetChildren(this, &SBindWidgetView::HandleGetChildren)
+			.OnContextMenuOpening(this, &SBindWidgetView::OnContextMenuOpening)
+			.OnIsSelectableOrNavigable(this, &SBindWidgetView::OnIsSelectableOrNavigable)
 			.TreeItemsSource(&CategorySourceData)
 			.HeaderRow
 			(
@@ -211,6 +236,18 @@ public:
 				.FixedWidth(16.f)
 				.DefaultLabel(LOCTEXT("EmptyHeaderName", ""))
 			));
+
+		CommandList = MakeShareable(new FUICommandList);
+		CommandList->MapAction(FBindWidgetCommands::Get().GotoNativeVarDefinition,
+			FExecuteAction::CreateSP(this, &SBindWidgetView::GotoNativeCodeVarDefinition),
+			FCanExecuteAction(),
+			FIsActionChecked());
+
+		// Default to open
+		for (const auto& Cat : CategorySourceData)
+		{
+			this->SetItemExpansion(Cat, true);
+		}
 	}
 
 	void SetSourceData(TArray<FBindWidgetListEntryPtr> InSourceData)
@@ -256,9 +293,44 @@ private:
 		}
 	}
 
+	TSharedPtr<SWidget> OnContextMenuOpening()
+	{
+		const TArray<FBindWidgetListEntryPtr> CurSelectedItems = GetSelectedItems();
+		if (CurSelectedItems.Num() != 1)
+		{
+			return TSharedPtr<SWidget>();
+		}
+
+		const bool bShouldCloseWindowAfterMenuSelection = true;
+		FMenuBuilder MenuBuilder(bShouldCloseWindowAfterMenuSelection, CommandList);
+
+		MenuBuilder.AddMenuEntry(FBindWidgetCommands::Get().GotoNativeVarDefinition);
+
+		return MenuBuilder.MakeWidget();
+	}
+
+	bool OnIsSelectableOrNavigable(FBindWidgetListEntryPtr InItem) const
+	{
+		return InItem->EntryType == FBindWidgetListEntry::EEntryType::Binding;
+	}
+
+	void GotoNativeCodeVarDefinition()
+	{
+		const TArray<FBindWidgetListEntryPtr> CurSelectedItems = GetSelectedItems();
+		if (CurSelectedItems.Num() == 1)
+		{
+			if (CurSelectedItems[0]->Property
+				&& FSourceCodeNavigation::CanNavigateToProperty(CurSelectedItems[0]->Property))
+			{
+				FSourceCodeNavigation::NavigateToProperty(CurSelectedItems[0]->Property);
+			}
+		}
+	}
+	
 private:
 	TArray<FBindWidgetListEntryPtr> AllSourceData;
 	TArray<FBindWidgetListEntryPtr> CategorySourceData;
+	TSharedPtr<FUICommandList> CommandList;
 };
 
 
diff --git a/Engine/Source/Editor/UMGEditor/Private/Widgets/SBindWidgetView.h b/Engine/Source/Editor/UMGEditor/Private/Widgets/SBindWidgetView.h
--- a/Engine/Source/Editor/UMGEditor/Private/Widgets/SBindWidgetView.h
+++ b/Engine/Source/Editor/UMGEditor/Private/Widgets/SBindWidgetView.h
@@ -3,10 +3,13 @@
 #pragma once
 
 #include "CoreMinimal.h"
+#include "Framework/Commands/Commands.h"
+#include "Styling/AppStyle.h"
 #include "Widgets/DeclarativeSyntaxSupport.h"
 #include "Widgets/SCompoundWidget.h"
 
 class FWidgetBlueprintEditor;
+class FUICommandInfo;
 class UBlueprint;
 class UObject;
 
@@ -18,6 +21,20 @@ namespace UMG
 }
 }
 
+class FBindWidgetCommands : public TCommands<FBindWidgetCommands>
+{
+public:
+	FBindWidgetCommands() 
+		: TCommands<FBindWidgetCommands>(TEXT("BindWidget"), NSLOCTEXT("Contexts", "Bind Widget", "Bind Widget"), NAME_None, FAppStyle::GetAppStyleSetName())
+	{
+	}
+
+	/** Initialize commands */
+	virtual void RegisterCommands() override;
+
+	TSharedPtr<FUICommandInfo> GotoNativeVarDefinition;
+};
+
 /**
  * 
  */
