UMG: Add the NamedSlot and the name of the viewmodel to the asset registry.

diff --git a/Engine/Plugins/Runtime/ModelViewViewModel/Source/ModelViewViewModelBlueprint/Private/MVVMBlueprintView.cpp b/Engine/Plugins/Runtime/ModelViewViewModel/Source/ModelViewViewModelBlueprint/Private/MVVMBlueprintView.cpp
--- a/Engine/Plugins/Runtime/ModelViewViewModel/Source/ModelViewViewModelBlueprint/Private/MVVMBlueprintView.cpp
+++ b/Engine/Plugins/Runtime/ModelViewViewModel/Source/ModelViewViewModelBlueprint/Private/MVVMBlueprintView.cpp
@@ -552,6 +552,8 @@ void UMVVMBlueprintView::AddAssetTags(FAssetRegistryTagsContext Context) const
 				{
 					Builder << TEXT(',');
 				}
+				Builder << ViewModelContext.GetViewModelName();
+				Builder << TEXT(';');
 				Builder << ViewModelContext.GetViewModelClass()->GetPathName();
 			}
 		}
diff --git a/Engine/Source/Editor/UMGEditor/Private/WidgetBlueprint.cpp b/Engine/Source/Editor/UMGEditor/Private/WidgetBlueprint.cpp
--- a/Engine/Source/Editor/UMGEditor/Private/WidgetBlueprint.cpp
+++ b/Engine/Source/Editor/UMGEditor/Private/WidgetBlueprint.cpp
@@ -741,6 +741,25 @@ void UWidgetBlueprint::GetAssetRegistryTags(FAssetRegistryTagsContext Context) c
 		Context.AddTag(MoveTemp(Tag));
 	}
 	PRAGMA_ENABLE_DEPRECATION_WARNINGS;
+
+	// Add AvailableNamedSlots,  also available on generated class, to the WidgetBlueprint
+	if (const UWidgetBlueprintGeneratedClass* WidgetBPGeneratedClass = Cast<const UWidgetBlueprintGeneratedClass>(GeneratedClass))
+	{
+		TStringBuilder<512> Builder;
+		for (FName NamedSlot : WidgetBPGeneratedClass->AvailableNamedSlots)
+		{
+			if (!NamedSlot.IsNone())
+			{
+				if (Builder.Len() > 0)
+				{
+					Builder << TEXT(',');
+				}
+				Builder << NamedSlot;
+			}
+		}
+		Context.AddTag(FAssetRegistryTag(FName("AvailableNamedSlots"), Builder.ToString(), FAssetRegistryTag::TT_Hidden));
+	}
+
 	FWidgetBlueprintDelegates::GetAssetTagsWithContext.Broadcast(this, Context);
 }
 
diff --git a/Engine/Source/Editor/UMGEditor/Public/WidgetBlueprint.h b/Engine/Source/Editor/UMGEditor/Public/WidgetBlueprint.h
--- a/Engine/Source/Editor/UMGEditor/Public/WidgetBlueprint.h
+++ b/Engine/Source/Editor/UMGEditor/Public/WidgetBlueprint.h
@@ -263,7 +263,7 @@ public:
 	 * in the CDO of the UUserWidget, but a copy is stored here so that it's available in the serialized 
 	 * Tag data in the asset header for access in the FAssetData.
 	 */
-	UPROPERTY(AssetRegistrySearchable, AssetRegistrySearchable)
+	UPROPERTY(AssetRegistrySearchable)
 	FString PaletteCategory;
 
 	/** Run the initialize event on widget that doesn't have a player context. */
diff --git a/Engine/Source/Runtime/UMG/Public/Blueprint/WidgetBlueprintGeneratedClass.h b/Engine/Source/Runtime/UMG/Public/Blueprint/WidgetBlueprintGeneratedClass.h
--- a/Engine/Source/Runtime/UMG/Public/Blueprint/WidgetBlueprintGeneratedClass.h
+++ b/Engine/Source/Runtime/UMG/Public/Blueprint/WidgetBlueprintGeneratedClass.h
@@ -130,7 +130,7 @@ public:
 	 * Available Named Slots for content in a subclass.  These are slots that are accumulated from all super
 	 * classes on compile.  They will exclude any named slots that are filled by a parent class.
 	 **/
-	UPROPERTY()
+	UPROPERTY(AssetRegistrySearchable)
 	TArray<FName> AvailableNamedSlots;
 
 	/**
