Dragging a widget in the UserWidget hierarchy view, will now expand collapsed entries when the user holds over them briefly.

diff --git a/Engine/Source/Editor/UMGEditor/Private/Hierarchy/SHierarchyViewItem.cpp b/Engine/Source/Editor/UMGEditor/Private/Hierarchy/SHierarchyViewItem.cpp
--- a/Engine/Source/Editor/UMGEditor/Private/Hierarchy/SHierarchyViewItem.cpp
+++ b/Engine/Source/Editor/UMGEditor/Private/Hierarchy/SHierarchyViewItem.cpp
@@ -1641,11 +1641,32 @@ bool SHierarchyViewItem::ShouldAppearHovered() const
 void SHierarchyViewItem::HandleDragEnter(FDragDropEvent const& DragDropEvent)
 {
 	Model->HandleDragEnter(DragDropEvent);
+	if (DragHoverExpandTimer.IsValid())
+	{
+		UnRegisterActiveTimer(DragHoverExpandTimer.ToSharedRef());
+		DragHoverExpandTimer.Reset();
+	}
+	if (!IsItemExpanded())
+	{
+		DragHoverExpandTimer = RegisterActiveTimer( 0.3f, FWidgetActiveTimerDelegate::CreateLambda([this](double InCurrentTime, float InDeltaTime)
+		{
+			if (!IsItemExpanded())
+			{
+				ToggleExpansion();
+			}
+			return EActiveTimerReturnType::Stop;
+		}));
+	}
 }
 
 void SHierarchyViewItem::HandleDragLeave(const FDragDropEvent& DragDropEvent)
 {
 	Model->HandleDragLeave(DragDropEvent);
+	if (DragHoverExpandTimer.IsValid())
+	{
+		UnRegisterActiveTimer(DragHoverExpandTimer.ToSharedRef());
+		DragHoverExpandTimer.Reset();
+	}
 }
 
 TOptional<EItemDropZone> SHierarchyViewItem::HandleCanAcceptDrop(const FDragDropEvent& DragDropEvent, EItemDropZone DropZone, TSharedPtr<FHierarchyModel> TargetItem)
diff --git a/Engine/Source/Editor/UMGEditor/Private/Hierarchy/SHierarchyViewItem.h b/Engine/Source/Editor/UMGEditor/Private/Hierarchy/SHierarchyViewItem.h
--- a/Engine/Source/Editor/UMGEditor/Private/Hierarchy/SHierarchyViewItem.h
+++ b/Engine/Source/Editor/UMGEditor/Private/Hierarchy/SHierarchyViewItem.h
@@ -443,4 +443,7 @@ private:
 
 	/** Keep an internal IsHovered flag*/
 	bool bHovered;
+	
+	/** If the user is dragging over a collapsed item, wait before expanding it */
+	TSharedPtr<FActiveTimerHandle> DragHoverExpandTimer;
 };
