UMG: Fix to static analysis warning in WidgeBlueprinttEditor.cpp introduced with CL 32321567

diff --git a/Engine/Source/Editor/UMGEditor/Private/WidgetBlueprintEditor.cpp b/Engine/Source/Editor/UMGEditor/Private/WidgetBlueprintEditor.cpp
--- a/Engine/Source/Editor/UMGEditor/Private/WidgetBlueprintEditor.cpp
+++ b/Engine/Source/Editor/UMGEditor/Private/WidgetBlueprintEditor.cpp
@@ -1170,7 +1170,7 @@ void FWidgetBlueprintEditor::MigrateFromChain(const FPropertyChangedEvent* Prope
 	UWidgetBlueprint* Blueprint = GetWidgetBlueprintObj();
 
 	UUserWidget* PreviewUserWidget = GetPreview();
-	if ( PreviewUserWidget != nullptr )
+	if ( PreviewUserWidget != nullptr && PropertyThatChanged != nullptr)
 	{
 		// The selected objects is always the root widget, which is a copy of the CDO instanced for preview purposes.
 		// if it's modified we need to copy the values back to the CDO - but we ALSO need to do what the details panel
@@ -1191,7 +1191,7 @@ void FWidgetBlueprintEditor::MigrateFromChain(const FPropertyChangedEvent* Prope
 				WidgetCDO->PreEditChange(*PropertyThatChanged);
 			}
 
-			if (PropertyChangedEvent && PropertyThatChanged)
+			if (PropertyChangedEvent)
 			{
 				// We have to do this before we call MigratePropertyValue, because that will change the CDO, and we won't be
 				// able cheat like we do still having a pristine CDO copy.  The details panel has to do shenanigans where it copies
