Widget Thumbnail: update the thumbnail renderer to use the latest API.
Make sure the thumbnail is only rendered on asset save.

diff --git a/Engine/Source/Editor/UMGEditor/Private/WidgetBlueprintThumbnailRenderer.cpp b/Engine/Source/Editor/UMGEditor/Private/WidgetBlueprintThumbnailRenderer.cpp
--- a/Engine/Source/Editor/UMGEditor/Private/WidgetBlueprintThumbnailRenderer.cpp
+++ b/Engine/Source/Editor/UMGEditor/Private/WidgetBlueprintThumbnailRenderer.cpp
@@ -211,9 +211,9 @@ void UWidgetBlueprintThumbnailRenderer::Draw(UObject* Object, int32 X, int32 Y,
 #endif
 }
 
-bool UWidgetBlueprintThumbnailRenderer::AllowsRealtimeThumbnails(UObject* Object) const 
-{ 
-	return false; 
+EThumbnailRenderFrequency UWidgetBlueprintThumbnailRenderer::GetThumbnailRenderFrequency(UObject* Object) const
+{
+	return EThumbnailRenderFrequency::OnAssetSave;
 }
 
 void UWidgetBlueprintThumbnailRenderer::OnBlueprintUnloaded(UBlueprint* Blueprint)
diff --git a/Engine/Source/Editor/UMGEditor/Public/WidgetBlueprintThumbnailRenderer.h b/Engine/Source/Editor/UMGEditor/Public/WidgetBlueprintThumbnailRenderer.h
--- a/Engine/Source/Editor/UMGEditor/Public/WidgetBlueprintThumbnailRenderer.h
+++ b/Engine/Source/Editor/UMGEditor/Public/WidgetBlueprintThumbnailRenderer.h
@@ -20,9 +20,9 @@ class UMGEDITOR_API UWidgetBlueprintThumbnailRenderer : public UDefaultSizedThum
 	virtual ~UWidgetBlueprintThumbnailRenderer();
 
 	//~ Begin UThumbnailRenderer Object
-	bool CanVisualizeAsset(UObject* Object) override;
-	void Draw(UObject* Object, int32 X, int32 Y, uint32 Width, uint32 Height, FRenderTarget* RenderTarget, FCanvas* Canvas, bool bAdditionalViewFamily) override;
-	bool AllowsRealtimeThumbnails(UObject* Object) const override;
+	virtual bool CanVisualizeAsset(UObject* Object) override;
+	virtual void Draw(UObject* Object, int32 X, int32 Y, uint32 Width, uint32 Height, FRenderTarget* RenderTarget, FCanvas* Canvas, bool bAdditionalViewFamily) override;
+	virtual EThumbnailRenderFrequency GetThumbnailRenderFrequency(UObject* Object) const override;
 	//~ End UThumbnailRenderer Object
 
 private:
