Change World that widget previews use to fix Material Parameter Collections not updating correctly.

#tests EUW run MPC & Sequencer test in UMG designer

diff --git a/Engine/Source/Developer/StandaloneRenderer/Private/OpenGL/SlateOpenGLRenderer.cpp b/Engine/Source/Developer/StandaloneRenderer/Private/OpenGL/SlateOpenGLRenderer.cpp
--- a/Engine/Source/Developer/StandaloneRenderer/Private/OpenGL/SlateOpenGLRenderer.cpp
+++ b/Engine/Source/Developer/StandaloneRenderer/Private/OpenGL/SlateOpenGLRenderer.cpp
@@ -354,6 +354,10 @@ int32 FSlateOpenGLRenderer::GetCurrentSceneIndex() const
 	return -1;
 }
 
+void FSlateOpenGLRenderer::SetCurrentSceneIndex(int32 InIndex)
+{
+	// This is a no-op
+}
 
 void FSlateOpenGLRenderer::ClearScenes() 
 {
diff --git a/Engine/Source/Developer/StandaloneRenderer/Private/OpenGL/SlateOpenGLRenderer.h b/Engine/Source/Developer/StandaloneRenderer/Private/OpenGL/SlateOpenGLRenderer.h
--- a/Engine/Source/Developer/StandaloneRenderer/Private/OpenGL/SlateOpenGLRenderer.h
+++ b/Engine/Source/Developer/StandaloneRenderer/Private/OpenGL/SlateOpenGLRenderer.h
@@ -170,6 +170,7 @@ public:
 	virtual FCriticalSection* GetResourceCriticalSection() override;
 	virtual int32 RegisterCurrentScene(FSceneInterface* Scene) override;
 	virtual int32 GetCurrentSceneIndex() const override;
+	virtual void SetCurrentSceneIndex(int32 InIndex) override;
 	virtual void ClearScenes() override;
 
 private:
diff --git a/Engine/Source/Developer/StandaloneRenderer/Private/Windows/D3D/SlateD3DRenderer.cpp b/Engine/Source/Developer/StandaloneRenderer/Private/Windows/D3D/SlateD3DRenderer.cpp
--- a/Engine/Source/Developer/StandaloneRenderer/Private/Windows/D3D/SlateD3DRenderer.cpp
+++ b/Engine/Source/Developer/StandaloneRenderer/Private/Windows/D3D/SlateD3DRenderer.cpp
@@ -282,6 +282,11 @@ int32 FSlateD3DRenderer::GetCurrentSceneIndex() const
 	return -1;
 }
 
+void FSlateD3DRenderer::SetCurrentSceneIndex(int32 InIndex)
+{
+	// This is a no-op
+}
+
 void FSlateD3DRenderer::ClearScenes() 
 {
 	// This is a no-op
diff --git a/Engine/Source/Developer/StandaloneRenderer/Private/Windows/D3D/SlateD3DRenderer.h b/Engine/Source/Developer/StandaloneRenderer/Private/Windows/D3D/SlateD3DRenderer.h
--- a/Engine/Source/Developer/StandaloneRenderer/Private/Windows/D3D/SlateD3DRenderer.h
+++ b/Engine/Source/Developer/StandaloneRenderer/Private/Windows/D3D/SlateD3DRenderer.h
@@ -74,6 +74,7 @@ public:
 	virtual FCriticalSection* GetResourceCriticalSection() override;
 	virtual int32 RegisterCurrentScene(FSceneInterface* Scene) override;
 	virtual int32 GetCurrentSceneIndex() const override;
+	virtual void SetCurrentSceneIndex(int32 InIndex) override;
 	virtual void ClearScenes() override;
 	
 	bool CreateDevice();
diff --git a/Engine/Source/Editor/UMGEditor/Private/Designer/SDesignSurface.cpp b/Engine/Source/Editor/UMGEditor/Private/Designer/SDesignSurface.cpp
--- a/Engine/Source/Editor/UMGEditor/Private/Designer/SDesignSurface.cpp
+++ b/Engine/Source/Editor/UMGEditor/Private/Designer/SDesignSurface.cpp
@@ -204,9 +204,16 @@ FCursorReply SDesignSurface::OnCursorQuery(const FGeometry& MyGeometry, const FP
 
 int32 SDesignSurface::OnPaint(const FPaintArgs& Args, const FGeometry& AllottedGeometry, const FSlateRect& MyCullingRect, FSlateWindowElementList& OutDrawElements, int32 LayerId, const FWidgetStyle& InWidgetStyle, bool bParentEnabled) const
 {
+	// Store the current Scene Index, as child classes may want to use a different scene to preview widgets (Ex: SDesignerView)
+	int32 SceneIndex = FSlateApplication::Get().GetRenderer()->GetCurrentSceneIndex();
+
 	OnPaintBackground(AllottedGeometry, MyCullingRect, OutDrawElements, LayerId);
+	int32 Result = SCompoundWidget::OnPaint(Args, AllottedGeometry, MyCullingRect, OutDrawElements, LayerId, InWidgetStyle, bParentEnabled);
+
+	// Restore previous Scene Index to avoid permanently changing the scene
+	FSlateApplication::Get().GetRenderer()->SetCurrentSceneIndex(SceneIndex);
 
-	return SCompoundWidget::OnPaint(Args, AllottedGeometry, MyCullingRect, OutDrawElements, LayerId, InWidgetStyle, bParentEnabled);
+	return Result;
 }
 
 void SDesignSurface::OnPaintBackground(const FGeometry& AllottedGeometry, const FSlateRect& MyCullingRect, FSlateWindowElementList& OutDrawElements, int32 LayerId) const
diff --git a/Engine/Source/Editor/UMGEditor/Private/Designer/SDesignerView.cpp b/Engine/Source/Editor/UMGEditor/Private/Designer/SDesignerView.cpp
--- a/Engine/Source/Editor/UMGEditor/Private/Designer/SDesignerView.cpp
+++ b/Engine/Source/Editor/UMGEditor/Private/Designer/SDesignerView.cpp
@@ -62,6 +62,7 @@
 #include "WidgetBlueprintEditorUtils.h"
 
 #include "ObjectEditorUtils.h"
+#include "PreviewScene.h"
 #include "ScopedTransaction.h"
 #include "Components/NamedSlot.h"
 
@@ -2408,6 +2409,19 @@ void SDesignerView::Tick( const FGeometry& AllottedGeometry, const double InCurr
 
 void SDesignerView::OnPaintBackground(const FGeometry& AllottedGeometry, const FSlateRect& MyCullingRect, FSlateWindowElementList& OutDrawElements, int32 LayerId) const
 {
+	// In order to get material parameter collections to function properly, we need the editor preview world's Scene
+	// properly propagated through to any widgets that depend on that functionality. 
+	if (TSharedPtr<FWidgetBlueprintEditor> BlueprintEditorPtr = BlueprintEditor.Pin())
+	{
+		if (FPreviewScene* Preview = BlueprintEditorPtr->GetPreviewScene())
+		{
+			if (FSceneInterface* Scene = Preview->GetWorld()->Scene)
+			{
+				FSlateApplication::Get().GetRenderer()->RegisterCurrentScene(Scene);
+			}
+		}
+	}
+
 	SDesignSurface::OnPaintBackground(AllottedGeometry, MyCullingRect, OutDrawElements, LayerId);
 
 	if (bShowResolutionOutlines)
diff --git a/Engine/Source/Runtime/SlateCore/Public/Rendering/SlateRenderer.h b/Engine/Source/Runtime/SlateCore/Public/Rendering/SlateRenderer.h
--- a/Engine/Source/Runtime/SlateCore/Public/Rendering/SlateRenderer.h
+++ b/Engine/Source/Runtime/SlateCore/Public/Rendering/SlateRenderer.h
@@ -538,6 +538,9 @@ public:
 	/** Get the currently registered scene index (set by RegisterCurrentScene)*/
 	virtual int32 GetCurrentSceneIndex() const  = 0;
 
+	/** Set currently registered scene index */
+	virtual void SetCurrentSceneIndex(int32 InIndex) = 0;
+
 	/** Reset the internal Scene tracking.*/
 	virtual void ClearScenes() = 0;
 
diff --git a/Engine/Source/Runtime/SlateNullRenderer/Private/SlateNullRenderer.cpp b/Engine/Source/Runtime/SlateNullRenderer/Private/SlateNullRenderer.cpp
--- a/Engine/Source/Runtime/SlateNullRenderer/Private/SlateNullRenderer.cpp
+++ b/Engine/Source/Runtime/SlateNullRenderer/Private/SlateNullRenderer.cpp
@@ -131,6 +131,11 @@ int32 FSlateNullRenderer::GetCurrentSceneIndex() const
 	return -1;
 }
 
+void FSlateNullRenderer::SetCurrentSceneIndex(int32 InIndex)
+{
+	// This is a no-op
+}
+
 void FSlateNullRenderer::ClearScenes() 
 {
 	// This is a no-op
diff --git a/Engine/Source/Runtime/SlateNullRenderer/Private/SlateNullRenderer.h b/Engine/Source/Runtime/SlateNullRenderer/Private/SlateNullRenderer.h
--- a/Engine/Source/Runtime/SlateNullRenderer/Private/SlateNullRenderer.h
+++ b/Engine/Source/Runtime/SlateNullRenderer/Private/SlateNullRenderer.h
@@ -43,6 +43,7 @@ public:
 	virtual FCriticalSection* GetResourceCriticalSection() override;
 	virtual int32 RegisterCurrentScene(FSceneInterface* Scene) override;
 	virtual int32 GetCurrentSceneIndex() const override;
+	virtual void SetCurrentSceneIndex(int32 InIndex) override;
 	virtual void ClearScenes() override;
 	virtual void Sync() const override;
 
diff --git a/Engine/Source/Runtime/SlateRHIRenderer/Private/SlateRHIRenderer.cpp b/Engine/Source/Runtime/SlateRHIRenderer/Private/SlateRHIRenderer.cpp
--- a/Engine/Source/Runtime/SlateRHIRenderer/Private/SlateRHIRenderer.cpp
+++ b/Engine/Source/Runtime/SlateRHIRenderer/Private/SlateRHIRenderer.cpp
@@ -2289,6 +2289,11 @@ int32 FSlateRHIRenderer::GetCurrentSceneIndex() const
 	return CurrentSceneIndex;
 }
 
+void FSlateRHIRenderer::SetCurrentSceneIndex(int32 InIndex)
+{
+	CurrentSceneIndex = InIndex;
+}
+
 void FSlateRHIRenderer::ClearScenes()
 {
 	if (!IsInSlateThread())
diff --git a/Engine/Source/Runtime/SlateRHIRenderer/Private/SlateRHIRenderer.h b/Engine/Source/Runtime/SlateRHIRenderer/Private/SlateRHIRenderer.h
--- a/Engine/Source/Runtime/SlateRHIRenderer/Private/SlateRHIRenderer.h
+++ b/Engine/Source/Runtime/SlateRHIRenderer/Private/SlateRHIRenderer.h
@@ -215,6 +215,7 @@ public:
 	virtual void ReleaseAccessedResources(bool bImmediatelyFlush) override;
 	virtual int32 RegisterCurrentScene(FSceneInterface* Scene) override;
 	virtual int32 GetCurrentSceneIndex() const override;
+	virtual void SetCurrentSceneIndex(int32 InIndex) override;
 	virtual void ClearScenes() override;
 	EPixelFormat GetSlateRecommendedColorFormat() override;
 	virtual void DestroyCachedFastPathRenderingData(struct FSlateCachedFastPathRenderingData* InRenderingData) override;
