[UMG] Fix some WeakObjectPtr derefs that were lacking null checks.

diff --git a/Engine/Source/Editor/UMGEditor/Private/Customizations/DynamicEntryBoxDetails.cpp b/Engine/Source/Editor/UMGEditor/Private/Customizations/DynamicEntryBoxDetails.cpp
--- a/Engine/Source/Editor/UMGEditor/Private/Customizations/DynamicEntryBoxDetails.cpp
+++ b/Engine/Source/Editor/UMGEditor/Private/Customizations/DynamicEntryBoxDetails.cpp
@@ -48,23 +48,39 @@ void FDynamicEntryBoxBaseDetails::CustomizeDetails(IDetailLayoutBuilder& DetailL
 
 bool FDynamicEntryBoxBaseDetails::CanEditSpacingPattern() const
 {
-	return EntryBox->GetBoxType() == EDynamicBoxType::Overlay;
+	if (EntryBox.IsValid())
+	{
+		return EntryBox->GetBoxType() == EDynamicBoxType::Overlay;
+	}
+	return false;
 }
 
 bool FDynamicEntryBoxBaseDetails::CanEditEntrySpacing() const
 {
-	return EntryBox->SpacingPattern.Num() == 0;
+	if (EntryBox.IsValid())
+	{
+		return EntryBox->SpacingPattern.Num() == 0;
+	}
+	return false;
 }
 
 bool FDynamicEntryBoxBaseDetails::CanEditAlignment() const
 {
-	return EntryBox->GetBoxType() != EDynamicBoxType::Overlay || CanEditEntrySpacing();
+	if (EntryBox.IsValid())
+	{
+		return EntryBox->GetBoxType() != EDynamicBoxType::Overlay || CanEditEntrySpacing();
+	}
+	return false;
 }
 
 bool FDynamicEntryBoxBaseDetails::CanEditMaxElementSize() const
 {
-	const EDynamicBoxType BoxType = EntryBox->GetBoxType();
-	return BoxType == EDynamicBoxType::Horizontal || BoxType == EDynamicBoxType::Vertical;
+	if (EntryBox.IsValid())
+	{
+		const EDynamicBoxType BoxType = EntryBox->GetBoxType();
+		return BoxType == EDynamicBoxType::Horizontal || BoxType == EDynamicBoxType::Vertical;
+	}
+	return false;
 }
 
 //////////////////////////////////////////////////////////////////////////
