UMG: make widget details panel compatible with AddExternalObjects.

diff --git a/Engine/Source/Editor/UMGEditor/Private/WidgetBlueprintEditor.cpp b/Engine/Source/Editor/UMGEditor/Private/WidgetBlueprintEditor.cpp
--- a/Engine/Source/Editor/UMGEditor/Private/WidgetBlueprintEditor.cpp
+++ b/Engine/Source/Editor/UMGEditor/Private/WidgetBlueprintEditor.cpp
@@ -1201,8 +1201,16 @@ void FWidgetBlueprintEditor::MigrateFromChain(const FPropertyChangedEvent* Prope
 				PropagateDefaultPropertyChange(*PropertyChangedEvent, *PropertyThatChanged, WidgetCDO, ObjectRef.Get());
 			}
 
-			// dealing with root widget here
-			MigratePropertyValue(ObjectRef.Get(), WidgetCDO, PropertyChainNode, PropertyChainNode->GetValue(), bIsModify);
+			// Only do the migration if the property on the head of the linked list lives in the CDO
+			// We want to skip the migration if the property isn't leading to this CDO, such as when we call 
+			// FDetailCategoryImpl::AddExternalObjects to inject external objects in the details panel of the CDO.
+			// We don't prevent invalid source and destination objects here as those will be handled in MigratePropertyValue.
+			if (!PropertyChainNode->GetValue() || ObjectRef.Get() == nullptr || WidgetCDO == nullptr ||
+				(ObjectRef.Get()->IsA(PropertyChainNode->GetValue()->GetOwnerClass()) && WidgetCDO->IsA(PropertyChainNode->GetValue()->GetOwnerClass())))
+			{
+				// dealing with root widget here
+				MigratePropertyValue(ObjectRef.Get(), WidgetCDO, PropertyChainNode, PropertyChainNode->GetValue(), bIsModify);
+			}
 
 			if (PropertyChangedEvent && WidgetCDO)
 			{
@@ -1245,7 +1253,16 @@ void FWidgetBlueprintEditor::MigrateFromChain(const FPropertyChangedEvent* Prope
 				if ( TemplateWidget )
 				{
 					FEditPropertyChain::TDoubleLinkedListNode* PropertyChainNode = PropertyThatChanged->GetHead();
-					MigratePropertyValue(PreviewWidget, TemplateWidget, PropertyChainNode, PropertyChainNode->GetValue(), bIsModify);
+
+					// Only do the migration if the property on the head of the linked list lives in this UWidget
+					// We want to skip the migration if the property isn't leading to this UWidget, such as when we call 
+					// FDetailCategoryImpl::AddExternalObjects to inject external objects in the details panel of this UWidget.
+					// We don't prevent invalid source and destination objects here as those will be handled in MigratePropertyValue.
+					if (!PropertyChainNode->GetValue() ||
+						(PreviewWidget->IsA(PropertyChainNode->GetValue()->GetOwnerClass()) && TemplateWidget->IsA(PropertyChainNode->GetValue()->GetOwnerClass())))
+					{
+						MigratePropertyValue(PreviewWidget, TemplateWidget, PropertyChainNode, PropertyChainNode->GetValue(), bIsModify);
+					}
 				}
 			}
 		}
