WidgetBlueprintEditor: Fix warning for using weak pointer without checking for null first.

diff --git a/Engine/Source/Editor/UMGEditor/Private/WidgetBlueprintEditor.cpp b/Engine/Source/Editor/UMGEditor/Private/WidgetBlueprintEditor.cpp
--- a/Engine/Source/Editor/UMGEditor/Private/WidgetBlueprintEditor.cpp
+++ b/Engine/Source/Editor/UMGEditor/Private/WidgetBlueprintEditor.cpp
@@ -1181,63 +1181,68 @@ void FWidgetBlueprintEditor::MigrateFromChain(const FPropertyChangedEvent* Prope
 			//TODO: We should probably ensure that the objectref is in fact the proxy CDO for the Widget, we should
 			// probably keep track of that somewhere and vet this - you never know if someone might make new changes
 			// that invalidate this assumption.
+			UObject* ObjectPtr = ObjectRef.Get();
+			check(ObjectPtr);
 
-			FEditPropertyChain::TDoubleLinkedListNode* PropertyChainNode = PropertyThatChanged->GetHead();
-			UObject* WidgetCDO = ObjectRef.Get()->GetClass()->GetDefaultObject(true);
-
-			if (PropertyChangedEvent == nullptr && WidgetCDO)
+			if (ObjectPtr)
 			{
-				WidgetCDO->SetEditChangePropagationFlags(EEditChangePropagationFlags::OnlyMarkRealignedInstancesAsDirty);
-				WidgetCDO->PreEditChange(*PropertyThatChanged);
-			}
+				FEditPropertyChain::TDoubleLinkedListNode* PropertyChainNode = PropertyThatChanged->GetHead();
+				UObject* WidgetCDO = ObjectPtr->GetClass()->GetDefaultObject(true);
 
-			if (PropertyChangedEvent)
-			{
-				// We have to do this before we call MigratePropertyValue, because that will change the CDO, and we won't be
-				// able cheat like we do still having a pristine CDO copy.  The details panel has to do shenanigans where it copies
-				// the data into a temp property piece of memory, but we don't need to do that since we've got the actual unmodified
-				// CDO still - but we're going to propagate the proxy's value that users modified instead of the CDO's value,
-				// then MigratePropertyValue will update the CDO's value.
-				PropagateDefaultPropertyChange(*PropertyChangedEvent, *PropertyThatChanged, WidgetCDO, ObjectRef.Get());
-			}
+				if (PropertyChangedEvent == nullptr && WidgetCDO)
+				{
+					WidgetCDO->SetEditChangePropagationFlags(EEditChangePropagationFlags::OnlyMarkRealignedInstancesAsDirty);
+					WidgetCDO->PreEditChange(*PropertyThatChanged);
+				}
 
-			// Only do the migration if the property on the head of the linked list lives in the CDO
-			// We want to skip the migration if the property isn't leading to this CDO, such as when we call 
-			// FDetailCategoryImpl::AddExternalObjects to inject external objects in the details panel of the CDO.
-			// We don't prevent invalid source and destination objects here as those will be handled in MigratePropertyValue.
-			if (!PropertyChainNode->GetValue() || ObjectRef.Get() == nullptr || WidgetCDO == nullptr ||
-				(ObjectRef.Get()->IsA(PropertyChainNode->GetValue()->GetOwnerClass()) && WidgetCDO->IsA(PropertyChainNode->GetValue()->GetOwnerClass())))
-			{
-				// dealing with root widget here
-				MigratePropertyValue(ObjectRef.Get(), WidgetCDO, PropertyChainNode, PropertyChainNode->GetValue(), bIsModify);
-			}
+				if (PropertyChangedEvent)
+				{
+					// We have to do this before we call MigratePropertyValue, because that will change the CDO, and we won't be
+					// able cheat like we do still having a pristine CDO copy.  The details panel has to do shenanigans where it copies
+					// the data into a temp property piece of memory, but we don't need to do that since we've got the actual unmodified
+					// CDO still - but we're going to propagate the proxy's value that users modified instead of the CDO's value,
+					// then MigratePropertyValue will update the CDO's value.
+					PropagateDefaultPropertyChange(*PropertyChangedEvent, *PropertyThatChanged, WidgetCDO, ObjectPtr);
+				}
 
-			if (PropertyChangedEvent && WidgetCDO)
-			{
-				// 
-				TArray<UObject*> CDOArray = { WidgetCDO };
-				FPropertyChangedEvent ChangeEvent(PropertyChangedEvent->Property, PropertyChangedEvent->ChangeType, CDOArray);
-				ChangeEvent.SetActiveMemberProperty(PropertyChangedEvent->MemberProperty);
+				// Only do the migration if the property on the head of the linked list lives in the CDO
+				// We want to skip the migration if the property isn't leading to this CDO, such as when we call 
+				// FDetailCategoryImpl::AddExternalObjects to inject external objects in the details panel of the CDO.
+				// We don't prevent invalid source and destination objects here as those will be handled in MigratePropertyValue.
+				if (!PropertyChainNode->GetValue() || ObjectRef.Get() == nullptr || WidgetCDO == nullptr ||
+					(ObjectPtr->IsA(PropertyChainNode->GetValue()->GetOwnerClass()) && WidgetCDO->IsA(PropertyChainNode->GetValue()->GetOwnerClass())))
+				{
+					// dealing with root widget here
+					MigratePropertyValue(ObjectPtr, WidgetCDO, PropertyChainNode, PropertyChainNode->GetValue(), bIsModify);
+				}
 
-				TArray<TMap<FString, int32>> ArrayIndexForProperty;
-				for (TDoubleLinkedList<FProperty*>::TIterator It(PropertyThatChanged->GetHead()); It; ++It)
+				if (PropertyChangedEvent && WidgetCDO)
 				{
-					FProperty* Property = *It;
+					// 
+					TArray<UObject*> CDOArray = { WidgetCDO };
+					FPropertyChangedEvent ChangeEvent(PropertyChangedEvent->Property, PropertyChangedEvent->ChangeType, CDOArray);
+					ChangeEvent.SetActiveMemberProperty(PropertyChangedEvent->MemberProperty);
 
-					const FString PropertyName = Property->GetName();
-					const int32 ArrayIndex = PropertyChangedEvent->GetArrayIndex(PropertyName);
-					if (ArrayIndex != INDEX_NONE)
+					TArray<TMap<FString, int32>> ArrayIndexForProperty;
+					for (TDoubleLinkedList<FProperty*>::TIterator It(PropertyThatChanged->GetHead()); It; ++It)
 					{
-						ArrayIndexForProperty.AddDefaulted();
-						ArrayIndexForProperty.Last().Add(PropertyName, ArrayIndex);
+						FProperty* Property = *It;
+
+						const FString PropertyName = Property->GetName();
+						const int32 ArrayIndex = PropertyChangedEvent->GetArrayIndex(PropertyName);
+						if (ArrayIndex != INDEX_NONE)
+						{
+							ArrayIndexForProperty.AddDefaulted();
+							ArrayIndexForProperty.Last().Add(PropertyName, ArrayIndex);
+						}
 					}
-				}
-				ChangeEvent.ObjectIteratorIndex = 0;
-				ChangeEvent.SetArrayIndexPerObject(ArrayIndexForProperty);
+					ChangeEvent.ObjectIteratorIndex = 0;
+					ChangeEvent.SetArrayIndexPerObject(ArrayIndexForProperty);
 
-				FPropertyChangedChainEvent ChainEvent(*PropertyThatChanged, ChangeEvent);
-				WidgetCDO->PostEditChangeChainProperty(ChainEvent);
-				WidgetCDO->SetEditChangePropagationFlags(EEditChangePropagationFlags::None);
+					FPropertyChangedChainEvent ChainEvent(*PropertyThatChanged, ChangeEvent);
+					WidgetCDO->PostEditChangeChainProperty(ChainEvent);
+					WidgetCDO->SetEditChangePropagationFlags(EEditChangePropagationFlags::None);
+				}
 			}
 		}
 
