WidgetBlueprintEditor: Re-add CL 32906781 that allows for using AddExternalObjects in the details panel of widgets.

diff --git a/Engine/Source/Editor/UMGEditor/Private/WidgetBlueprintEditor.cpp b/Engine/Source/Editor/UMGEditor/Private/WidgetBlueprintEditor.cpp
--- a/Engine/Source/Editor/UMGEditor/Private/WidgetBlueprintEditor.cpp
+++ b/Engine/Source/Editor/UMGEditor/Private/WidgetBlueprintEditor.cpp
@@ -1082,10 +1082,24 @@ void FWidgetBlueprintEditor::MigrateFromChain(FEditPropertyChain* PropertyThatCh
 	{
 		for ( TWeakObjectPtr<UObject> ObjectRef : SelectedObjects )
 		{
-			// dealing with root widget here
-			FEditPropertyChain::TDoubleLinkedListNode* PropertyChainNode = PropertyThatChanged->GetHead();
-			UObject* WidgetCDO = ObjectRef.Get()->GetClass()->GetDefaultObject(true);
-			MigratePropertyValue(ObjectRef.Get(), WidgetCDO, PropertyChainNode, PropertyChainNode->GetValue(), bIsModify);
+			UObject* ObjectPtr = ObjectRef.Get();
+			check(ObjectPtr);
+
+			if (ObjectPtr)
+			{
+				// dealing with root widget here
+				FEditPropertyChain::TDoubleLinkedListNode* PropertyChainNode = PropertyThatChanged->GetHead();
+				UObject* WidgetCDO = ObjectPtr->GetClass()->GetDefaultObject(true);
+
+				// Only do the migration if the property on the head of the linked list lives in the CDO
+				// We want to skip the migration if the property isn't leading to this CDO, such as when we call 
+				// FDetailCategoryImpl::AddExternalObjects to inject external objects in the details panel of the CDO.
+				if (!PropertyChainNode->GetValue() || WidgetCDO == nullptr ||
+					(ObjectPtr->IsA(PropertyChainNode->GetValue()->GetOwnerClass()) && WidgetCDO->IsA(PropertyChainNode->GetValue()->GetOwnerClass())))
+				{
+					MigratePropertyValue(ObjectPtr, WidgetCDO, PropertyChainNode, PropertyChainNode->GetValue(), bIsModify);
+				}
+			}
 		}
 
 		for ( FWidgetReference& WidgetRef : SelectedWidgets )
@@ -1100,7 +1114,15 @@ void FWidgetBlueprintEditor::MigrateFromChain(FEditPropertyChain* PropertyThatCh
 				if ( TemplateWidget )
 				{
 					FEditPropertyChain::TDoubleLinkedListNode* PropertyChainNode = PropertyThatChanged->GetHead();
-					MigratePropertyValue(PreviewWidget, TemplateWidget, PropertyChainNode, PropertyChainNode->GetValue(), bIsModify);
+
+					// Only do the migration if the property on the head of the linked list lives in this UWidget
+					// We want to skip the migration if the property isn't leading to this UWidget, such as when we call 
+					// FDetailCategoryImpl::AddExternalObjects to inject external objects in the details panel of this UWidget.
+					if (!PropertyChainNode->GetValue() ||
+						(PreviewWidget->IsA(PropertyChainNode->GetValue()->GetOwnerClass()) && TemplateWidget->IsA(PropertyChainNode->GetValue()->GetOwnerClass())))
+					{
+						MigratePropertyValue(PreviewWidget, TemplateWidget, PropertyChainNode, PropertyChainNode->GetValue(), bIsModify);
+					}
 				}
 			}
 		}
