Disable slate postbuffers when launching multi-PIE

#tests PIE, Multi-PIE

diff --git a/Engine/Source/Editor/UMGEditor/Private/UMGEditorModule.cpp b/Engine/Source/Editor/UMGEditor/Private/UMGEditorModule.cpp
--- a/Engine/Source/Editor/UMGEditor/Private/UMGEditorModule.cpp
+++ b/Engine/Source/Editor/UMGEditor/Private/UMGEditorModule.cpp
@@ -34,6 +34,7 @@
 
 #include "ISettingsModule.h"
 #include "SequencerSettings.h"
+#include "Settings/LevelEditorPlaySettings.h"
 
 #include "BlueprintEditorModule.h"
 #include "PropertyEditorModule.h"
@@ -46,6 +47,8 @@
 
 #define LOCTEXT_NAMESPACE "UMG"
 
+DEFINE_LOG_CATEGORY_STATIC(LogUMGEditor, Log, All);
+
 const FName UMGEditorAppIdentifier = FName(TEXT("UMGEditorApp"));
 static TAutoConsoleVariable<bool> CVarThumbnailRenderEnable(
 	TEXT("UMG.ThumbnailRenderer.Enable"),
@@ -61,6 +64,7 @@ public:
 		: Settings(nullptr)
 		, bThumbnailRenderersRegistered(false)
 		, bOnPostEngineInitHandled(false)
+		, bCachedArePostBuffersEnabled(false)
 	{
 	}
 
@@ -71,6 +75,8 @@ public:
 
 		// Any attempt to use GEditor right now will fail as it hasn't been initialized yet. Waiting for post engine init resolves that.
 		FCoreDelegates::OnPostEngineInit.AddRaw(this, &FUMGEditorModule::OnPostEngineInit);
+		FEditorDelegates::StartPIE.AddRaw(this, &FUMGEditorModule::OnStartPIE);
+		FEditorDelegates::EndPIE.AddRaw(this, &FUMGEditorModule::OnEndPIE);
 
 		if (GIsEditor)
 		{
@@ -135,6 +141,9 @@ public:
 	virtual void ShutdownModule() override
 	{
 		FCoreDelegates::OnPostEngineInit.RemoveAll(this);
+		FEditorDelegates::StartPIE.RemoveAll(this);
+		FEditorDelegates::EndPIE.RemoveAll(this);
+
 		FModuleManager::Get().OnModulesChanged().Remove(ModuleChangedHandle);
 
 		if (UObjectInitialized() && bThumbnailRenderersRegistered && IConsoleManager::Get().FindConsoleVariable(TEXT("UMGEditor.ThumbnailRenderer.Enable"))->GetBool())
@@ -365,6 +374,50 @@ private:
 		bOnPostEngineInitHandled = true;
 	}
 
+	void OnStartPIE(const bool bIsSimulating)
+	{
+		if (UEditorEngine* EditorEngine = Cast<UEditorEngine>(GEngine))
+		{
+			if (TOptional<FRequestPlaySessionParams> PlayRequest = EditorEngine->GetPlaySessionRequest())
+			{
+				if (!PlayRequest.GetValue().EditorPlaySettings)
+				{
+					return;
+				}
+
+				int32 NumClients;
+				PlayRequest.GetValue().EditorPlaySettings->GetPlayNumberOfClients(NumClients);
+
+				if (NumClients > 1)
+				{
+					if (IConsoleVariable* PostBuffersEnabled = IConsoleManager::Get().FindConsoleVariable(TEXT("Slate.CopyBackbufferToSlatePostRenderTargets")))
+					{
+						UE_LOG(LogUMGEditor, Warning, TEXT("Disabling Slate Post Buffers for multi-window PIE session, currently not supported."));
+
+						bCachedArePostBuffersEnabled = PostBuffersEnabled->GetBool();
+						PostBuffersEnabled->Set(false);
+					}
+				}
+			}
+		}
+	}
+
+	void OnEndPIE(const bool bIsSimulating)
+	{
+		if (bCachedArePostBuffersEnabled)
+		{
+			if (IConsoleVariable* PostBuffersEnabled = IConsoleManager::Get().FindConsoleVariable(TEXT("Slate.CopyBackbufferToSlatePostRenderTargets")))
+			{
+				if (!PostBuffersEnabled->GetBool())
+				{
+					UE_LOG(LogUMGEditor, Warning, TEXT("Restoring Slate Post Buffers, previously disabled due to multi-window PIE session."));
+
+					PostBuffersEnabled->Set(bCachedArePostBuffersEnabled);
+				}
+			}
+		}
+	}
+
 	static void ThumbnailRenderingEnabled(IConsoleVariable* Variable)
 	{
 		FUMGEditorModule* UMGEditorModule = FModuleManager::GetModulePtr<FUMGEditorModule>(TEXT("UMGEditor"));
@@ -426,6 +479,7 @@ private:
 
 	bool bThumbnailRenderersRegistered;
 	bool bOnPostEngineInitHandled;
+	bool bCachedArePostBuffersEnabled;
 };
 
 IMPLEMENT_MODULE(FUMGEditorModule, UMGEditor);
