Disable drag and dropping children into panel hierarchy when it has an MVVM extension.

IDragDropExtension:
- Extension class providing CanDropOnTarget and GetDropFailureText functions.

FDragDropExtensibilityManager:
- Maintains a list of IDragDropExtensions.

IHasDragDropExtensibility:
- Extension class providing GetDragDropExtensibilityManager function.

IUMGEditorModule:
- Inherits IHasDragDropExtensibility.

FUMGEditorModule:
- Creates FDragDropExtensibilityManager instance on startup.

FDragDropExtension:
- MVVM IDragDropExtension implementation. CanDropOnTarget returns false if the target widget has a UMVVMBlueprintViewExtension_PanelWidget extension.

FModelViewViewModelEditorModule:
- Adds FDragDropExtension on startup and removes on shutdown.

FWidgetBlueprintEditorUtils:
- Added CanDropOnTargetExtensions function which iterates current IDragDropExtensions (via UMG module's FDragDropExtensibilityManager).

ProcessHierarchyDragDrop:
- Checks FWidgetBlueprintEditorUtils::CanDropOnTargetExtensions.

diff --git a/Engine/Plugins/Runtime/ModelViewViewModel/Source/ModelViewViewModelEditor/Private/Customizations/MVVMDragDropExtension.cpp b/Engine/Plugins/Runtime/ModelViewViewModel/Source/ModelViewViewModelEditor/Private/Customizations/MVVMDragDropExtension.cpp
new file mode 100644
--- /dev/null
+++ b/Engine/Plugins/Runtime/ModelViewViewModel/Source/ModelViewViewModelEditor/Private/Customizations/MVVMDragDropExtension.cpp
@@ -0,0 +1,40 @@
+// Copyright Epic Games, Inc. All Rights Reserved.
+
+#include "Customizations/MVVMDragDropExtension.h"
+
+#include "Components/PanelWidget.h"
+#include "Extensions/MVVMViewBlueprintPanelWidgetExtension.h"
+#include "Internationalization/Internationalization.h"
+#include "MVVMWidgetBlueprintExtension_View.h"
+#include "WidgetBlueprint.h"
+#include "WidgetBlueprintEditorUtils.h"
+
+namespace UE::MVVM
+{
+bool FDragDropExtension::CanDropOnTarget(const UWidget* Target, const TSharedPtr<FDragDropOperation>& DragDropOp) const
+{
+	bool bCanDropWidgetOnTarget = true;
+
+	if (Target && Target->IsA<UPanelWidget>())
+	{
+		if (const UWidgetBlueprint* WidgetBlueprint = FWidgetBlueprintEditorUtils::GetWidgetBlueprintFromWidget(Target))
+		{
+			if (const UMVVMWidgetBlueprintExtension_View* ExtensionView = UMVVMWidgetBlueprintExtension_View::GetExtension<UMVVMWidgetBlueprintExtension_View>(WidgetBlueprint))
+			{
+				const TArray<UMVVMBlueprintViewExtension*> BlueprintExtensions = ExtensionView->GetBlueprintExtensionsForWidget(Target->GetFName());
+				bCanDropWidgetOnTarget = !BlueprintExtensions.ContainsByPredicate([](const UMVVMBlueprintViewExtension* BlueprintExtension) -> bool
+				{
+					return BlueprintExtension->IsA<UMVVMBlueprintViewExtension_PanelWidget>();
+				});
+			}
+		}
+	}
+
+	return bCanDropWidgetOnTarget;
+}
+
+FText FDragDropExtension::GetDropFailureText(const UWidget* Target, const TSharedPtr<FDragDropOperation>& DragDropOp) const
+{
+	return NSLOCTEXT("MVVMDragDropExtension", "UnableToAddChildWidget", "Cannot add children to a panel widget with an MVVM extension.");
+}
+}
diff --git a/Engine/Plugins/Runtime/ModelViewViewModel/Source/ModelViewViewModelEditor/Private/Customizations/MVVMDragDropExtension.h b/Engine/Plugins/Runtime/ModelViewViewModel/Source/ModelViewViewModelEditor/Private/Customizations/MVVMDragDropExtension.h
new file mode 100644
--- /dev/null
+++ b/Engine/Plugins/Runtime/ModelViewViewModel/Source/ModelViewViewModelEditor/Private/Customizations/MVVMDragDropExtension.h
@@ -0,0 +1,16 @@
+// Copyright Epic Games, Inc. All Rights Reserved.
+
+#pragma once
+
+#include "IHasDragDropExtensibility.h"
+
+namespace UE::MVVM
+{
+class FDragDropExtension : public IDragDropExtension
+{
+	//~ Begin IDragDropExtension overrides
+	virtual bool CanDropOnTarget(const UWidget* Target, const TSharedPtr<FDragDropOperation>& DragDropOp) const override;
+	virtual FText GetDropFailureText(const UWidget* Target, const TSharedPtr<FDragDropOperation>& DragDropOp) const override;
+	//~ End IDragDropExtension overrides
+};
+}
diff --git a/Engine/Plugins/Runtime/ModelViewViewModel/Source/ModelViewViewModelEditor/Private/ModelViewViewModelEditorModule.cpp b/Engine/Plugins/Runtime/ModelViewViewModel/Source/ModelViewViewModelEditor/Private/ModelViewViewModelEditorModule.cpp
--- a/Engine/Plugins/Runtime/ModelViewViewModel/Source/ModelViewViewModelEditor/Private/ModelViewViewModelEditorModule.cpp
+++ b/Engine/Plugins/Runtime/ModelViewViewModel/Source/ModelViewViewModelEditor/Private/ModelViewViewModelEditorModule.cpp
@@ -8,6 +8,7 @@
 #include "Customizations/MVVMBlueprintViewDesignerExtension.h"
 #include "Customizations/MVVMBlueprintViewModelContextCustomization.h"
 #include "Customizations/MVVMClipboardExtension.h"
+#include "Customizations/MVVMDragDropExtension.h"
 #include "Customizations/MVVMListViewBaseExtensionCustomizationExtender.h"
 #include "Customizations/MVVMPanelWidgetExtensionCustomizationExtender.h"
 #include "Customizations/MVVMPropertyBindingExtension.h"
@@ -52,6 +53,9 @@ void FModelViewViewModelEditorModule::StartupModule()
 	ClipboardExtension = MakeShared<UE::MVVM::FClipboardExtension>();
 	UMGEditorModule.GetClipboardExtensibilityManager()->AddExtension(ClipboardExtension.ToSharedRef());
 
+	DragDropExtension = MakeShared<UE::MVVM::FDragDropExtension>();
+	UMGEditorModule.GetDragDropExtensibilityManager()->AddExtension(DragDropExtension.ToSharedRef());
+
 	ListViewBaseCustomizationExtender = UE::MVVM::FMVVMListViewBaseExtensionCustomizationExtender::MakeInstance();
 	UMGEditorModule.AddWidgetCustomizationExtender(ListViewBaseCustomizationExtender.ToSharedRef());
 
@@ -102,6 +106,7 @@ void FModelViewViewModelEditorModule::ShutdownModule()
 	{
 		UMGEditorModule->GetDesignerExtensibilityManager()->RemoveDesignerExtensionFactory(BlueprintViewDesignerExtensionFactory.ToSharedRef());
 		UMGEditorModule->OnRegisterTabsForEditor().RemoveAll(this);
+		UMGEditorModule->GetDragDropExtensibilityManager()->RemoveExtension(DragDropExtension.ToSharedRef());
 		UMGEditorModule->GetClipboardExtensibilityManager()->RemoveExtension(ClipboardExtension.ToSharedRef());
 		UMGEditorModule->GetPropertyBindingExtensibilityManager()->RemoveExtension(PropertyBindingExtension.ToSharedRef());
 
diff --git a/Engine/Plugins/Runtime/ModelViewViewModel/Source/ModelViewViewModelEditor/Private/ModelViewViewModelEditorModule.h b/Engine/Plugins/Runtime/ModelViewViewModel/Source/ModelViewViewModelEditor/Private/ModelViewViewModelEditorModule.h
--- a/Engine/Plugins/Runtime/ModelViewViewModel/Source/ModelViewViewModelEditor/Private/ModelViewViewModelEditorModule.h
+++ b/Engine/Plugins/Runtime/ModelViewViewModel/Source/ModelViewViewModelEditor/Private/ModelViewViewModelEditorModule.h
@@ -10,6 +10,7 @@ namespace UE::MVVM
 	class FAssetTypeActions_ViewModelBlueprint;
 	class FBlueprintViewDesignerExtensionFactory;
 	class FClipboardExtension;
+	class FDragDropExtension;
 	class FMVVMListViewBaseExtensionCustomizationExtender;
 	class FMVVMPanelWidgetExtensionCustomizationExtender;
 	class FMVVMPropertyBindingExtension;
@@ -48,6 +49,7 @@ private:
 private:
 	TSharedPtr<UE::MVVM::FMVVMPropertyBindingExtension> PropertyBindingExtension;
 	TSharedPtr<UE::MVVM::FClipboardExtension> ClipboardExtension;
+	TSharedPtr<UE::MVVM::FDragDropExtension> DragDropExtension;
 	TSharedPtr<UE::MVVM::FAssetTypeActions_ViewModelBlueprint> ViewModelBlueprintActions;
 	TSharedPtr<UE::MVVM::FMVVMListViewBaseExtensionCustomizationExtender> ListViewBaseCustomizationExtender;
 	TSharedPtr<UE::MVVM::FMVVMPanelWidgetExtensionCustomizationExtender> PanelWidgetCustomizationExtender;
diff --git a/Engine/Source/Editor/UMGEditor/Private/Hierarchy/SHierarchyViewItem.cpp b/Engine/Source/Editor/UMGEditor/Private/Hierarchy/SHierarchyViewItem.cpp
--- a/Engine/Source/Editor/UMGEditor/Private/Hierarchy/SHierarchyViewItem.cpp
+++ b/Engine/Source/Editor/UMGEditor/Private/Hierarchy/SHierarchyViewItem.cpp
@@ -188,6 +188,20 @@ TOptional<EItemDropZone> ProcessHierarchyDragDrop(const FDragDropEvent& DragDrop
 	UWidgetBlueprint* Blueprint = BlueprintEditor->GetWidgetBlueprintObj();
 	check( Blueprint != nullptr && Blueprint->WidgetTree != nullptr );
 
+	const auto CanDropOnTargetExtensions = [](UWidget* Target, const TSharedPtr<FDecoratedDragDropOp>& DecoratedDragDropOp) -> bool
+		{
+			FText CanDropOnTargetFailureText = FText::GetEmpty();
+			const bool bCanDropOnTargetExtensions = FWidgetBlueprintEditorUtils::CanDropOnTargetExtensions(Target, DecoratedDragDropOp, CanDropOnTargetFailureText);
+
+			if (!bCanDropOnTargetExtensions && DecoratedDragDropOp.IsValid())
+			{
+				DecoratedDragDropOp->CurrentIconBrush = FAppStyle::GetBrush(TEXT("Graph.ConnectorFeedback.Error"));
+				DecoratedDragDropOp->CurrentHoverText = CanDropOnTargetFailureText;
+			}
+
+			return bCanDropOnTargetExtensions;
+		};
+
 	// Is this a drag/drop op to create a new widget in the tree?
 	TSharedPtr<FDragDropOperation> DragDropOp = DragDropEvent.GetOperation();
 	if (DragDropOp.IsValid() && !DragDropOp->IsOfType<FHierarchyWidgetDragDropOpImpl>())
@@ -199,6 +213,11 @@ TOptional<EItemDropZone> ProcessHierarchyDragDrop(const FDragDropEvent& DragDrop
 			DecoratedDragDropOp->ResetToDefaultToolTip();
 		}
 
+		if (!CanDropOnTargetExtensions(TargetTemplate, DecoratedDragDropOp))
+		{
+			return TOptional<EItemDropZone>();
+		}
+
 		// Are we adding to a locked widget?
 		if ( TargetItem.IsValid() && TargetItem.GetPreview()->IsLockedInDesigner() )
 		{
@@ -340,6 +359,11 @@ TOptional<EItemDropZone> ProcessHierarchyDragDrop(const FDragDropEvent& DragDrop
 				return TOptional<EItemDropZone>();
 			}
 
+			if (!CanDropOnTargetExtensions(TargetTemplate, HierarchyDragDropOp))
+			{
+				return TOptional<EItemDropZone>();
+			}
+
 			bool bFoundNewParentInChildSet = false;
 			for (const auto& DraggedWidget : HierarchyDragDropOp->DraggedWidgets)
 			{
diff --git a/Engine/Source/Editor/UMGEditor/Private/UMGEditorModule.cpp b/Engine/Source/Editor/UMGEditor/Private/UMGEditorModule.cpp
--- a/Engine/Source/Editor/UMGEditor/Private/UMGEditorModule.cpp
+++ b/Engine/Source/Editor/UMGEditor/Private/UMGEditorModule.cpp
@@ -92,6 +92,7 @@ public:
 
 		PropertyBindingExtensibilityManager = MakeShared<FPropertyBindingExtensibilityManager>();
 		ClipboardExtensibilityManager = MakeShared<FClipboardExtensibilityManager>();
+		DragDropExtensibilityManager = MakeShared<FDragDropExtensibilityManager>();
 
 		// Register widget blueprint compiler we do this no matter what.
 		IKismetCompilerInterface& KismetCompilerModule = FModuleManager::LoadModuleChecked<IKismetCompilerInterface>("KismetCompiler");
@@ -223,6 +224,7 @@ public:
 	virtual TSharedPtr<FDesignerExtensibilityManager> GetDesignerExtensibilityManager() override { return DesignerExtensibilityManager; }
 	virtual TSharedPtr<FPropertyBindingExtensibilityManager> GetPropertyBindingExtensibilityManager() override { return PropertyBindingExtensibilityManager; }
 	virtual TSharedPtr<FClipboardExtensibilityManager> GetClipboardExtensibilityManager() override { return ClipboardExtensibilityManager; }
+	virtual TSharedPtr<FDragDropExtensibilityManager> GetDragDropExtensibilityManager() override { return DragDropExtensibilityManager; }
 
 	/** Register settings objects. */
 	void RegisterSettings()
@@ -440,6 +442,7 @@ private:
 	TSharedPtr<FDesignerExtensibilityManager> DesignerExtensibilityManager;
 	TSharedPtr<FPropertyBindingExtensibilityManager> PropertyBindingExtensibilityManager;
 	TSharedPtr<FClipboardExtensibilityManager> ClipboardExtensibilityManager;
+	TSharedPtr<FDragDropExtensibilityManager> DragDropExtensibilityManager;
 	TSharedPtr<FGraphPanelPinFactory> GraphPanelPinFactory;
 
 	FDelegateHandle SequenceEditorHandle;
diff --git a/Engine/Source/Editor/UMGEditor/Private/WidgetBlueprintEditorUtils.cpp b/Engine/Source/Editor/UMGEditor/Private/WidgetBlueprintEditorUtils.cpp
--- a/Engine/Source/Editor/UMGEditor/Private/WidgetBlueprintEditorUtils.cpp
+++ b/Engine/Source/Editor/UMGEditor/Private/WidgetBlueprintEditorUtils.cpp
@@ -1613,6 +1613,26 @@ UWidget* FWidgetBlueprintEditorUtils::GetWidgetTemplateFromDragDrop(UWidgetBluep
 	return Widget;
 }
 
+bool FWidgetBlueprintEditorUtils::CanDropOnTargetExtensions(const UWidget* Target, const TSharedPtr<FDragDropOperation>& DragDropOp, FText& OutFailureText)
+{
+	if (Target)
+	{
+		IUMGEditorModule& EditorModule = FModuleManager::LoadModuleChecked<IUMGEditorModule>("UMGEditor");
+		const TArrayView<const TSharedPtr<IDragDropExtension>> DragDropExtensions = EditorModule.GetDragDropExtensibilityManager()->GetExtensions();
+
+		for (const TSharedPtr<IDragDropExtension>& DragDropExtension : DragDropExtensions)
+		{
+			if (ensure(DragDropExtension.IsValid()) && !DragDropExtension->CanDropOnTarget(Target, DragDropOp))
+			{
+				OutFailureText = DragDropExtension->GetDropFailureText(Target, DragDropOp);
+				return false;
+			}
+		}
+	}
+
+	return true;
+}
+
 void FWidgetBlueprintEditorUtils::ExportWidgetsToText(TArray<UWidget*> WidgetsToExport, /*out*/ FString& ExportedText)
 {
 	// Clear the mark state for saving.
diff --git a/Engine/Source/Editor/UMGEditor/Public/IHasDragDropExtensibility.h b/Engine/Source/Editor/UMGEditor/Public/IHasDragDropExtensibility.h
new file mode 100644
--- /dev/null
+++ b/Engine/Source/Editor/UMGEditor/Public/IHasDragDropExtensibility.h
@@ -0,0 +1,53 @@
+// Copyright Epic Games, Inc. All Rights Reserved.
+
+#pragma once
+
+#include "Internationalization/Text.h"
+
+class FDragDropOperation;
+class UWidget;
+
+class UMGEDITOR_API IDragDropExtension
+{
+public:
+	virtual ~IDragDropExtension() { }
+
+	virtual bool CanDropOnTarget(const UWidget* Target, const TSharedPtr<FDragDropOperation>& DragDropOp) const = 0;
+	virtual FText GetDropFailureText(const UWidget* Target, const TSharedPtr<FDragDropOperation>& DragDropOp) const = 0;
+};
+
+/**
+ * Drag & drop extensibility manager holds a list of registered drag and drop extensions.
+ */
+class UMGEDITOR_API FDragDropExtensibilityManager
+{
+public:
+	void AddExtension(const TSharedRef<IDragDropExtension>& Extension)
+	{
+		if (ensure(!Extensions.Contains(Extension)))
+		{
+			Extensions.Add(Extension);
+		}
+	}
+
+	void RemoveExtension(const TSharedRef<IDragDropExtension>& Extension)
+	{
+		int32 NumRemoved = Extensions.RemoveSingleSwap(Extension);
+		ensure(NumRemoved == 1);
+	}
+
+	TArrayView<const TSharedPtr<IDragDropExtension>> GetExtensions() const
+	{
+		return Extensions;
+	}
+
+private:
+	TArray<TSharedPtr<IDragDropExtension>> Extensions;
+};
+
+/** Indicates that a class can extend drag & drop functionality */
+class IHasDragDropExtensibility
+{
+public:
+	virtual TSharedPtr<FDragDropExtensibilityManager> GetDragDropExtensibilityManager() = 0;
+};
diff --git a/Engine/Source/Editor/UMGEditor/Public/UMGEditorModule.h b/Engine/Source/Editor/UMGEditor/Public/UMGEditorModule.h
--- a/Engine/Source/Editor/UMGEditor/Public/UMGEditorModule.h
+++ b/Engine/Source/Editor/UMGEditor/Public/UMGEditorModule.h
@@ -6,6 +6,7 @@
 #include "Toolkits/AssetEditorToolkit.h"
 #include "IHasClipboardExtensibility.h"
 #include "IHasDesignerExtensibility.h"
+#include "IHasDragDropExtensibility.h"
 #include "IHasPropertyBindingExtensibility.h"
 #include "UObject/TopLevelAssetPath.h"
 
@@ -24,7 +25,8 @@ class IUMGEditorModule :
 	public IHasToolBarExtensibility, 
 	public IHasDesignerExtensibility,
 	public IHasPropertyBindingExtensibility,
-	public IHasClipboardExtensibility
+	public IHasClipboardExtensibility,
+	public IHasDragDropExtensibility
 {
 public:
 	virtual class FWidgetBlueprintCompiler* GetRegisteredCompiler() = 0;
diff --git a/Engine/Source/Editor/UMGEditor/Public/WidgetBlueprintEditorUtils.h b/Engine/Source/Editor/UMGEditor/Public/WidgetBlueprintEditorUtils.h
--- a/Engine/Source/Editor/UMGEditor/Public/WidgetBlueprintEditorUtils.h
+++ b/Engine/Source/Editor/UMGEditor/Public/WidgetBlueprintEditorUtils.h
@@ -56,6 +56,8 @@ public:
 
 	static UWidget* GetWidgetTemplateFromDragDrop(UWidgetBlueprint* Blueprint, UWidgetTree* RootWidgetTree, TSharedPtr<FDragDropOperation>& DragDropOp);
 
+	static bool CanDropOnTargetExtensions(const UWidget* Target, const TSharedPtr<FDragDropOperation>& DragDropOp, FText& OutFailureText);
+
 	static bool IsBindWidgetProperty(const FProperty* InProperty);
 	static bool IsBindWidgetProperty(const FProperty* InProperty, bool& bIsOptional);
 
