UMG: Check the drag-drop extensions to see if a widget accepts drop when moving widgets in the designer.

diff --git a/Engine/Plugins/Runtime/ModelViewViewModel/Source/ModelViewViewModelEditor/Private/Customizations/MVVMDragDropExtension.cpp b/Engine/Plugins/Runtime/ModelViewViewModel/Source/ModelViewViewModelEditor/Private/Customizations/MVVMDragDropExtension.cpp
--- a/Engine/Plugins/Runtime/ModelViewViewModel/Source/ModelViewViewModelEditor/Private/Customizations/MVVMDragDropExtension.cpp
+++ b/Engine/Plugins/Runtime/ModelViewViewModel/Source/ModelViewViewModelEditor/Private/Customizations/MVVMDragDropExtension.cpp
@@ -11,9 +11,9 @@
 
 namespace UE::MVVM
 {
-bool FDragDropExtension::CanDropOnTarget(const UWidget* Target, const TSharedPtr<FDragDropOperation>& DragDropOp) const
+bool FWidgetDragDropExtension::ShouldPreventDropOnTarget(const UWidget* Target, const TSharedPtr<FDragDropOperation>& DragDropOp) const
 {
-	bool bCanDropWidgetOnTarget = true;
+	bool bShouldPreventDropWidgetOnTarget = false;
 
 	if (Target && Target->IsA<UPanelWidget>())
 	{
@@ -22,7 +22,7 @@ bool FDragDropExtension::CanDropOnTarget(const UWidget* Target, const TSharedPtr
 			if (const UMVVMWidgetBlueprintExtension_View* ExtensionView = UMVVMWidgetBlueprintExtension_View::GetExtension<UMVVMWidgetBlueprintExtension_View>(WidgetBlueprint))
 			{
 				const TArray<UMVVMBlueprintViewExtension*> BlueprintExtensions = ExtensionView->GetBlueprintExtensionsForWidget(Target->GetFName());
-				bCanDropWidgetOnTarget = !BlueprintExtensions.ContainsByPredicate([](const UMVVMBlueprintViewExtension* BlueprintExtension) -> bool
+				bShouldPreventDropWidgetOnTarget = BlueprintExtensions.ContainsByPredicate([](const UMVVMBlueprintViewExtension* BlueprintExtension) -> bool
 				{
 					return BlueprintExtension->IsA<UMVVMBlueprintViewExtension_PanelWidget>();
 				});
@@ -30,10 +30,10 @@ bool FDragDropExtension::CanDropOnTarget(const UWidget* Target, const TSharedPtr
 		}
 	}
 
-	return bCanDropWidgetOnTarget;
+	return bShouldPreventDropWidgetOnTarget;
 }
 
-FText FDragDropExtension::GetDropFailureText(const UWidget* Target, const TSharedPtr<FDragDropOperation>& DragDropOp) const
+FText FWidgetDragDropExtension::GetDropFailureText(const UWidget* Target, const TSharedPtr<FDragDropOperation>& DragDropOp) const
 {
 	return NSLOCTEXT("MVVMDragDropExtension", "UnableToAddChildWidget", "Cannot add children to a panel widget with an MVVM extension.");
 }
diff --git a/Engine/Plugins/Runtime/ModelViewViewModel/Source/ModelViewViewModelEditor/Private/Customizations/MVVMDragDropExtension.h b/Engine/Plugins/Runtime/ModelViewViewModel/Source/ModelViewViewModelEditor/Private/Customizations/MVVMDragDropExtension.h
--- a/Engine/Plugins/Runtime/ModelViewViewModel/Source/ModelViewViewModelEditor/Private/Customizations/MVVMDragDropExtension.h
+++ b/Engine/Plugins/Runtime/ModelViewViewModel/Source/ModelViewViewModelEditor/Private/Customizations/MVVMDragDropExtension.h
@@ -2,15 +2,15 @@
 
 #pragma once
 
-#include "IHasDragDropExtensibility.h"
+#include "IHasWidgetDragDropExtensibility.h"
 
 namespace UE::MVVM
 {
-class FDragDropExtension : public IDragDropExtension
+class FWidgetDragDropExtension : public IWidgetDragDropExtension
 {
-	//~ Begin IDragDropExtension overrides
-	virtual bool CanDropOnTarget(const UWidget* Target, const TSharedPtr<FDragDropOperation>& DragDropOp) const override;
+	//~ Begin IWidgetDragDropExtension overrides
+	virtual bool ShouldPreventDropOnTarget(const UWidget* Target, const TSharedPtr<FDragDropOperation>& DragDropOp) const override;
 	virtual FText GetDropFailureText(const UWidget* Target, const TSharedPtr<FDragDropOperation>& DragDropOp) const override;
-	//~ End IDragDropExtension overrides
+	//~ End IWidgetDragDropExtension overrides
 };
 }
diff --git a/Engine/Plugins/Runtime/ModelViewViewModel/Source/ModelViewViewModelEditor/Private/ModelViewViewModelEditorModule.cpp b/Engine/Plugins/Runtime/ModelViewViewModel/Source/ModelViewViewModelEditor/Private/ModelViewViewModelEditorModule.cpp
--- a/Engine/Plugins/Runtime/ModelViewViewModel/Source/ModelViewViewModelEditor/Private/ModelViewViewModelEditorModule.cpp
+++ b/Engine/Plugins/Runtime/ModelViewViewModel/Source/ModelViewViewModelEditor/Private/ModelViewViewModelEditorModule.cpp
@@ -53,8 +53,8 @@ void FModelViewViewModelEditorModule::StartupModule()
 	ClipboardExtension = MakeShared<UE::MVVM::FClipboardExtension>();
 	UMGEditorModule.GetClipboardExtensibilityManager()->AddExtension(ClipboardExtension.ToSharedRef());
 
-	DragDropExtension = MakeShared<UE::MVVM::FDragDropExtension>();
-	UMGEditorModule.GetDragDropExtensibilityManager()->AddExtension(DragDropExtension.ToSharedRef());
+	DragDropExtension = MakeShared<UE::MVVM::FWidgetDragDropExtension>();
+	UMGEditorModule.GetWidgetDragDropExtensibilityManager()->AddExtension(DragDropExtension.ToSharedRef());
 
 	ListViewBaseCustomizationExtender = UE::MVVM::FMVVMListViewBaseExtensionCustomizationExtender::MakeInstance();
 	UMGEditorModule.AddWidgetCustomizationExtender(ListViewBaseCustomizationExtender.ToSharedRef());
@@ -106,7 +106,7 @@ void FModelViewViewModelEditorModule::ShutdownModule()
 	{
 		UMGEditorModule->GetDesignerExtensibilityManager()->RemoveDesignerExtensionFactory(BlueprintViewDesignerExtensionFactory.ToSharedRef());
 		UMGEditorModule->OnRegisterTabsForEditor().RemoveAll(this);
-		UMGEditorModule->GetDragDropExtensibilityManager()->RemoveExtension(DragDropExtension.ToSharedRef());
+		UMGEditorModule->GetWidgetDragDropExtensibilityManager()->RemoveExtension(DragDropExtension.ToSharedRef());
 		UMGEditorModule->GetClipboardExtensibilityManager()->RemoveExtension(ClipboardExtension.ToSharedRef());
 		UMGEditorModule->GetPropertyBindingExtensibilityManager()->RemoveExtension(PropertyBindingExtension.ToSharedRef());
 
diff --git a/Engine/Plugins/Runtime/ModelViewViewModel/Source/ModelViewViewModelEditor/Private/ModelViewViewModelEditorModule.h b/Engine/Plugins/Runtime/ModelViewViewModel/Source/ModelViewViewModelEditor/Private/ModelViewViewModelEditorModule.h
--- a/Engine/Plugins/Runtime/ModelViewViewModel/Source/ModelViewViewModelEditor/Private/ModelViewViewModelEditorModule.h
+++ b/Engine/Plugins/Runtime/ModelViewViewModel/Source/ModelViewViewModelEditor/Private/ModelViewViewModelEditorModule.h
@@ -10,7 +10,7 @@ namespace UE::MVVM
 	class FAssetTypeActions_ViewModelBlueprint;
 	class FBlueprintViewDesignerExtensionFactory;
 	class FClipboardExtension;
-	class FDragDropExtension;
+	class FWidgetDragDropExtension;
 	class FMVVMListViewBaseExtensionCustomizationExtender;
 	class FMVVMPanelWidgetExtensionCustomizationExtender;
 	class FMVVMPropertyBindingExtension;
@@ -49,7 +49,7 @@ private:
 private:
 	TSharedPtr<UE::MVVM::FMVVMPropertyBindingExtension> PropertyBindingExtension;
 	TSharedPtr<UE::MVVM::FClipboardExtension> ClipboardExtension;
-	TSharedPtr<UE::MVVM::FDragDropExtension> DragDropExtension;
+	TSharedPtr<UE::MVVM::FWidgetDragDropExtension> DragDropExtension;
 	TSharedPtr<UE::MVVM::FAssetTypeActions_ViewModelBlueprint> ViewModelBlueprintActions;
 	TSharedPtr<UE::MVVM::FMVVMListViewBaseExtensionCustomizationExtender> ListViewBaseCustomizationExtender;
 	TSharedPtr<UE::MVVM::FMVVMPanelWidgetExtensionCustomizationExtender> PanelWidgetCustomizationExtender;
diff --git a/Engine/Source/Editor/UMGEditor/Private/Designer/SDesignerView.cpp b/Engine/Source/Editor/UMGEditor/Private/Designer/SDesignerView.cpp
--- a/Engine/Source/Editor/UMGEditor/Private/Designer/SDesignerView.cpp
+++ b/Engine/Source/Editor/UMGEditor/Private/Designer/SDesignerView.cpp
@@ -2725,8 +2725,15 @@ void SDesignerView::ProcessDropAndAddWidget(const FGeometry& MyGeometry, const F
 
 		FScopedTransaction Transaction(LOCTEXT("Designer_AddWidget", "Add Widget"));
 
+		FText DropOnTargetFailureText = FText::GetEmpty();
+		const bool bShouldPreventDropOnTargetExtensions = FWidgetBlueprintEditorUtils::ShouldPreventDropOnTargetExtensions(Target, SelectedDragDropOp, DropOnTargetFailureText);
+		if ( bShouldPreventDropOnTargetExtensions )
+		{
+			DragOperation->SetCursorOverride(EMouseCursor::SlashedCircle);
+			DragAndDropTransaction.Cancel();
+		}
 		// If there's no root widget go ahead and add the widget into the root slot.
-		if ( BP->WidgetTree->RootWidget == nullptr )
+		else if ( BP->WidgetTree->RootWidget == nullptr )
 		{
 			if ( !bIsPreview )
 			{
@@ -2905,8 +2912,11 @@ void SDesignerView::MoveWidgets(const FGeometry& MyGeometry, const FDragDropEven
 			FWidgetReference TargetReference = bIsPreview ? BlueprintEditor.Pin()->GetReferenceFromPreview(Target) : BlueprintEditor.Pin()->GetReferenceFromTemplate(Target);
 			BlueprintEditor.Pin()->SetHoveredWidget(TargetReference);
 
+			FText DropOnTargetFailureText = FText::GetEmpty();
+			const bool bShouldPreventDropOnTargetExtensions = FWidgetBlueprintEditorUtils::ShouldPreventDropOnTargetExtensions(Target, SelectedDragDropOp, DropOnTargetFailureText);
+
 			// If the widget being hovered over is a panel, attempt to place it into that panel.
-			if (Target && Target->IsA(UPanelWidget::StaticClass()))
+			if (Target && Target->IsA(UPanelWidget::StaticClass()) && !bShouldPreventDropOnTargetExtensions)
 			{
 				bWidgetMoved = true;
 
diff --git a/Engine/Source/Editor/UMGEditor/Private/Hierarchy/SHierarchyViewItem.cpp b/Engine/Source/Editor/UMGEditor/Private/Hierarchy/SHierarchyViewItem.cpp
--- a/Engine/Source/Editor/UMGEditor/Private/Hierarchy/SHierarchyViewItem.cpp
+++ b/Engine/Source/Editor/UMGEditor/Private/Hierarchy/SHierarchyViewItem.cpp
@@ -188,18 +188,18 @@ TOptional<EItemDropZone> ProcessHierarchyDragDrop(const FDragDropEvent& DragDrop
 	UWidgetBlueprint* Blueprint = BlueprintEditor->GetWidgetBlueprintObj();
 	check( Blueprint != nullptr && Blueprint->WidgetTree != nullptr );
 
-	const auto CanDropOnTargetExtensions = [](UWidget* Target, const TSharedPtr<FDecoratedDragDropOp>& DecoratedDragDropOp) -> bool
+	const auto ShouldPreventDropOnTargetExtensions = [](UWidget* Target, const TSharedPtr<FDecoratedDragDropOp>& DecoratedDragDropOp) -> bool
 		{
-			FText CanDropOnTargetFailureText = FText::GetEmpty();
-			const bool bCanDropOnTargetExtensions = FWidgetBlueprintEditorUtils::CanDropOnTargetExtensions(Target, DecoratedDragDropOp, CanDropOnTargetFailureText);
+			FText DropOnTargetFailureText = FText::GetEmpty();
+			const bool bShouldPreventDropOnTargetExtensions = FWidgetBlueprintEditorUtils::ShouldPreventDropOnTargetExtensions(Target, DecoratedDragDropOp, DropOnTargetFailureText);
 
-			if (!bCanDropOnTargetExtensions && DecoratedDragDropOp.IsValid())
+			if (bShouldPreventDropOnTargetExtensions && DecoratedDragDropOp.IsValid())
 			{
 				DecoratedDragDropOp->CurrentIconBrush = FAppStyle::GetBrush(TEXT("Graph.ConnectorFeedback.Error"));
-				DecoratedDragDropOp->CurrentHoverText = CanDropOnTargetFailureText;
+				DecoratedDragDropOp->CurrentHoverText = DropOnTargetFailureText;
 			}
 
-			return bCanDropOnTargetExtensions;
+			return bShouldPreventDropOnTargetExtensions;
 		};
 
 	// Is this a drag/drop op to create a new widget in the tree?
@@ -213,7 +213,7 @@ TOptional<EItemDropZone> ProcessHierarchyDragDrop(const FDragDropEvent& DragDrop
 			DecoratedDragDropOp->ResetToDefaultToolTip();
 		}
 
-		if (!CanDropOnTargetExtensions(TargetTemplate, DecoratedDragDropOp))
+		if ( ShouldPreventDropOnTargetExtensions(TargetTemplate, DecoratedDragDropOp) )
 		{
 			return TOptional<EItemDropZone>();
 		}
@@ -359,7 +359,7 @@ TOptional<EItemDropZone> ProcessHierarchyDragDrop(const FDragDropEvent& DragDrop
 				return TOptional<EItemDropZone>();
 			}
 
-			if (!CanDropOnTargetExtensions(TargetTemplate, HierarchyDragDropOp))
+			if (ShouldPreventDropOnTargetExtensions(TargetTemplate, HierarchyDragDropOp))
 			{
 				return TOptional<EItemDropZone>();
 			}
diff --git a/Engine/Source/Editor/UMGEditor/Private/UMGEditorModule.cpp b/Engine/Source/Editor/UMGEditor/Private/UMGEditorModule.cpp
--- a/Engine/Source/Editor/UMGEditor/Private/UMGEditorModule.cpp
+++ b/Engine/Source/Editor/UMGEditor/Private/UMGEditorModule.cpp
@@ -94,7 +94,7 @@ public:
 
 		PropertyBindingExtensibilityManager = MakeShared<FPropertyBindingExtensibilityManager>();
 		ClipboardExtensibilityManager = MakeShared<FClipboardExtensibilityManager>();
-		DragDropExtensibilityManager = MakeShared<FDragDropExtensibilityManager>();
+		WidgetDragDropExtensibilityManager = MakeShared<FWidgetDragDropExtensibilityManager>();
 
 		// Register widget blueprint compiler we do this no matter what.
 		IKismetCompilerInterface& KismetCompilerModule = FModuleManager::LoadModuleChecked<IKismetCompilerInterface>("KismetCompiler");
@@ -232,7 +232,7 @@ public:
 	virtual TSharedPtr<FDesignerExtensibilityManager> GetDesignerExtensibilityManager() override { return DesignerExtensibilityManager; }
 	virtual TSharedPtr<FPropertyBindingExtensibilityManager> GetPropertyBindingExtensibilityManager() override { return PropertyBindingExtensibilityManager; }
 	virtual TSharedPtr<FClipboardExtensibilityManager> GetClipboardExtensibilityManager() override { return ClipboardExtensibilityManager; }
-	virtual TSharedPtr<FDragDropExtensibilityManager> GetDragDropExtensibilityManager() override { return DragDropExtensibilityManager; }
+	virtual TSharedPtr<FWidgetDragDropExtensibilityManager> GetWidgetDragDropExtensibilityManager() override { return WidgetDragDropExtensibilityManager; }
 
 	/** Register settings objects. */
 	void RegisterSettings()
@@ -502,7 +502,7 @@ private:
 	TSharedPtr<FDesignerExtensibilityManager> DesignerExtensibilityManager;
 	TSharedPtr<FPropertyBindingExtensibilityManager> PropertyBindingExtensibilityManager;
 	TSharedPtr<FClipboardExtensibilityManager> ClipboardExtensibilityManager;
-	TSharedPtr<FDragDropExtensibilityManager> DragDropExtensibilityManager;
+	TSharedPtr<FWidgetDragDropExtensibilityManager> WidgetDragDropExtensibilityManager;
 	TSharedPtr<FGraphPanelPinFactory> GraphPanelPinFactory;
 
 	FDelegateHandle SequenceEditorHandle;
diff --git a/Engine/Source/Editor/UMGEditor/Private/WidgetBlueprintEditorUtils.cpp b/Engine/Source/Editor/UMGEditor/Private/WidgetBlueprintEditorUtils.cpp
--- a/Engine/Source/Editor/UMGEditor/Private/WidgetBlueprintEditorUtils.cpp
+++ b/Engine/Source/Editor/UMGEditor/Private/WidgetBlueprintEditorUtils.cpp
@@ -1613,24 +1613,24 @@ UWidget* FWidgetBlueprintEditorUtils::GetWidgetTemplateFromDragDrop(UWidgetBluep
 	return Widget;
 }
 
-bool FWidgetBlueprintEditorUtils::CanDropOnTargetExtensions(const UWidget* Target, const TSharedPtr<FDragDropOperation>& DragDropOp, FText& OutFailureText)
+bool FWidgetBlueprintEditorUtils::ShouldPreventDropOnTargetExtensions(const UWidget* Target, const TSharedPtr<FDragDropOperation>& DragDropOp, FText& OutFailureText)
 {
 	if (Target)
 	{
 		IUMGEditorModule& EditorModule = FModuleManager::LoadModuleChecked<IUMGEditorModule>("UMGEditor");
-		const TArrayView<const TSharedPtr<IDragDropExtension>> DragDropExtensions = EditorModule.GetDragDropExtensibilityManager()->GetExtensions();
+		const TArrayView<const TSharedPtr<IWidgetDragDropExtension>> DragDropExtensions = EditorModule.GetWidgetDragDropExtensibilityManager()->GetExtensions();
 
-		for (const TSharedPtr<IDragDropExtension>& DragDropExtension : DragDropExtensions)
+		for (const TSharedPtr<IWidgetDragDropExtension>& DragDropExtension : DragDropExtensions)
 		{
-			if (ensure(DragDropExtension.IsValid()) && !DragDropExtension->CanDropOnTarget(Target, DragDropOp))
+			if (ensure(DragDropExtension.IsValid()) && DragDropExtension->ShouldPreventDropOnTarget(Target, DragDropOp))
 			{
 				OutFailureText = DragDropExtension->GetDropFailureText(Target, DragDropOp);
-				return false;
+				return true;
 			}
 		}
 	}
 
-	return true;
+	return false;
 }
 
 void FWidgetBlueprintEditorUtils::ExportWidgetsToText(TArray<UWidget*> WidgetsToExport, /*out*/ FString& ExportedText)
diff --git a/Engine/Source/Editor/UMGEditor/Public/IHasDragDropExtensibility.h b/Engine/Source/Editor/UMGEditor/Public/IHasDragDropExtensibility.h
deleted file mode 100644
diff --git a/Engine/Source/Editor/UMGEditor/Public/IHasWidgetDragDropExtensibility.h b/Engine/Source/Editor/UMGEditor/Public/IHasWidgetDragDropExtensibility.h
new file mode 100644
--- /dev/null
+++ b/Engine/Source/Editor/UMGEditor/Public/IHasWidgetDragDropExtensibility.h
@@ -0,0 +1,53 @@
+// Copyright Epic Games, Inc. All Rights Reserved.
+
+#pragma once
+
+#include "Internationalization/Text.h"
+
+class FDragDropOperation;
+class UWidget;
+
+class UMGEDITOR_API IWidgetDragDropExtension
+{
+public:
+	virtual ~IWidgetDragDropExtension() { }
+
+	virtual bool ShouldPreventDropOnTarget(const UWidget* Target, const TSharedPtr<FDragDropOperation>& DragDropOp) const = 0;
+	virtual FText GetDropFailureText(const UWidget* Target, const TSharedPtr<FDragDropOperation>& DragDropOp) const = 0;
+};
+
+/**
+ * Drag & drop extensibility manager holds a list of registered drag and drop extensions.
+ */
+class UMGEDITOR_API FWidgetDragDropExtensibilityManager
+{
+public:
+	void AddExtension(const TSharedRef<IWidgetDragDropExtension>& Extension)
+	{
+		if (ensure(!Extensions.Contains(Extension)))
+		{
+			Extensions.Add(Extension);
+		}
+	}
+
+	void RemoveExtension(const TSharedRef<IWidgetDragDropExtension>& Extension)
+	{
+		int32 NumRemoved = Extensions.RemoveSingleSwap(Extension);
+		ensure(NumRemoved == 1);
+	}
+
+	TArrayView<const TSharedPtr<IWidgetDragDropExtension>> GetExtensions() const
+	{
+		return Extensions;
+	}
+
+private:
+	TArray<TSharedPtr<IWidgetDragDropExtension>> Extensions;
+};
+
+/** Indicates that a class can extend drag & drop functionality */
+class IHasWidgetDragDropExtensibility
+{
+public:
+	virtual TSharedPtr<FWidgetDragDropExtensibilityManager> GetWidgetDragDropExtensibilityManager() = 0;
+};
diff --git a/Engine/Source/Editor/UMGEditor/Public/UMGEditorModule.h b/Engine/Source/Editor/UMGEditor/Public/UMGEditorModule.h
--- a/Engine/Source/Editor/UMGEditor/Public/UMGEditorModule.h
+++ b/Engine/Source/Editor/UMGEditor/Public/UMGEditorModule.h
@@ -6,8 +6,8 @@
 #include "Toolkits/AssetEditorToolkit.h"
 #include "IHasClipboardExtensibility.h"
 #include "IHasDesignerExtensibility.h"
-#include "IHasDragDropExtensibility.h"
 #include "IHasPropertyBindingExtensibility.h"
+#include "IHasWidgetDragDropExtensibility.h"
 #include "UObject/TopLevelAssetPath.h"
 
 extern const FName UMGEditorAppIdentifier;
@@ -26,7 +26,7 @@ class IUMGEditorModule :
 	public IHasDesignerExtensibility,
 	public IHasPropertyBindingExtensibility,
 	public IHasClipboardExtensibility,
-	public IHasDragDropExtensibility
+	public IHasWidgetDragDropExtensibility
 {
 public:
 	virtual class FWidgetBlueprintCompiler* GetRegisteredCompiler() = 0;
diff --git a/Engine/Source/Editor/UMGEditor/Public/WidgetBlueprintEditorUtils.h b/Engine/Source/Editor/UMGEditor/Public/WidgetBlueprintEditorUtils.h
--- a/Engine/Source/Editor/UMGEditor/Public/WidgetBlueprintEditorUtils.h
+++ b/Engine/Source/Editor/UMGEditor/Public/WidgetBlueprintEditorUtils.h
@@ -56,7 +56,7 @@ public:
 
 	static UWidget* GetWidgetTemplateFromDragDrop(UWidgetBlueprint* Blueprint, UWidgetTree* RootWidgetTree, TSharedPtr<FDragDropOperation>& DragDropOp);
 
-	static bool CanDropOnTargetExtensions(const UWidget* Target, const TSharedPtr<FDragDropOperation>& DragDropOp, FText& OutFailureText);
+	static bool ShouldPreventDropOnTargetExtensions(const UWidget* Target, const TSharedPtr<FDragDropOperation>& DragDropOp, FText& OutFailureText);
 
 	static bool IsBindWidgetProperty(const FProperty* InProperty);
 	static bool IsBindWidgetProperty(const FProperty* InProperty, bool& bIsOptional);
