[UMG] Provided a better description for "bCanCallInitializedWithoutPlayerContext"

diff --git a/Engine/Source/Editor/UMGEditor/Public/WidgetBlueprint.h b/Engine/Source/Editor/UMGEditor/Public/WidgetBlueprint.h
--- a/Engine/Source/Editor/UMGEditor/Public/WidgetBlueprint.h
+++ b/Engine/Source/Editor/UMGEditor/Public/WidgetBlueprint.h
@@ -237,7 +237,11 @@ public:
 	UPROPERTY(AssetRegistrySearchable)
 	FString PaletteCategory;
 
-	/** Run the initialize event on widget that doesn't have a player context. */
+	/**
+	 * Determines whether this widget blueprint can be initialized without
+	 * a valid player context (PlayerController, etc.).
+	 * Required to be true for use with UMG Widget Preview.
+	 */
 	UPROPERTY(EditAnywhere, Category="Widget")
 	bool bCanCallInitializedWithoutPlayerContext;
 #endif
diff --git a/Engine/Source/Runtime/UMG/Public/Blueprint/WidgetBlueprintGeneratedClass.h b/Engine/Source/Runtime/UMG/Public/Blueprint/WidgetBlueprintGeneratedClass.h
--- a/Engine/Source/Runtime/UMG/Public/Blueprint/WidgetBlueprintGeneratedClass.h
+++ b/Engine/Source/Runtime/UMG/Public/Blueprint/WidgetBlueprintGeneratedClass.h
@@ -103,7 +103,11 @@ public:
 #endif
 
 public:
-	/** */
+	/**
+	 * Determines whether this widget blueprint can be initialized without
+	 * a valid player context (PlayerController, etc.).
+	 * Required to be true for use with UMG Widget Preview.
+	 */
 	UPROPERTY()
 	uint32 bCanCallInitializedWithoutPlayerContext : 1;
 
