[UMG] Fix LOCTEXT blank entry in SWidgetDesigner.

diff --git a/Engine/Source/Editor/UMGEditor/Private/Designer/SDesignerView.cpp b/Engine/Source/Editor/UMGEditor/Private/Designer/SDesignerView.cpp
--- a/Engine/Source/Editor/UMGEditor/Private/Designer/SDesignerView.cpp
+++ b/Engine/Source/Editor/UMGEditor/Private/Designer/SDesignerView.cpp
@@ -484,7 +484,7 @@ TSharedRef<SWidget> SDesignerView::CreateOverlayUI()
 			[
 				SNew(STextBlock)
 				.TextStyle(FAppStyle::Get(), TEXT("Graph.ZoomText"))
-				.Text(LOCTEXT("", "Sequencer Selection Limited"))
+				.Text(LOCTEXT("SequencerSelectionLimited", "Sequencer Selection Limited"))
 				.ColorAndOpacity(FLinearColor(1.f, 1.f, 1.f, 0.25f))
 				.Visibility(this, &SDesignerView::GetSelectionLimitedTextVisibility)
 			]
