UMG: Fix BP compilation crash when the base class doesn't exist anymore.

diff --git a/Engine/Source/Editor/UMGEditor/Private/WidgetBlueprint.cpp b/Engine/Source/Editor/UMGEditor/Private/WidgetBlueprint.cpp
--- a/Engine/Source/Editor/UMGEditor/Private/WidgetBlueprint.cpp
+++ b/Engine/Source/Editor/UMGEditor/Private/WidgetBlueprint.cpp
@@ -785,7 +785,7 @@ bool UWidgetBlueprint::DetectSlateWidgetLeaks(FDataValidationContext& Context) c
 
 	// The detection relies on instantiation of the class: don't try to create an abstract class. 
 	// The validation will have to be run on the WBP inheriting from abstract ones.
-	if (GeneratedClass->HasAnyClassFlags(CLASS_Abstract))
+	if (GeneratedClass == nullptr || GeneratedClass->HasAnyClassFlags(CLASS_Abstract))
 	{
 		return false;
 	}
