Build fix: potential access to NewComponent ptr before null check

diff --git a/Engine/Source/Editor/UMGEditor/Private/UIComponentUtils.cpp b/Engine/Source/Editor/UMGEditor/Private/UIComponentUtils.cpp
--- a/Engine/Source/Editor/UMGEditor/Private/UIComponentUtils.cpp
+++ b/Engine/Source/Editor/UMGEditor/Private/UIComponentUtils.cpp
@@ -13,9 +13,12 @@
 UUIComponent* FUIComponentUtils::CreateUIComponent(TSubclassOf<UUIComponent> ComponentClass, UUserWidget* Outer)
 {
 	ensure(Outer);
-	UObject* NewComponent = NewObject<UObject>(Outer, ComponentClass);
-	NewComponent->SetFlags(RF_Transactional);
-	return NewComponent ? CastChecked<UUIComponent>(NewComponent) : nullptr;
+	if (UUIComponent* NewComponent = NewObject<UUIComponent>(Outer, ComponentClass))
+	{
+		NewComponent->SetFlags(RF_Transactional);
+		return NewComponent;
+	}
+	return nullptr;
 }
 
 UUIComponentContainer* FUIComponentUtils::GetOrCreateComponentsContainerForUserWidget(UUserWidget* UserWidget)
