PR #11890: UMG: Fix BindWidgetAnim causing errors when the animation exists in a parent widget blueprint

diff --git a/Engine/Source/Editor/UMGEditor/Private/WidgetBlueprintCompiler.cpp b/Engine/Source/Editor/UMGEditor/Private/WidgetBlueprintCompiler.cpp
--- a/Engine/Source/Editor/UMGEditor/Private/WidgetBlueprintCompiler.cpp
+++ b/Engine/Source/Editor/UMGEditor/Private/WidgetBlueprintCompiler.cpp
@@ -618,35 +618,48 @@ void FWidgetBlueprintCompilerContext::CreateClassVariablesFromBlueprint()
 		}
 	}
 
-	// Add movie scenes variables here
-	for (UWidgetAnimation* Animation : WidgetBP->Animations)
+	WidgetBPToScan = WidgetBP;
+	while (WidgetBPToScan != nullptr)
 	{
-		FObjectPropertyBase* ExistingProperty = CastField<FObjectPropertyBase>(ParentClass->FindPropertyByName(Animation->GetFName()));
-		if (ExistingProperty &&
-			FWidgetBlueprintEditorUtils::IsBindWidgetAnimProperty(ExistingProperty) &&
-			ExistingProperty->PropertyClass->IsChildOf(UWidgetAnimation::StaticClass()))
+		// Look for BindWidgetAnim properties in parent widgetblueprints
+		for (UWidgetAnimation* Animation : WidgetBPToScan->Animations)
 		{
-			WidgetAnimToMemberVariableMap.Add(Animation, ExistingProperty);
-			continue;
-		}
+			FObjectPropertyBase* ExistingProperty = CastField<FObjectPropertyBase>(ParentClass->FindPropertyByName(Animation->GetFName()));
+			if (ExistingProperty &&
+				FWidgetBlueprintEditorUtils::IsBindWidgetAnimProperty(ExistingProperty) &&
+				ExistingProperty->PropertyClass->IsChildOf(UWidgetAnimation::StaticClass()))
+			{
+				WidgetAnimToMemberVariableMap.Add(Animation, ExistingProperty);
+				continue;
+			}
 
-		FEdGraphPinType WidgetPinType(UEdGraphSchema_K2::PC_Object, NAME_None, Animation->GetClass(), EPinContainerType::None, true, FEdGraphTerminalType());
-		FProperty* AnimationProperty = CreateVariable(Animation->GetFName(), WidgetPinType);
+			// Create variables for widget animation
+			if (WidgetBPToScan == WidgetBP)
+			{
+				FEdGraphPinType WidgetPinType(UEdGraphSchema_K2::PC_Object, NAME_None, Animation->GetClass(), EPinContainerType::None, true, FEdGraphTerminalType());
+				FProperty* AnimationProperty = CreateVariable(Animation->GetFName(), WidgetPinType);
 
-		if ( AnimationProperty != nullptr )
-		{
-			const FString DisplayName = Animation->GetDisplayName().ToString();
-			AnimationProperty->SetMetaData(TEXT("DisplayName"), *DisplayName);
+				if (AnimationProperty != nullptr)
+				{
+					const FString DisplayName = Animation->GetDisplayName().ToString();
+					AnimationProperty->SetMetaData(TEXT("DisplayName"), *DisplayName);
 
-			AnimationProperty->SetMetaData(TEXT("Category"), TEXT("Animations"));
+					AnimationProperty->SetMetaData(TEXT("Category"), TEXT("Animations"));
 
-			AnimationProperty->SetPropertyFlags(CPF_Transient);
-			AnimationProperty->SetPropertyFlags(CPF_BlueprintVisible);
-			AnimationProperty->SetPropertyFlags(CPF_BlueprintReadOnly);
-			AnimationProperty->SetPropertyFlags(CPF_RepSkip);
+					AnimationProperty->SetPropertyFlags(CPF_Transient);
+					AnimationProperty->SetPropertyFlags(CPF_BlueprintVisible);
+					AnimationProperty->SetPropertyFlags(CPF_BlueprintReadOnly);
+					AnimationProperty->SetPropertyFlags(CPF_RepSkip);
 
-			WidgetAnimToMemberVariableMap.Add(Animation, AnimationProperty);
+					WidgetAnimToMemberVariableMap.Add(Animation, AnimationProperty);
+				}
+			}
 		}
+
+		// Get the parent WidgetBlueprint
+		WidgetBPToScan = WidgetBPToScan->ParentClass && WidgetBPToScan->ParentClass->ClassGeneratedBy
+			? Cast<UWidgetBlueprint>(WidgetBPToScan->ParentClass->ClassGeneratedBy)
+			: nullptr;
 	}
 
 	FWidgetBlueprintCompilerContext* Self = this;
diff --git a/Engine/Source/Editor/UMGEditor/Private/Widgets/SBindWidgetView.cpp b/Engine/Source/Editor/UMGEditor/Private/Widgets/SBindWidgetView.cpp
--- a/Engine/Source/Editor/UMGEditor/Private/Widgets/SBindWidgetView.cpp
+++ b/Engine/Source/Editor/UMGEditor/Private/Widgets/SBindWidgetView.cpp
@@ -354,7 +354,14 @@ TArray<UWidget*> GetAllSourceWidgets(UWidgetBlueprint* WidgetBlueprint)
 
 TArray<UWidgetAnimation*> GetAllSourceWidgetAnimations(UWidgetBlueprint* WidgetBlueprint)
 {
-	return WidgetBlueprint->Animations;
+	TArray<UWidgetAnimation*> WidgetAnimations;
+	UWidgetBlueprint* WidgetBPToScan = WidgetBlueprint;
+	while (WidgetBPToScan != nullptr)
+	{
+		WidgetAnimations.Append(WidgetBPToScan->Animations);
+		WidgetBPToScan = WidgetBPToScan->ParentClass && WidgetBPToScan->ParentClass->ClassGeneratedBy ? Cast<UWidgetBlueprint>(WidgetBPToScan->ParentClass->ClassGeneratedBy) : nullptr;
+	}
+	return WidgetAnimations;
 }
 
 
