[UMG Editor] Moved preview creation to FWidgetBlueprintEditorUtils

diff --git a/Engine/Source/Editor/UMGEditor/Private/WidgetBlueprintEditor.cpp b/Engine/Source/Editor/UMGEditor/Private/WidgetBlueprintEditor.cpp
--- a/Engine/Source/Editor/UMGEditor/Private/WidgetBlueprintEditor.cpp
+++ b/Engine/Source/Editor/UMGEditor/Private/WidgetBlueprintEditor.cpp
@@ -1680,12 +1680,7 @@ void FWidgetBlueprintEditor::DestroyPreview()
 		// otherwise the leak detection can't be trusted.
 		OnWidgetPreviewUpdated.Broadcast();
 
-		TWeakPtr<SWidget> PreviewSlateWidgetWeak = PreviewUserWidget->GetCachedWidget();
-
-		PreviewUserWidget->MarkAsGarbage();
-		PreviewUserWidget->ReleaseSlateResources(true);
-
-		ensure(!PreviewSlateWidgetWeak.IsValid());
+		FWidgetBlueprintEditorUtils::DestroyUserWidget(PreviewUserWidget);
 	}
 }
 
@@ -1708,55 +1703,13 @@ void FWidgetBlueprintEditor::UpdatePreview(UBlueprint* InBlueprint, bool bInForc
 		// Save the Blueprint we're creating a preview for
 		PreviewBlueprint = Cast<UWidgetBlueprint>(InBlueprint);
 
-		// Create the Widget, we have to do special swapping out of the widget tree.
-		{
-			// Assign the outer to the game instance if it exists, otherwise use the world
-			{
-				FMakeClassSpawnableOnScope TemporarilySpawnable(PreviewBlueprint->GeneratedClass);
-				PreviewUserWidget = NewObject<UUserWidget>(PreviewScene.GetWorld(), PreviewBlueprint->GeneratedClass);
-			}
-
-			// The preview widget should not be transactional.
-			PreviewUserWidget->ClearFlags(RF_Transactional);
-
-			// Establish the widget as being in design time before initializing and before duplication 
-            // (so that IsDesignTime is reliable within both calls to Initialize)
-            // The preview widget is also the outer widget that will update all child flags
-			PreviewUserWidget->SetDesignerFlags(GetCurrentDesignerFlags());
-
-			if ( ULocalPlayer* Player = PreviewScene.GetWorld()->GetFirstLocalPlayerFromController() )
-			{
-				PreviewUserWidget->SetPlayerContext(FLocalPlayerContext(Player));
-			}
-
-			UWidgetTree* LatestWidgetTree = FWidgetBlueprintEditorUtils::FindLatestWidgetTree(PreviewBlueprint, PreviewUserWidget);
-
-			TMap<FName, UWidget*> NamedSlotContentToMerge;
-			UWidgetBlueprint* WidgetBPIt = PreviewBlueprint;
-			while (WidgetBPIt)
-			{
-				TArray<FName> SlotNames;
-				WidgetBPIt->WidgetTree->GetSlotNames(SlotNames);
-
-				for(const FName SlotName : SlotNames)
-				{
-					if(UWidget* Content = WidgetBPIt->WidgetTree->GetContentForSlot(SlotName))
-					{
-						NamedSlotContentToMerge.Add(SlotName, Content);
-					}
-				}
-
-				WidgetBPIt = Cast<UWidgetBlueprint>(WidgetBPIt->GeneratedClass->GetSuperClass()->ClassGeneratedBy);
-			}
-
-			// Update the widget tree directly to match the blueprint tree.  That way the preview can update
-			// without needing to do a full recompile.
-			PreviewUserWidget->DuplicateAndInitializeFromWidgetTree(LatestWidgetTree, NamedSlotContentToMerge);
-
-			// Establish the widget as being in design time before initializing (so that IsDesignTime is reliable within Initialize)
-            // We have to call it to make sure that all the WidgetTree had the DesignerFlags set correctly
-			PreviewUserWidget->SetDesignerFlags(GetCurrentDesignerFlags());
-		}
+		PreviewUserWidget = FWidgetBlueprintEditorUtils::CreateUserWidgetFromBlueprint(
+			PreviewScene.GetWorld(),
+			PreviewBlueprint,
+			FWidgetBlueprintEditorUtils::FCreateWidgetFromBlueprintParams{
+				GetCurrentDesignerFlags(),
+				PreviewScene.GetWorld()->GetFirstLocalPlayerFromController()
+			});
 
 		// Store a reference to the preview actor.
 		PreviewWidgetPtr = PreviewUserWidget;
diff --git a/Engine/Source/Editor/UMGEditor/Private/WidgetBlueprintEditorUtils.cpp b/Engine/Source/Editor/UMGEditor/Private/WidgetBlueprintEditorUtils.cpp
--- a/Engine/Source/Editor/UMGEditor/Private/WidgetBlueprintEditorUtils.cpp
+++ b/Engine/Source/Editor/UMGEditor/Private/WidgetBlueprintEditorUtils.cpp
@@ -1517,6 +1517,78 @@ TArray<UWidget*> FWidgetBlueprintEditorUtils::DuplicateWidgets(TSharedRef<FWidge
 	return DuplicatedWidgets;
 }
 
+UUserWidget* FWidgetBlueprintEditorUtils::CreateUserWidgetFromBlueprint(UObject* Outer, UWidgetBlueprint* BP, const FCreateWidgetFromBlueprintParams& Params)
+{
+	check(Outer);
+	check(BP);
+
+	UUserWidget* CreatedUserWidget = nullptr;
+
+	// Create the Widget, we have to do special swapping out of the widget tree.
+	{
+		// Assign the outer to the game instance if it exists, otherwise use the world
+		{
+			FMakeClassSpawnableOnScope TemporarilySpawnable(BP->GeneratedClass);
+			CreatedUserWidget = NewObject<UUserWidget>(Outer, BP->GeneratedClass);
+		}
+
+		// The preview widget should not be transactional.
+		CreatedUserWidget->ClearFlags(RF_Transactional);
+
+		// Establish the widget as being in design time before initializing and before duplication
+        // (so that IsDesignTime is reliable within both calls to Initialize)
+        // The preview widget is also the outer widget that will update all child flags
+		CreatedUserWidget->SetDesignerFlags(Params.FlagsToApply);
+
+		if (ULocalPlayer* Player = Params.LocalPlayer)
+		{
+			CreatedUserWidget->SetPlayerContext(FLocalPlayerContext(Player));
+		}
+
+		UWidgetTree* LatestWidgetTree = FWidgetBlueprintEditorUtils::FindLatestWidgetTree(BP, CreatedUserWidget);
+
+		TMap<FName, UWidget*> NamedSlotContentToMerge;
+		UWidgetBlueprint* WidgetBlueprintIterator = BP;
+		while (WidgetBlueprintIterator)
+		{
+			TArray<FName> SlotNames;
+			WidgetBlueprintIterator->WidgetTree->GetSlotNames(SlotNames);
+
+			for (const FName SlotName : SlotNames)
+			{
+				if (UWidget* Content = WidgetBlueprintIterator->WidgetTree->GetContentForSlot(SlotName))
+				{
+					NamedSlotContentToMerge.Add(SlotName, Content);
+				}
+			}
+
+			WidgetBlueprintIterator = Cast<UWidgetBlueprint>(WidgetBlueprintIterator->GeneratedClass->GetSuperClass()->ClassGeneratedBy);
+		}
+
+		// Update the widget tree directly to match the blueprint tree.  That way the preview can update
+		// without needing to do a full recompile.
+		CreatedUserWidget->DuplicateAndInitializeFromWidgetTree(LatestWidgetTree, NamedSlotContentToMerge);
+
+		// Establish the widget as being in design time before initializing (so that IsDesignTime is reliable within Initialize)
+        // We have to call it to make sure that all the WidgetTree had the DesignerFlags set correctly
+		CreatedUserWidget->SetDesignerFlags(Params.FlagsToApply);
+	}
+
+	return CreatedUserWidget;
+}
+
+void FWidgetBlueprintEditorUtils::DestroyUserWidget(UUserWidget* UserWidget)
+{
+	check(UserWidget);
+
+	TWeakPtr<SWidget> SlateWidgetWeak = UserWidget->GetCachedWidget();
+
+	UserWidget->MarkAsGarbage();
+	UserWidget->ReleaseSlateResources(true);
+
+	ensure(!SlateWidgetWeak.IsValid());
+}
+
 bool FWidgetBlueprintEditorUtils::IsAnySelectedWidgetLocked(TSet<FWidgetReference> SelectedWidgets)
 {
 	for (const FWidgetReference& Widget : SelectedWidgets)
diff --git a/Engine/Source/Editor/UMGEditor/Public/WidgetBlueprintEditorUtils.h b/Engine/Source/Editor/UMGEditor/Public/WidgetBlueprintEditorUtils.h
--- a/Engine/Source/Editor/UMGEditor/Public/WidgetBlueprintEditorUtils.h
+++ b/Engine/Source/Editor/UMGEditor/Public/WidgetBlueprintEditorUtils.h
@@ -16,6 +16,7 @@ class UWidgetTree;
 class SWindow;
 class UWidgetEditingProjectSettings;
 class FWidgetObjectTextFactory;
+class ULocalPlayer;
 
 //////////////////////////////////////////////////////////////////////////
 // FWidgetBlueprintEditorUtils
@@ -30,6 +31,12 @@ public:
 		FVector2D Offset;
 	};
 
+	struct FCreateWidgetFromBlueprintParams
+	{
+		EWidgetDesignFlags FlagsToApply;
+		ULocalPlayer* LocalPlayer = nullptr; // Optionally specify if available
+	};
+
 	static bool VerifyWidgetRename(TSharedRef<class FWidgetBlueprintEditor> BlueprintEditor, FWidgetReference Widget, const FText& NewName, FText& OutErrorMessage);
 
 	static bool RenameWidget(TSharedRef<class FWidgetBlueprintEditor> BlueprintEditor, const FName& OldObjectName, const FString& NewDisplayName);
@@ -50,6 +57,11 @@ public:
 
 	static TArray<UWidget*> DuplicateWidgets(TSharedRef<FWidgetBlueprintEditor> BlueprintEditor, UWidgetBlueprint* BP, TSet<FWidgetReference> Widgets);
 
+	static UUserWidget* CreateUserWidgetFromBlueprint(UObject* Outer, UWidgetBlueprint* BP, const FCreateWidgetFromBlueprintParams& Params);
+
+	/** Performs cleanup on the specified UserWidget. */
+	static void DestroyUserWidget(UUserWidget* UserWidget);
+
 	static bool IsAnySelectedWidgetLocked(TSet<FWidgetReference> SelectedWidgets);
 
 	static bool CanPasteWidgetsExtension(TSet<FWidgetReference> SelectedWidgets);
