UMG Namedslot: hide namedslot in the hierarchy if it already has a child in the same tree that owns it but doesn't have content in the widget that inherits it
- fix order of namedslots when building the preview userwidget

diff --git a/Engine/Source/Editor/UMGEditor/Private/Hierarchy/SHierarchyViewItem.cpp b/Engine/Source/Editor/UMGEditor/Private/Hierarchy/SHierarchyViewItem.cpp
--- a/Engine/Source/Editor/UMGEditor/Private/Hierarchy/SHierarchyViewItem.cpp
+++ b/Engine/Source/Editor/UMGEditor/Private/Hierarchy/SHierarchyViewItem.cpp
@@ -762,8 +762,17 @@ void FHierarchyRoot::GetChildren(TArray< TSharedPtr<FHierarchyModel> >& Children
 	}
 
 	// Grab any exposed named slots from the super classes CDO.  These slots can have content slotted into them by this subclass.
+	TSet<FName> InheritedNamedSlotsWithContentInSameTree = Blueprint->GetInheritedNamedSlotsWithContentInSameTree();
 	for ( const FName& SlotName : Blueprint->GetInheritedAvailableNamedSlots() )
 	{
+		if (InheritedNamedSlotsWithContentInSameTree.Contains(SlotName))
+		{
+			if (!Blueprint->WidgetTree->GetContentForSlot(SlotName))
+			{
+				continue;
+			}
+		}
+
 		TSharedPtr<FNamedSlotModelSubclass> ChildItem = MakeShareable(new FNamedSlotModelSubclass(Blueprint, SlotName, BPEd));
 		Children.Add(ChildItem);
 	}
diff --git a/Engine/Source/Editor/UMGEditor/Private/WidgetBlueprint.cpp b/Engine/Source/Editor/UMGEditor/Private/WidgetBlueprint.cpp
--- a/Engine/Source/Editor/UMGEditor/Private/WidgetBlueprint.cpp
+++ b/Engine/Source/Editor/UMGEditor/Private/WidgetBlueprint.cpp
@@ -1426,6 +1426,16 @@ TArray<FName> UWidgetBlueprint::GetInheritedAvailableNamedSlots() const
 	return TArray<FName>();
 }
 
+TSet<FName> UWidgetBlueprint::GetInheritedNamedSlotsWithContentInSameTree() const
+{
+	if (const UWidgetBlueprintGeneratedClass* GeneratedBPClass = Cast<UWidgetBlueprintGeneratedClass>(GeneratedClass->GetSuperClass()))
+	{
+		return GeneratedBPClass->NamedSlotsWithContentInSameTree;
+	}
+
+	return TSet<FName>();
+}
+
 UWidgetEditingProjectSettings* UWidgetBlueprint::GetRelevantSettings()
 {
 	return GetMutableDefault<UUMGEditorProjectSettings>();
diff --git a/Engine/Source/Editor/UMGEditor/Private/WidgetBlueprintCompiler.cpp b/Engine/Source/Editor/UMGEditor/Private/WidgetBlueprintCompiler.cpp
--- a/Engine/Source/Editor/UMGEditor/Private/WidgetBlueprintCompiler.cpp
+++ b/Engine/Source/Editor/UMGEditor/Private/WidgetBlueprintCompiler.cpp
@@ -985,41 +985,52 @@ void FWidgetBlueprintCompilerContext::FinishCompilingClass(UClass* Class)
 
 		// Add all the names of the named slot widgets to the slot names structure.
 		{
-			TArray<FName> NamedSlotsWithContentInSameTree;
 		#if WITH_EDITOR
 			BPGClass->NamedSlotsWithID.Reset();
+			BPGClass->NamedSlotsWithContentInSameTree.Reset();
 		#endif
 			BPGClass->NamedSlots.Reset();
 			BPGClass->InstanceNamedSlots.Reset();
+
+			TArray<FName> NamedSlotsPerWidgetBlueprint;
+
 			UWidgetBlueprint* WidgetBPIt = WidgetBP;
 			while (WidgetBPIt)
 			{
+				NamedSlotsPerWidgetBlueprint.Reset();
 				WidgetBPIt->ForEachSourceWidget([&] (const UWidget* Widget) {
 					if (const UNamedSlot* NamedSlot = Cast<UNamedSlot>(Widget))
 					{
-						BPGClass->NamedSlots.Add(Widget->GetFName());
+						NamedSlotsPerWidgetBlueprint.Add(Widget->GetFName());
 
 					#if WITH_EDITOR
 						BPGClass->NamedSlotsWithID.Add(TPair<FName, FGuid>(Widget->GetFName(), NamedSlot->GetSlotGUID()));
-					#endif
 
-						if (NamedSlot->bExposeOnInstanceOnly)
+						// A namedslot whose content is in the same blueprint class is treated as a regular panel widget.
+						// We need to keep track of these to later remove them from the hierarchy.
+						if (NamedSlot->GetChildrenCount() > 0)
 						{
-							BPGClass->InstanceNamedSlots.Add(Widget->GetFName());
+							BPGClass->NamedSlotsWithContentInSameTree.Add(NamedSlot->GetFName());
 						}
+					#endif
 
-						// A namedslot whose content is in the same blueprint class is treated as a regular panel widget.
-						// We need to keep track of these to later remove them from the available namedslots list.
-						if (NamedSlot->GetChildrenCount() > 0)
+						if (NamedSlot->bExposeOnInstanceOnly)
 						{
-							NamedSlotsWithContentInSameTree.Add(NamedSlot->GetFName());
+							BPGClass->InstanceNamedSlots.Add(Widget->GetFName());
 						}
 					}
 				});
 				
+				// Here we reverse this array to maintain the order of sibling namedslots once the final array BPGClass->NamedSlots is reversed.
+				Algo::Reverse(NamedSlotsPerWidgetBlueprint);
+				BPGClass->NamedSlots.Append(NamedSlotsPerWidgetBlueprint);
+
 				WidgetBPIt = Cast<UWidgetBlueprint>(WidgetBPIt->ParentClass->ClassGeneratedBy);
 			}
 
+			// We iterate widget blueprints from child to parent, but we need the final namedslot array to be sorted from parent to child, so we reverse it.
+			Algo::Reverse(BPGClass->NamedSlots);
+
 			BPGClass->AvailableNamedSlots = BPGClass->NamedSlots;
 
 			// Remove any named slots from the available slots that has content for it.
@@ -1029,12 +1040,6 @@ void FWidgetBlueprintCompilerContext::FinishCompilingClass(UClass* Class)
 				BPGClass->AvailableNamedSlots.Remove(SlotName);
 			});
 
-			// Remove any named slots with content in the same widget tree from the available slots.
-			for (const FName& NamedSlotWithContent : NamedSlotsWithContentInSameTree)
-			{
-				BPGClass->AvailableNamedSlots.Remove(NamedSlotWithContent);
-			}
-
 			// Remove any available subclass named slots that are marked as instance named slot.
 			for (const FName& InstanceNamedSlot : BPGClass->InstanceNamedSlots)
 			{
@@ -1086,6 +1091,16 @@ void FWidgetBlueprintCompilerContext::FinishCompilingClass(UClass* Class)
 				{
 					if (UWidget* ContentInSlot = Tree->GetContentForSlot(SlotName))
 					{
+						if (NamedSlotClass->NamedSlotsWithContentInSameTree.Contains(SlotName))
+						{
+							UClass* SubClassWithSlotFilled = ContentInSlot->GetTypedOuter<UClass>();
+							MessageLog.Error(
+								*FText::Format(
+									LOCTEXT("NamedSlotAlreadyFilled", "The Named Slot '{0}' already has content in the widget blueprint it was created in but the subclass @@ tried to slot @@ into it. Please remove at least one of the contents."),
+									FText::FromName(SlotName)
+								).ToString(),
+								SubClassWithSlotFilled, ContentInSlot);
+						}
 						if (!NamedSlotContentMap.Contains(SlotName))
 						{
 							NamedSlotContentMap.Add(SlotName, ContentInSlot);
@@ -1094,13 +1109,14 @@ void FWidgetBlueprintCompilerContext::FinishCompilingClass(UClass* Class)
 						{
 							UClass* SubClassWithSlotFilled = ContentInSlot->GetTypedOuter<UClass>();
 							UClass* ParentClassWithSlotFilled = NamedSlotClass;
-							MessageLog.Note(
+							MessageLog.Error(
 								*FText::Format(
-									LOCTEXT("NamedSlotAlreadyFilled", "The Named Slot '{0}' already contains @@ from the class @@ but the subclass @@ tried to slot @@ into it."),
+									LOCTEXT("NamedSlotAlreadyFilled", "The Named Slot '{0}' already contains @@ from the class @@ but the subclass @@ tried to slot @@ into it. Please remove the content @@ from the class @@ to fix this error."),
 									FText::FromName(SlotName)
 								).ToString(),
 								ContentInSlot, ParentClassWithSlotFilled,
-								SubClassWithSlotFilled, NamedSlotContentMap.FindRef(SlotName));
+								SubClassWithSlotFilled, NamedSlotContentMap.FindRef(SlotName),
+								ContentInSlot, ParentClassWithSlotFilled);
 						}
 					}
 				}
@@ -1207,7 +1223,10 @@ void FWidgetBlueprintCompilerContext::FinishCompilingClass(UClass* Class)
 		}
 	}
 
-	BPGClass->bCanCallInitializedWithoutPlayerContext = WidgetBP->bCanCallInitializedWithoutPlayerContext;
+	if (BPGClass)
+	{
+		BPGClass->bCanCallInitializedWithoutPlayerContext = WidgetBP->bCanCallInitializedWithoutPlayerContext;
+	}
 
 	Super::FinishCompilingClass(Class);
 
diff --git a/Engine/Source/Editor/UMGEditor/Private/WidgetBlueprintEditorUtils.cpp b/Engine/Source/Editor/UMGEditor/Private/WidgetBlueprintEditorUtils.cpp
--- a/Engine/Source/Editor/UMGEditor/Private/WidgetBlueprintEditorUtils.cpp
+++ b/Engine/Source/Editor/UMGEditor/Private/WidgetBlueprintEditorUtils.cpp
@@ -1547,27 +1547,39 @@ UUserWidget* FWidgetBlueprintEditorUtils::CreateUserWidgetFromBlueprint(UObject*
 
 		UWidgetTree* LatestWidgetTree = FWidgetBlueprintEditorUtils::FindLatestWidgetTree(BP, CreatedUserWidget);
 
-		TMap<FName, UWidget*> NamedSlotContentToMerge;
+		TMap<FName, UWidget*> SortedNamedSlotContentToMerge;
 		UWidgetBlueprint* WidgetBlueprintIterator = BP;
+		TArray<TTuple<FName, UWidget*>> NamedSlotContentToMergeArray;
+
 		while (WidgetBlueprintIterator)
 		{
 			TArray<FName> SlotNames;
 			WidgetBlueprintIterator->WidgetTree->GetSlotNames(SlotNames);
 
-			for (const FName SlotName : SlotNames)
+			// We iterate widget blueprints from child to parent, but we need the final namedslot array to be sorted from parent to child.
+			// Here, we iterate the slot names in reverse to maintain the order of namedslots per widget blueprint once the final array is reversed.
+			for (int32 Index = SlotNames.Num() - 1; Index >= 0; Index--)
 			{
+				FName SlotName = SlotNames[Index];
 				if (UWidget* Content = WidgetBlueprintIterator->WidgetTree->GetContentForSlot(SlotName))
 				{
-					NamedSlotContentToMerge.Add(SlotName, Content);
+					NamedSlotContentToMergeArray.Add(TTuple<FName, UWidget*>(SlotName, Content));
 				}
 			}
 
 			WidgetBlueprintIterator = Cast<UWidgetBlueprint>(WidgetBlueprintIterator->GeneratedClass->GetSuperClass()->ClassGeneratedBy);
 		}
 
+		// We iterate the array in reverse so that the final SortedNamedSlotContentToMerge map ends up sorted from outermost namedslot to innermost.
+		for (int32 Index = NamedSlotContentToMergeArray.Num() - 1; Index >= 0; Index--)
+		{
+			TTuple<FName, UWidget*>& Element = NamedSlotContentToMergeArray[Index];
+			SortedNamedSlotContentToMerge.Add(Element.Key, Element.Value);
+		}
+		 
 		// Update the widget tree directly to match the blueprint tree.  That way the preview can update
 		// without needing to do a full recompile.
-		CreatedUserWidget->DuplicateAndInitializeFromWidgetTree(LatestWidgetTree, NamedSlotContentToMerge);
+		CreatedUserWidget->DuplicateAndInitializeFromWidgetTree(LatestWidgetTree, SortedNamedSlotContentToMerge);
 
 		// Establish the widget as being in design time before initializing (so that IsDesignTime is reliable within Initialize)
         // We have to call it to make sure that all the WidgetTree had the DesignerFlags set correctly
@@ -3074,7 +3086,7 @@ UWidgetBlueprint* FWidgetBlueprintEditorUtils::GetWidgetBlueprintFromWidget(cons
 			{
 				return WidgetBlueprint;
 			}
-			else
+			else if (WidgetTree->GetOuter())
 			{
 				WidgetBlueprint = Cast<UWidgetBlueprint>(WidgetTree->GetOuter()->GetClass()->ClassGeneratedBy);
 				if (WidgetBlueprint)
diff --git a/Engine/Source/Editor/UMGEditor/Public/WidgetBlueprint.h b/Engine/Source/Editor/UMGEditor/Public/WidgetBlueprint.h
--- a/Engine/Source/Editor/UMGEditor/Public/WidgetBlueprint.h
+++ b/Engine/Source/Editor/UMGEditor/Public/WidgetBlueprint.h
@@ -313,6 +313,9 @@ public:
 	/** Gets any named slots exposed by the parent generated class that can be slotted into by the subclass. */
 	TArray<FName> GetInheritedAvailableNamedSlots() const;
 
+	/** Gets named slots exposed by the parent generated class that already have content in the parent widget tree. */
+	TSet<FName> GetInheritedNamedSlotsWithContentInSameTree() const;
+
 	virtual UWidgetEditingProjectSettings* GetRelevantSettings();
 	virtual const UWidgetEditingProjectSettings* GetRelevantSettings() const;
 
diff --git a/Engine/Source/Runtime/UMG/Public/Blueprint/WidgetBlueprintGeneratedClass.h b/Engine/Source/Runtime/UMG/Public/Blueprint/WidgetBlueprintGeneratedClass.h
--- a/Engine/Source/Runtime/UMG/Public/Blueprint/WidgetBlueprintGeneratedClass.h
+++ b/Engine/Source/Runtime/UMG/Public/Blueprint/WidgetBlueprintGeneratedClass.h
@@ -131,6 +131,9 @@ public:
 	/** All named slots mapped the assigned GUID of their UNamedSlot widget. **/
 	UPROPERTY()
 	TMap<FName, FGuid> NamedSlotsWithID;
+
+	UPROPERTY(Transient)
+	TSet<FName> NamedSlotsWithContentInSameTree;
 #endif
 
 	/**
