UE-221298

#rnx
#tests tried putting a control in the parent's slot and in the derived class to see if the compilation error still outputs correctly.

diff --git a/Engine/Source/Editor/UMGEditor/Private/WidgetBlueprintCompiler.cpp b/Engine/Source/Editor/UMGEditor/Private/WidgetBlueprintCompiler.cpp
--- a/Engine/Source/Editor/UMGEditor/Private/WidgetBlueprintCompiler.cpp
+++ b/Engine/Source/Editor/UMGEditor/Private/WidgetBlueprintCompiler.cpp
@@ -1125,7 +1125,7 @@ void FWidgetBlueprintCompilerContext::FinishCompilingClass(UClass* Class)
 							UClass* SubClassWithSlotFilled = ContentInSlot->GetTypedOuter<UClass>();
 							MessageLog.Error(
 								*FText::Format(
-									LOCTEXT("NamedSlotAlreadyFilled", "The Named Slot '{0}' already has content in the widget blueprint it was created in but the subclass @@ tried to slot @@ into it. Please remove at least one of the contents."),
+									 LOCTEXT("NamedSlotAlreadyFilledInOriginalTree", "The Named Slot '{0}' already has content in the widget blueprint it was created in but the subclass @@ tried to slot @@ into it. Please remove at least one of the contents."),
 									FText::FromName(SlotName)
 								).ToString(),
 								SubClassWithSlotFilled, ContentInSlot);
