PR #12291: UMGEditor: Restore selection of the root preview widget when refreshing the preview

diff --git a/Engine/Source/Editor/UMGEditor/Private/WidgetBlueprintEditor.cpp b/Engine/Source/Editor/UMGEditor/Private/WidgetBlueprintEditor.cpp
--- a/Engine/Source/Editor/UMGEditor/Private/WidgetBlueprintEditor.cpp
+++ b/Engine/Source/Editor/UMGEditor/Private/WidgetBlueprintEditor.cpp
@@ -102,6 +102,7 @@ FWidgetBlueprintEditor::FWidgetBlueprintEditor()
 	, bRefreshGeneratedClassAnimations(false)
 	, bUpdatingSequencerSelection(false)
 	, bUpdatingExternalSelection(false)
+	, bIsPreviewWidgetSelected(false)
 {
 	PreviewScene.GetWorld()->SetBegunPlay(false);
 
@@ -618,6 +619,7 @@ void FWidgetBlueprintEditor::SelectWidgets(const TSet<FWidgetReference>& Widgets
 	}
 	SelectedObjects.Empty();
 	SelectedNamedSlot.Reset();
+	bIsPreviewWidgetSelected = false;
 
 	for ( const FWidgetReference& Widget : TempSelection )
 	{
@@ -641,10 +643,15 @@ void FWidgetBlueprintEditor::SelectObjects(const TSet<UObject*>& Objects)
 	SelectedWidgets.Empty();
 	SelectedObjects.Empty();
 	SelectedNamedSlot.Reset();
+	bIsPreviewWidgetSelected = false;
 
 	for ( UObject* Obj : Objects )
 	{
 		SelectedObjects.Add(Obj);
+		if (Obj == GetPreview())
+		{
+			bIsPreviewWidgetSelected = true;
+		}
 	}
 
 	OnSelectedWidgetsChanged.Broadcast();
@@ -685,6 +692,7 @@ void FWidgetBlueprintEditor::SetSelectedNamedSlot(TOptional<FNamedSlotSelection>
 	SelectedWidgets.Empty();
 	SelectedObjects.Empty();
 	SelectedNamedSlot.Reset();
+	bIsPreviewWidgetSelected = false;
 
 	SelectedNamedSlot = InSelectedNamedSlot;
 	if (InSelectedNamedSlot.IsSet())
@@ -721,6 +729,19 @@ void FWidgetBlueprintEditor::CleanSelection()
 	{
 		SelectWidgets(TempSelection, false);
 	}
+	else if ( bIsPreviewWidgetSelected )
+	{
+		TSet<UObject*> ValidObjects = { GetPreview() };
+		for ( const TWeakObjectPtr<UObject>& SelectedObject : SelectedObjects )
+		{
+			if ( SelectedObject.IsValid() )
+			{
+				ValidObjects.Add(SelectedObject.Get());
+			}
+		}
+
+		SelectObjects(ValidObjects);
+	}
 }
 
 const TSet<FWidgetReference>& FWidgetBlueprintEditor::GetSelectedWidgets() const
diff --git a/Engine/Source/Editor/UMGEditor/Public/WidgetBlueprintEditor.h b/Engine/Source/Editor/UMGEditor/Public/WidgetBlueprintEditor.h
--- a/Engine/Source/Editor/UMGEditor/Public/WidgetBlueprintEditor.h
+++ b/Engine/Source/Editor/UMGEditor/Public/WidgetBlueprintEditor.h
@@ -546,4 +546,7 @@ private:
 
 	/** When true the external selection is being updated from changes to the sequencer selection. */
 	bool bUpdatingExternalSelection;
+
+	/** True when the selected object is the preview widget. Used to restore selection after refreshing the preview.*/
+	bool bIsPreviewWidgetSelected;
 };
\ No newline at end of file
