UMG: Support fixing up references to renamed widgets and animations in unloaded blueprints

- Expose Generated Variable population to Widget Blueprint Extensions via UWidgetBlueprintExtension::PopulateGeneratedVariables
- Refactor Widget/Animation variable creation to use Generated Variables instead of directly calling CreateVariable
- Track GUIDs for each widget/animation in a widget BP and update the list as the user modifies the blueprint
  - Existing Widget BPs will initialize with deterministic GUIDs for all its variables
- References will need to be saved after this CL to be able to auto-fix up
- Improve the QoL of "Replace With..." to keep the same widget name when possible
- Fix cases where the user could be prompted to cancel deleting a widget when inappropriate to cancel
- Fix case where we would create a widget just to not use it, leaving it in the widget tree's child objects

diff --git a/Engine/Source/Editor/UMGEditor/Private/Hierarchy/SHierarchyViewItem.cpp b/Engine/Source/Editor/UMGEditor/Private/Hierarchy/SHierarchyViewItem.cpp
--- a/Engine/Source/Editor/UMGEditor/Private/Hierarchy/SHierarchyViewItem.cpp
+++ b/Engine/Source/Editor/UMGEditor/Private/Hierarchy/SHierarchyViewItem.cpp
@@ -40,6 +40,11 @@
 
 #define LOCTEXT_NAMESPACE "UMG"
 
+namespace UE::UMG::Editor
+{
+	extern bool bDisplayWidgetVariableGuids;
+}
+
 /**
 *
 */
@@ -239,6 +244,7 @@ TOptional<EItemDropZone> ProcessHierarchyDragDrop(const FDragDropEvent& DragDrop
 				if (UWidget* Widget = FWidgetBlueprintEditorUtils::GetWidgetTemplateFromDragDrop(Blueprint, Blueprint->WidgetTree, DragDropOp))
 				{
 					Blueprint->WidgetTree->RootWidget = Widget;
+					Blueprint->OnVariableAdded(Widget->GetFName());
 				}
 				FBlueprintEditorUtils::MarkBlueprintAsStructurallyModified(Blueprint);
 			}
@@ -284,6 +290,8 @@ TOptional<EItemDropZone> ProcessHierarchyDragDrop(const FDragDropEvent& DragDrop
 						}
 						check(NewSlot);
 
+						Blueprint->OnVariableAdded(Widget->GetFName());
+
 						FBlueprintEditorUtils::MarkBlueprintAsStructurallyModified(Blueprint);
 					}
 				}
@@ -441,8 +449,9 @@ TOptional<EItemDropZone> ProcessHierarchyDragDrop(const FDragDropEvent& DragDrop
 						}
 					}
 
+					const FName OriginalWidgetName = TemplateWidget->GetFName();
 					UPanelWidget* OriginalParent = TemplateWidget->GetParent();
-					UBlueprint* OriginalBP = nullptr;
+					UWidgetBlueprint* OriginalBP = nullptr;
 
 					// The widget's parent is changing
 					if (OriginalParent != NewParent)
@@ -460,7 +469,7 @@ TOptional<EItemDropZone> ProcessHierarchyDragDrop(const FDragDropEvent& DragDrop
 							OriginalWidgetTree->SetFlags(RF_Transactional);
 							OriginalWidgetTree->Modify();
 
-							OriginalBP = OriginalWidgetTree->GetTypedOuter<UBlueprint>();
+							OriginalBP = OriginalWidgetTree->GetTypedOuter<UWidgetBlueprint>();
 						}
 					}
 
@@ -468,6 +477,7 @@ TOptional<EItemDropZone> ProcessHierarchyDragDrop(const FDragDropEvent& DragDrop
 					
 					if (OriginalBP != nullptr && OriginalBP != Blueprint)
 					{
+						OriginalBP->OnVariableRemoved(OriginalWidgetName);
 						FBlueprintEditorUtils::MarkBlueprintAsStructurallyModified(OriginalBP);
 					}
 
@@ -482,6 +492,11 @@ TOptional<EItemDropZone> ProcessHierarchyDragDrop(const FDragDropEvent& DragDrop
 						NewSlot = NewParent->AddChild(TemplateWidget);
 					}
 					check(NewSlot);
+					
+					if (OriginalBP != nullptr && OriginalBP != Blueprint)
+					{
+						Blueprint->OnVariableAdded(TemplateWidget->GetFName());
+					}
 
 					// Import the old slot properties
 					FWidgetBlueprintEditorUtils::ImportPropertiesFromText(NewSlot, DraggedWidget.ExportedSlotProperties);
@@ -1036,6 +1051,7 @@ FReply FNamedSlotModelBase::HandleAcceptDrop(FDragDropEvent const& DragDropEvent
 
 		if (UWidget* DroppingWidget = FWidgetBlueprintEditorUtils::GetWidgetTemplateFromDragDrop(Blueprint, Blueprint->WidgetTree, DragDropOp))
 		{
+			Blueprint->OnVariableAdded(DroppingWidget->GetFName());
 			DoDrop(NamedSlotHost, DroppingWidget);
 			return FReply::Handled();
 		}
@@ -1267,14 +1283,38 @@ FText FHierarchyWidget::GetImageToolTipText() const
 
 FText FHierarchyWidget::GetLabelToolTipText() const
 {
-	// If the user has provided a name, give a tooltip with the widget type for easy reference
-	UWidget* WidgetTemplate = Item.GetTemplate();
-	if ( WidgetTemplate && !WidgetTemplate->IsGeneratedName() )
+	FText ToolTipText = FText::GetEmpty();
+
+	auto AppendNewLineOrSetValue = [&ToolTipText](const FText& NewText)
 	{
-		return FText::FromString(TEXT( "[" ) + WidgetTemplate->GetClass()->GetDisplayNameText().ToString() + TEXT( "]" ) );
+		static const FTextFormat ToolTipNewLineFormat = LOCTEXT("ToolTipNewLineFormat", "{0}\n\n{1}");
+		ToolTipText = ToolTipText.IsEmpty() ? NewText : FText::Format(ToolTipNewLineFormat, ToolTipText, NewText);
+	};
+
+	if (UWidget* WidgetTemplate = Item.GetTemplate())
+	{
+		// If the user has provided a name, give a tooltip with the widget type for easy reference
+		if (!WidgetTemplate->IsGeneratedName())
+		{
+			AppendNewLineOrSetValue(FText::FromString(TEXT( "[" ) + WidgetTemplate->GetClass()->GetDisplayNameText().ToString() + TEXT( "]" ) ));
+		}
+	
+		if (UE::UMG::Editor::bDisplayWidgetVariableGuids)
+		{
+			TSharedPtr<FWidgetBlueprintEditor> WidgetBlueprintEditor = BlueprintEditor.Pin();
+			if (WidgetBlueprintEditor.IsValid())
+			{
+				if (const UWidgetBlueprint* WidgetBlueprint = WidgetBlueprintEditor->GetWidgetBlueprintObj())
+				{
+					static const FTextFormat VariableGUIDFormat = LOCTEXT("WidgetVariableGUIDFormat", "GUID: {0}");
+					const FGuid VariableGuid = WidgetBlueprint->WidgetVariableNameToGuidMap.FindRef(WidgetTemplate->GetFName());
+					AppendNewLineOrSetValue(FText::Format(VariableGUIDFormat, FText::FromString(VariableGuid.ToString(EGuidFormats::DigitsWithHyphens))));
+				}
+			}
+		}
 	}
 
-	return FText::GetEmpty();
+	return ToolTipText;
 }
 
 void FHierarchyWidget::GetFilterStrings(TArray<FString>& OutStrings) const
diff --git a/Engine/Source/Editor/UMGEditor/Private/TabFactory/AnimationTabSummoner.cpp b/Engine/Source/Editor/UMGEditor/Private/TabFactory/AnimationTabSummoner.cpp
--- a/Engine/Source/Editor/UMGEditor/Private/TabFactory/AnimationTabSummoner.cpp
+++ b/Engine/Source/Editor/UMGEditor/Private/TabFactory/AnimationTabSummoner.cpp
@@ -39,6 +39,11 @@
 
 #define LOCTEXT_NAMESPACE "UMG"
 
+namespace UE::UMG::Editor
+{
+	extern bool bDisplayWidgetVariableGuids;
+}
+
 const FName FAnimationTabSummoner::TabID(TEXT("Animations"));
 const FName FAnimationTabSummoner::WidgetAnimSequencerDrawerID(TEXT("WidgetAnimSequencer"));
 
@@ -187,6 +192,7 @@ public:
 				SAssignNew(InlineTextBlock, SInlineEditableTextBlock)
 				.Font(FCoreStyle::Get().GetFontStyle("NormalFont"))
 				.Text(this, &SWidgetAnimationListItem::GetMovieSceneText)
+				.ToolTipText(this, &SWidgetAnimationListItem::GetMovieSceneToolTipText)
 				//.HighlightText(InArgs._HighlightText)
 				.OnVerifyTextChanged(this, &SWidgetAnimationListItem::OnVerifyNameTextChanged)
 				.OnTextCommitted(this, &SWidgetAnimationListItem::OnNameTextCommited)
@@ -211,6 +217,23 @@ private:
 		return FText::GetEmpty();
 	}
 
+	FText GetMovieSceneToolTipText() const
+	{
+		if (UE::UMG::Editor::bDisplayWidgetVariableGuids && ListItem.IsValid() && BlueprintEditor.IsValid())
+		{
+			const UWidgetAnimation* WidgetAnimation = ListItem.Pin()->Animation;
+			const UWidgetBlueprint* Blueprint = BlueprintEditor.Pin()->GetWidgetBlueprintObj();
+			if (WidgetAnimation && Blueprint)
+			{
+				static const FTextFormat VariableGUIDFormat = LOCTEXT("WidgetVariableGUIDFormat", "GUID: {0}");
+				const FGuid VariableGuid = Blueprint->WidgetVariableNameToGuidMap.FindRef(WidgetAnimation->GetFName());
+				return FText::Format(VariableGUIDFormat, FText::FromString(VariableGuid.ToString(EGuidFormats::DigitsWithHyphens)));
+			}
+		}
+
+		return FText::GetEmpty();
+	}
+
 	bool OnVerifyNameTextChanged(const FText& InText, FText& OutErrorMessage)
 	{
 		UWidgetAnimation* Animation = ListItem.Pin()->Animation;
@@ -261,6 +284,7 @@ private:
 				{
 					Blueprint->Modify();
 					Blueprint->Animations.Add(WidgetAnimation);
+					Blueprint->OnVariableAdded(WidgetAnimation->GetFName());
 					ListItem.Pin()->bNewAnimation = false;
 
 					if (TSharedPtr<FWidgetBlueprintEditor> WidgetBlueprintEditorPin = BlueprintEditor.Pin())
@@ -268,6 +292,10 @@ private:
 						WidgetBlueprintEditorPin->NotifyWidgetAnimListChanged();
 					}
 				}
+				else
+				{
+					Blueprint->OnVariableRenamed(OldFName, NewFName);
+				}
 			}
 
 			FBlueprintEditorUtils::ReplaceVariableReferences(Blueprint, OldFName, NewFName);
@@ -278,6 +306,7 @@ private:
 			const FScopedTransaction Transaction(LOCTEXT("NewAnimation", "New Animation"));
 			Blueprint->Modify();
 			Blueprint->Animations.Add(WidgetAnimation);
+			Blueprint->OnVariableAdded(WidgetAnimation->GetFName());
 			ListItem.Pin()->bNewAnimation = false;
 			FBlueprintEditorUtils::MarkBlueprintAsStructurallyModified(Blueprint);
 
@@ -852,9 +881,11 @@ private:
 		{
 			const FScopedTransaction Transaction(LOCTEXT("DeleteAnimationTransaction", "Delete Animation"));
 			WidgetBlueprint->Modify();
+			const FName RemovedAnimationName = SelectedAnimation->Animation->GetFName();
 			// Rename the animation and move it to the transient package to avoid collisions.
 			SelectedAnimation->Animation->Rename( NULL, GetTransientPackage() );
 			WidgetAnimations.Remove(SelectedAnimation->Animation);
+			WidgetBlueprint->OnVariableRemoved(RemovedAnimationName);
 
 			UpdateAnimationList();
 		}
diff --git a/Engine/Source/Editor/UMGEditor/Private/WidgetBlueprint.cpp b/Engine/Source/Editor/UMGEditor/Private/WidgetBlueprint.cpp
--- a/Engine/Source/Editor/UMGEditor/Private/WidgetBlueprint.cpp
+++ b/Engine/Source/Editor/UMGEditor/Private/WidgetBlueprint.cpp
@@ -2,6 +2,7 @@
 
 #include "WidgetBlueprint.h"
 
+#include "Components/NamedSlot.h"
 #include "Components/Widget.h"
 #include "Blueprint/UserWidget.h"
 #include "MovieScene.h"
@@ -38,6 +39,8 @@
 #include "K2Node_Composite.h"
 #include "K2Node_FunctionResult.h"
 #include "Blueprint/WidgetNavigation.h"
+#include "WidgetBlueprintEditorUtils.h"
+#include "WidgetBlueprintExtension.h"
 #include "WidgetEditingProjectSettings.h"
 
 #define LOCTEXT_NAMESPACE "UMG"
@@ -562,6 +565,8 @@ UWidgetBlueprint::UWidgetBlueprint(const FObjectInitializer& ObjectInitializer)
 	, bCanCallInitializedWithoutPlayerContext(false)
 	, TickFrequency(EWidgetTickFrequency::Auto)
 {
+	// Widget GUIDs must be cooked to support fixups at runtime, so enable by default
+	ShouldCookPropertyGuidsValue = EShouldCookBlueprintPropertyGuids::Yes;
 }
 
 void UWidgetBlueprint::ReplaceDeprecatedNodes()
@@ -698,9 +703,6 @@ void UWidgetBlueprint::PreSave(FObjectPreSaveContext ObjectSaveContext)
 {
 	Super::PreSave(ObjectSaveContext);
 }
-#endif // WITH_EDITORONLY_DATA
-
-#if WITH_EDITORONLY_DATA
 
 void UWidgetBlueprint::GetAssetRegistryTags(TArray<FAssetRegistryTag>& OutTags) const
 {
@@ -981,7 +983,41 @@ bool UWidgetBlueprint::FindDiffs(const UBlueprint* OtherBlueprint, FDiffResults&
 	return true;
 }
 
-#endif
+void UWidgetBlueprint::OnVariableAdded(const FName& VariableName)
+{
+	ensure(!WidgetVariableNameToGuidMap.Contains(VariableName));
+	Modify();
+	WidgetVariableNameToGuidMap.Emplace(VariableName, FGuid::NewGuid());
+}
+
+void UWidgetBlueprint::OnVariableRenamed(const FName& OldName, const FName& NewName)
+{
+	if (OldName == NewName)
+	{
+		return;
+	}
+
+	// We should never have an entry for the new name and always have an entry for the current name
+	ensure(!WidgetVariableNameToGuidMap.Contains(NewName));
+	
+	FGuid VariableGuid;
+	if (ensure(WidgetVariableNameToGuidMap.RemoveAndCopyValue(OldName, VariableGuid)))
+	{
+		Modify();
+		WidgetVariableNameToGuidMap.Emplace(NewName, VariableGuid);
+	}
+	else
+	{
+		OnVariableAdded(NewName);
+	}
+}
+
+void UWidgetBlueprint::OnVariableRemoved(const FName& VariableName)
+{
+	Modify();
+	WidgetVariableNameToGuidMap.Remove(VariableName);
+}
+#endif // WITH_EDITORONLY_DATA
 
 void UWidgetBlueprint::Serialize(FArchive& Ar)
 {
diff --git a/Engine/Source/Editor/UMGEditor/Private/WidgetBlueprintCompiler.cpp b/Engine/Source/Editor/UMGEditor/Private/WidgetBlueprintCompiler.cpp
--- a/Engine/Source/Editor/UMGEditor/Private/WidgetBlueprintCompiler.cpp
+++ b/Engine/Source/Editor/UMGEditor/Private/WidgetBlueprintCompiler.cpp
@@ -33,6 +33,38 @@
 
 extern COREUOBJECT_API bool GMinimalCompileOnLoad;
 
+namespace UE::WidgetBlueprintCompiler::Private
+{
+	FProperty* FindChildProperty(const UStruct* Struct, const FName& PropertyName)
+	{
+		for (FField* Field = Struct->ChildProperties; Field != nullptr; Field = Field->Next)
+		{
+			if (Field->GetFName() == PropertyName)
+			{
+				return CastField<FProperty>(Field);
+			}
+		}
+				
+		return nullptr;
+	}
+}
+
+//////////////////////////////////////////////////////////////////////////
+// FWidgetBlueprintCompiler::FPopulateGeneratedVariablesContext
+FWidgetBlueprintCompilerContext::FPopulateGeneratedVariablesContext::FPopulateGeneratedVariablesContext(FWidgetBlueprintCompilerContext& InContext)
+	: Context(InContext)
+{}
+
+void FWidgetBlueprintCompilerContext::FPopulateGeneratedVariablesContext::AddGeneratedVariable(FBPVariableDescription&& VariableDescription) const
+{
+	Context.AddGeneratedVariable(MoveTemp(VariableDescription));
+}
+
+UWidgetBlueprint* FWidgetBlueprintCompilerContext::FPopulateGeneratedVariablesContext::GetWidgetBlueprint() const
+{
+	return Context.WidgetBlueprint();
+}
+
 //////////////////////////////////////////////////////////////////////////
 // FWidgetBlueprintCompiler::FCreateVariableContext
 FWidgetBlueprintCompilerContext::FCreateVariableContext::FCreateVariableContext(FWidgetBlueprintCompilerContext& InContext)
@@ -268,21 +300,6 @@ void FWidgetBlueprintCompilerContext::CreateFunctionList()
 	}
 }
 
-void FWidgetBlueprintCompilerContext::ValidateWidgetNames()
-{
-	UWidgetBlueprint* WidgetBP = WidgetBlueprint();
-
-	TSharedPtr<FKismetNameValidator> ParentBPNameValidator;
-	if ( WidgetBP->ParentClass != nullptr )
-	{
-		UBlueprint* ParentBP = Cast<UBlueprint>(WidgetBP->ParentClass->ClassGeneratedBy);
-		if ( ParentBP != nullptr )
-		{
-			ParentBPNameValidator = MakeShareable(new FKismetNameValidator(ParentBP));
-		}
-	}
-}
-
 template<typename TOBJ>
 struct FCullTemplateObjectsHelper
 {
@@ -510,13 +527,10 @@ void FWidgetBlueprintCompilerContext::CreateClassVariablesFromBlueprint()
 
 	UClass* ParentClass = WidgetBP->ParentClass;
 
-	ValidateWidgetNames();
-
 	// Build the set of variables based on the variable widgets in the first Widget Tree we find:
 	// in the current blueprint, the parent blueprint, and so on, until we find one.
 	TArray<UWidget*> Widgets;
 	UWidgetBlueprint* WidgetBPToScan = WidgetBP;
-	bool bSkipVariableCreation = false;
 	while (WidgetBPToScan != nullptr)
 	{
 		Widgets = WidgetBPToScan->GetAllSourceWidgets();
@@ -547,10 +561,6 @@ void FWidgetBlueprintCompilerContext::CreateClassVariablesFromBlueprint()
 
 			break;
 		}
-
-		// We don't want to create variables for widgets that are in a parent blueprint. They will be created at the Parent compilation.
-		// But we want them to be added to the Member variable map for validation of the BindWidget property
-		bSkipVariableCreation = true;
 		
 		// Get the parent WidgetBlueprint
 		WidgetBPToScan = WidgetBPToScan->ParentClass && WidgetBPToScan->ParentClass->ClassGeneratedBy 
@@ -558,32 +568,9 @@ void FWidgetBlueprintCompilerContext::CreateClassVariablesFromBlueprint()
 			: nullptr;
 	}
 
-	// Sort the widgets alphabetically
-	Widgets.Sort( []( const UWidget& Lhs, const UWidget& Rhs ) { return Rhs.GetFName().LexicalLess(Lhs.GetFName()); } );
-
 	// Add widget variables
 	for ( UWidget* Widget : Widgets ) 
 	{
-		bool bIsVariable = Widget->bIsVariable;
-
-		// In the event there are bindings for a widget, but it's not marked as a variable, make it one, but hide it from the UI.
-		// we do this so we can use FindField to locate it at runtime.
-		bIsVariable |= WidgetBP->Bindings.ContainsByPredicate([&Widget] (const FDelegateEditorBinding& Binding) {
-			return Binding.ObjectName == Widget->GetName();
-		});
-
-		// All UNamedSlot widgets are automatically variables, so that we can properly look them up quickly with FindField
-		// in UserWidgets.
-		bIsVariable |= Widget->IsA<UNamedSlot>();
-
-		// This code was added to fix the problem of recompiling dependent widgets, not using the newest
-		// class thus resulting in REINST failures in dependent blueprints.
-		UClass* WidgetClass = Widget->GetClass();
-		if ( UBlueprintGeneratedClass* BPWidgetClass = Cast<UBlueprintGeneratedClass>(WidgetClass) )
-		{
-			WidgetClass = BPWidgetClass->GetAuthoritativeClass();
-		}
-
 		// Look in the Parent class properties to find a property with the BindWidget meta tag of the same name and Type.
 		FObjectPropertyBase* ExistingProperty = CastField<FObjectPropertyBase>(ParentClass->FindPropertyByName(Widget->GetFName()));
 		if (ExistingProperty && 
@@ -594,42 +581,19 @@ void FWidgetBlueprintCompilerContext::CreateClassVariablesFromBlueprint()
 			continue;
 		}
 
-		// Skip non-variable widgets
-		if ( !bIsVariable )
-		{
-			continue;
-		}
-
-        // We skip variable creation if the Widget Tree was in the Parent Blueprint.
-		if (bSkipVariableCreation)
+		// Check if the widget has a generated variable
+		for (const FBPVariableDescription& VarDesc : WidgetBP->GeneratedVariables)
 		{
-			continue;
-		}
-
-		FEdGraphPinType WidgetPinType(UEdGraphSchema_K2::PC_Object, NAME_None, WidgetClass, EPinContainerType::None, false, FEdGraphTerminalType());
-		
-		// Always name the variable according to the underlying FName of the widget object
-		FProperty* WidgetProperty = CreateVariable(Widget->GetFName(), WidgetPinType);
-		if ( WidgetProperty != nullptr )
-		{
-			const FString DisplayName = Widget->IsGeneratedName() ? Widget->GetName() : Widget->GetLabelText().ToString();
-			WidgetProperty->SetMetaData(TEXT("DisplayName"), *DisplayName);
-			
-			// Only show variables if they're explicitly marked as variables.
-			if ( Widget->bIsVariable )
+			if (VarDesc.VarName == Widget->GetFName())
 			{
-				WidgetProperty->SetPropertyFlags(CPF_BlueprintVisible | CPF_BlueprintReadOnly | CPF_DisableEditOnInstance);
-
-				const FString& CategoryName = Widget->GetCategoryName();
-				
-				// Only include Category metadata for variables (i.e. a visible/editable property); otherwise, UHT will raise a warning if this Blueprint is nativized.
-				WidgetProperty->SetMetaData(TEXT("Category"), *(CategoryName.IsEmpty() ? WidgetBP->GetName() : CategoryName));
+				FProperty* WidgetProperty = UE::WidgetBlueprintCompiler::Private::FindChildProperty(NewClass, VarDesc.VarName);
+				if (ensureMsgf(WidgetProperty, TEXT("The Widget Blueprint [%s] has a generated variable for the widget [%s] but we failed to find a property for it."), *WidgetBP->GetName(), *Widget->GetName()))
+				{
+					WidgetToMemberVariableMap.Add(Widget, WidgetProperty);
+				}
+						
+				break;
 			}
-
-			WidgetProperty->SetPropertyFlags(CPF_Instanced);
-			WidgetProperty->SetPropertyFlags(CPF_RepSkip);
-
-			WidgetToMemberVariableMap.Add(Widget, WidgetProperty);
 		}
 	}
 
@@ -651,22 +615,19 @@ void FWidgetBlueprintCompilerContext::CreateClassVariablesFromBlueprint()
 			// Create variables for widget animation
 			if (WidgetBPToScan == WidgetBP)
 			{
-				FEdGraphPinType WidgetPinType(UEdGraphSchema_K2::PC_Object, NAME_None, Animation->GetClass(), EPinContainerType::None, true, FEdGraphTerminalType());
-				FProperty* AnimationProperty = CreateVariable(Animation->GetFName(), WidgetPinType);
-
-				if (AnimationProperty != nullptr)
+				// Check if the animation has a generated variable
+				for (const FBPVariableDescription& VarDesc : WidgetBP->GeneratedVariables)
 				{
-					const FString DisplayName = Animation->GetDisplayName().ToString();
-					AnimationProperty->SetMetaData(TEXT("DisplayName"), *DisplayName);
-
-					AnimationProperty->SetMetaData(TEXT("Category"), TEXT("Animations"));
-
-					AnimationProperty->SetPropertyFlags(CPF_Transient);
-					AnimationProperty->SetPropertyFlags(CPF_BlueprintVisible);
-					AnimationProperty->SetPropertyFlags(CPF_BlueprintReadOnly);
-					AnimationProperty->SetPropertyFlags(CPF_RepSkip);
-
-					WidgetAnimToMemberVariableMap.Add(Animation, AnimationProperty);
+					if (VarDesc.VarName == Animation->GetFName())
+					{
+						FProperty* AnimationProperty = UE::WidgetBlueprintCompiler::Private::FindChildProperty(NewClass, VarDesc.VarName);
+						if (ensureMsgf(AnimationProperty, TEXT("The Widget Blueprint [%s] has a generated variable for the animation [%s] but we failed to find a property for it."), *WidgetBP->GetName(), *Animation->GetName()))
+						{
+							WidgetAnimToMemberVariableMap.Add(Animation, AnimationProperty);
+						}
+						
+						break;
+					}
 				}
 			}
 		}
@@ -799,6 +760,79 @@ void FWidgetBlueprintCompilerContext::SanitizeBindings(UBlueprintGeneratedClass*
 	WidgetBP->PropertyBindings = AttributeBindings;
 }
 
+void FWidgetBlueprintCompilerContext::ValidateAndFixUpVariableGuids()
+{
+	UWidgetBlueprint* WidgetBP = WidgetBlueprint();
+
+	// If we don't yet have any tracked variable guids, populate them deterministically
+	// The determinism is required so that stable guids are serialized when creating external references to this widget's variables before this widget is resaved
+	if (WidgetBP->WidgetVariableNameToGuidMap.IsEmpty())
+	{
+		WidgetBP->ForEachSourceWidget([WidgetBP](UWidget* Widget)
+		{
+			ensureAlways(!WidgetBP->WidgetVariableNameToGuidMap.Contains(Widget->GetFName()));
+			WidgetBP->WidgetVariableNameToGuidMap.Emplace(Widget->GetFName(), FGuid::NewDeterministicGuid(Widget->GetPathName()));
+		});
+
+		for (UWidgetAnimation* Animation : WidgetBP->Animations)
+		{
+			if (Animation)
+			{
+				ensureAlways(!WidgetBP->WidgetVariableNameToGuidMap.Contains(Animation->GetFName()));
+				WidgetBP->WidgetVariableNameToGuidMap.Emplace(Animation->GetFName(), FGuid::NewDeterministicGuid(Animation->GetPathName()));
+			}
+		}
+	}
+	else
+	{
+		// Validate that our variable guids are properly tracked and fixup issues that may have been caused by missed cases
+
+		// Verify all variables have a stored GUID
+		TSet<FName> SeenVariableNames;
+		WidgetBP->ForEachSourceWidget([WidgetBP, &SeenVariableNames](UWidget* Widget)
+		{
+			if (!ensureAlwaysMsgf(WidgetBP->WidgetVariableNameToGuidMap.Contains(Widget->GetFName()), TEXT("Widget [%s] was added but did not get a GUID"), *Widget->GetName()))
+			{
+				WidgetBP->WidgetVariableNameToGuidMap.Add(Widget->GetFName(), FGuid::NewGuid());
+			}
+			SeenVariableNames.Add(Widget->GetFName());
+		});
+
+		for (UWidgetAnimation* Animation : WidgetBP->Animations)
+		{
+			if (Animation)
+			{
+				if (!ensureAlwaysMsgf(WidgetBP->WidgetVariableNameToGuidMap.Contains(Animation->GetFName()), TEXT("Animation [%s] was added but did not get a GUID"), *Animation->GetName()))
+				{
+					WidgetBP->WidgetVariableNameToGuidMap.Add(Animation->GetFName(), FGuid::NewGuid());
+				}
+				SeenVariableNames.Add(Animation->GetFName());
+			}
+		}
+
+		TMap<FGuid, FName> GuidToVariableNameMap;
+
+		// Verify we're only storing GUIDs for variables we still have and that none collide
+		for (auto It = WidgetBP->WidgetVariableNameToGuidMap.CreateIterator(); It; ++It)
+		{
+			if (!ensureAlwaysMsgf(It.Value().IsValid(), TEXT("Variable [%s] has an invalid GUID"), *It.Key().ToString()))
+			{
+				It.Value() = FGuid::NewGuid();
+			}
+			
+			if (ensureAlwaysMsgf(!GuidToVariableNameMap.Contains(It.Value()), TEXT("The variables [%s] and [%s] have the same GUID, delete and recreate one of them to fix this error"), *It.Key().ToString(), *GuidToVariableNameMap[It.Value()].ToString()))
+			{
+				GuidToVariableNameMap.Add(It.Value(), It.Key());
+			}
+
+			if (!ensureAlwaysMsgf(SeenVariableNames.Contains(It.Key()), TEXT("Variable [%s] was deleted but still has a GUID"), *It.Key().ToString()))
+			{
+				It.RemoveCurrent();
+			}
+		}
+	}
+}
+
 void FWidgetBlueprintCompilerContext::FixAbandonedWidgetTree(UWidgetBlueprint* WidgetBP)
 {
 	UWidgetTree* WidgetTree = WidgetBP->WidgetTree;
@@ -1399,6 +1433,110 @@ void FWidgetBlueprintCompilerContext::EnsureProperGeneratedClass(UClass*& Target
 	}
 }
 
+void FWidgetBlueprintCompilerContext::PopulateBlueprintGeneratedVariables()
+{
+	Super::PopulateBlueprintGeneratedVariables();
+
+	ValidateAndFixUpVariableGuids();
+	
+	UWidgetBlueprint* WidgetBP = WidgetBlueprint();
+
+	// Widget Variables
+	{
+		TArray<UWidget*> SortedWidgets = WidgetBP->GetAllSourceWidgets();
+		SortedWidgets.Sort( []( const UWidget& Lhs, const UWidget& Rhs ) { return Rhs.GetFName().LexicalLess(Lhs.GetFName()); } );
+
+		for (UWidget* Widget : SortedWidgets)
+		{
+			// All UNamedSlot widgets are automatically variables, so that we can properly look them up quickly with FindField
+			// in UserWidgets.
+			// In the event there are bindings for a widget, but it's not marked as a variable, make it one, but hide it from the UI.
+			// we do this so we can use FindField to locate it at runtime.
+			const bool bShouldGenerateVariable = Widget->bIsVariable
+				|| Widget->IsA<UNamedSlot>()
+				|| WidgetBP->Bindings.ContainsByPredicate([&Widget] (const FDelegateEditorBinding& Binding) {
+					return Binding.ObjectName == Widget->GetName();
+				});
+
+			if (!bShouldGenerateVariable)
+			{
+				continue;
+			}
+
+			// Look in the Parent class properties to find a property with the BindWidget meta tag of the same name and Type.
+			FObjectPropertyBase* ExistingProperty = CastField<FObjectPropertyBase>(WidgetBP->ParentClass->FindPropertyByName(Widget->GetFName()));
+			if (ExistingProperty && 
+				FWidgetBlueprintEditorUtils::IsBindWidgetProperty(ExistingProperty) && 
+				Widget->IsA(ExistingProperty->PropertyClass))
+			{
+				continue;
+			}
+
+			// This code was added to fix the problem of recompiling dependent widgets, not using the newest
+			// class thus resulting in REINST failures in dependent blueprints.
+			UClass* WidgetClass = Widget->GetClass();
+			if ( UBlueprintGeneratedClass* BPWidgetClass = Cast<UBlueprintGeneratedClass>(WidgetClass) )
+			{
+				WidgetClass = BPWidgetClass->GetAuthoritativeClass();
+			}
+			
+			ensure(WidgetBP->WidgetVariableNameToGuidMap.Contains(Widget->GetFName()));
+
+			FBPVariableDescription WidgetVariableDesc;
+			WidgetVariableDesc.VarName = Widget->GetFName();
+			WidgetVariableDesc.VarGuid = WidgetBP->WidgetVariableNameToGuidMap.FindRef(Widget->GetFName());
+			WidgetVariableDesc.VarType = FEdGraphPinType(UEdGraphSchema_K2::PC_Object, NAME_None, WidgetClass, EPinContainerType::None, false, FEdGraphTerminalType());
+			WidgetVariableDesc.FriendlyName = Widget->IsGeneratedName() ? Widget->GetName() : Widget->GetLabelText().ToString();
+			WidgetVariableDesc.PropertyFlags = (CPF_PersistentInstance | CPF_ExportObject | CPF_InstancedReference | CPF_RepSkip);
+		
+			// Only show variables if they're explicitly marked as variables.
+			if (Widget->bIsVariable)
+			{
+				WidgetVariableDesc.PropertyFlags |= (CPF_BlueprintVisible | CPF_BlueprintReadOnly | CPF_DisableEditOnInstance);
+
+				// Only include Category metadata for variables (i.e. a visible/editable property); otherwise, UHT will raise a warning if this Blueprint is nativized.
+				const FString& CategoryName = Widget->GetCategoryName();
+				WidgetVariableDesc.SetMetaData(TEXT("Category"), *(CategoryName.IsEmpty() ? WidgetBP->GetName() : CategoryName));
+			}
+
+			AddGeneratedVariable(MoveTemp(WidgetVariableDesc));
+		}
+	}
+
+	// Animation Variables
+	{
+		for (UWidgetAnimation* Animation : WidgetBP->Animations)
+		{
+			// BindWidgetAnims already have properties
+			FObjectPropertyBase* ExistingProperty = CastField<FObjectPropertyBase>(WidgetBP->ParentClass->FindPropertyByName(Animation->GetFName()));
+			if (ExistingProperty &&
+				FWidgetBlueprintEditorUtils::IsBindWidgetAnimProperty(ExistingProperty) &&
+				ExistingProperty->PropertyClass->IsChildOf(UWidgetAnimation::StaticClass()))
+			{
+				continue;
+			}
+
+			ensure(WidgetBP->WidgetVariableNameToGuidMap.Contains(Animation->GetFName()));
+			FBPVariableDescription AnimVariableDesc;
+			AnimVariableDesc.VarName = Animation->GetFName();
+			AnimVariableDesc.VarGuid = WidgetBP->WidgetVariableNameToGuidMap.FindRef(Animation->GetFName());
+			AnimVariableDesc.VarType = FEdGraphPinType(UEdGraphSchema_K2::PC_Object, NAME_None, Animation->GetClass(), EPinContainerType::None, true, FEdGraphTerminalType());
+			AnimVariableDesc.FriendlyName =  Animation->GetDisplayName().ToString();
+			AnimVariableDesc.PropertyFlags = (CPF_Transient | CPF_BlueprintVisible | CPF_BlueprintReadOnly | CPF_RepSkip);
+			AnimVariableDesc.SetMetaData(TEXT("Category"), TEXT("Animations"));
+			ensure(!WidgetBP->GeneratedVariables.ContainsByPredicate([&AnimVariableDesc](const FBPVariableDescription& VariableDescription) { return VariableDescription.VarName == AnimVariableDesc.VarName; }));
+			ensure(!AnimVariableDesc.VarGuid.IsValid() || !WidgetBP->GeneratedVariables.ContainsByPredicate([&AnimVariableDesc](const FBPVariableDescription& VariableDescription) { return VariableDescription.VarGuid.IsValid() && VariableDescription.VarGuid == AnimVariableDesc.VarGuid; }));
+			AddGeneratedVariable(MoveTemp(AnimVariableDesc));
+		}
+	}
+
+	FWidgetBlueprintCompilerContext* Self = this;
+	UWidgetBlueprintExtension::ForEachExtension(WidgetBlueprint(), [Self](UWidgetBlueprintExtension* InExtension)
+		{
+			InExtension->PopulateGeneratedVariables(FPopulateGeneratedVariablesContext(*Self));
+		});
+}
+
 void FWidgetBlueprintCompilerContext::SpawnNewClass(const FString& NewClassName)
 {
 	NewWidgetBlueprintClass = FindObject<UWidgetBlueprintGeneratedClass>(Blueprint->GetOutermost(), *NewClassName);
@@ -1476,4 +1614,15 @@ void FWidgetBlueprintCompilerContext::AddExtension(UWidgetBlueprintGeneratedClas
 	Class->Extensions.Add(Extension);
 }
 
+void FWidgetBlueprintCompilerContext::AddGeneratedVariable(FBPVariableDescription&& VariableDescription) const
+{
+	UWidgetBlueprint* WidgetBP = WidgetBlueprint();
+	ensureAlwaysMsgf(!WidgetBP->GeneratedVariables.ContainsByPredicate([&VariableDescription](const FBPVariableDescription& TestVar) { return TestVar.VarName == VariableDescription.VarName; }),
+		TEXT("Widget Blueprint [%s] already contains generated variable with name [%s]"), *GetNameSafe(WidgetBP), *VariableDescription.VarName.ToString());
+	ensureAlwaysMsgf(!VariableDescription.VarGuid.IsValid() || !WidgetBP->GeneratedVariables.ContainsByPredicate([&VariableDescription](const FBPVariableDescription& TestVar) { return TestVar.VarGuid.IsValid() && TestVar.VarGuid == VariableDescription.VarGuid; }),
+		TEXT("Attempting to add generated variable [%s] to Widget Blueprint [%s] that has the same GUID as another variable"), *VariableDescription.VarName.ToString(), *GetNameSafe(WidgetBP));
+	
+	WidgetBP->GeneratedVariables.Emplace(MoveTemp(VariableDescription));
+}
+
 #undef LOCTEXT_NAMESPACE
diff --git a/Engine/Source/Editor/UMGEditor/Private/WidgetBlueprintEditor.cpp b/Engine/Source/Editor/UMGEditor/Private/WidgetBlueprintEditor.cpp
--- a/Engine/Source/Editor/UMGEditor/Private/WidgetBlueprintEditor.cpp
+++ b/Engine/Source/Editor/UMGEditor/Private/WidgetBlueprintEditor.cpp
@@ -93,6 +93,15 @@
 
 #define LOCTEXT_NAMESPACE "UMG"
 
+namespace UE::UMG::Editor
+{
+	bool bDisplayWidgetVariableGuids = false;
+	FAutoConsoleVariableRef CVarUMGEditorDisplayWidgetVariableGuids(
+		TEXT("UMG.Editor.DisplayWidgetVariableGuids"),
+		bDisplayWidgetVariableGuids,
+		TEXT("Toggles displaying guids for variables generated by this blueprint (Widgets and Animations)."));
+}
+
 FWidgetBlueprintEditor::FWidgetBlueprintEditor()
 	: PreviewScene(FPreviewScene::ConstructionValues().AllowAudioPlayback(true).ShouldSimulatePhysics(true))
 	, PreviewBlueprint(nullptr)
diff --git a/Engine/Source/Editor/UMGEditor/Private/WidgetBlueprintEditorUtils.cpp b/Engine/Source/Editor/UMGEditor/Private/WidgetBlueprintEditorUtils.cpp
--- a/Engine/Source/Editor/UMGEditor/Private/WidgetBlueprintEditorUtils.cpp
+++ b/Engine/Source/Editor/UMGEditor/Private/WidgetBlueprintEditorUtils.cpp
@@ -312,6 +312,8 @@ bool FWidgetBlueprintEditorUtils::RenameWidget(TSharedRef<FWidgetBlueprintEditor
 		Blueprint->Modify();
 		Widget->Modify();
 
+		Blueprint->OnVariableRenamed(OldObjectName, NewFName);
+
 		// Rename Preview before renaming the template widget so the preview widget can be found
 		UWidget* WidgetPreview = BlueprintEditor->GetReferenceFromTemplate(Widget).GetPreview();
 		if (WidgetPreview)
@@ -553,6 +555,7 @@ void FWidgetBlueprintEditorUtils::DeleteWidgets(TSharedRef<FWidgetBlueprintEdito
 		{
 			UWidget* WidgetTemplate = Item.GetTemplate();
 			WidgetTemplate->SetFlags(RF_Transactional);
+			const FName WidgetName = WidgetTemplate->GetFName();
 			// Find and update all binding references in the widget blueprint
 			for (int32 BindingIndex = Blueprint->Bindings.Num() - 1; BindingIndex >= 0; BindingIndex--)
 			{
@@ -594,11 +597,22 @@ void FWidgetBlueprintEditorUtils::DeleteWidgets(TSharedRef<FWidgetBlueprintEdito
 			// Rename the removed widget to the transient package so that it doesn't conflict with future widgets sharing the same name.
 			WidgetTemplate->Rename(nullptr, GetTransientPackage());
 
+			// Deletion can sometimes happen from replacing a widget with another one with the same name, so only delete the variable data if we no longer have a widget with the same name
+			const bool bHasWidgetWithSameName = Blueprint->GetAllSourceWidgets().ContainsByPredicate([WidgetName](const UWidget* Widget)
+			{
+				return WidgetName == Widget->GetFName();
+			});
+			if (!bHasWidgetWithSameName)
+			{
+				Blueprint->OnVariableRemoved(WidgetName);
+			}
+
 			// Rename all child widgets as well, to the transient package so that they don't conflict with future widgets sharing the same name.
 			TArray<UWidget*> ChildWidgets;
 			UWidgetTree::GetChildWidgets(WidgetTemplate, ChildWidgets);
 			for ( UWidget* Widget : ChildWidgets )
 			{
+				const FName ChildWidgetName = Widget->GetFName();
 				Widget->SetFlags(RF_Transactional);
 				Widget->Modify();
 				if (UsedVariables.Contains(Widget))
@@ -606,6 +620,16 @@ void FWidgetBlueprintEditorUtils::DeleteWidgets(TSharedRef<FWidgetBlueprintEdito
 					FBlueprintEditorUtils::RemoveVariableNodes(Blueprint, Widget->GetFName());
 				}
 				Widget->Rename(nullptr, GetTransientPackage());
+				
+				// Deletion can sometimes happen from replacing a widget with another one with the same name, so only delete the variable data if we no longer have a widget with the same name
+				const bool bHasChildWidgetWithSameName = Blueprint->GetAllSourceWidgets().ContainsByPredicate([ChildWidgetName](const UWidget* Widget)
+				{
+					return ChildWidgetName == Widget->GetFName();
+				});
+				if (!bHasChildWidgetWithSameName)
+				{
+					Blueprint->OnVariableRemoved(ChildWidgetName);
+				}
 			}
 		}
 
@@ -921,6 +945,8 @@ void FWidgetBlueprintEditorUtils::WrapWidgets(TSharedRef<FWidgetBlueprintEditor>
 				UPanelWidget* NewSlotContents = CastChecked<UPanelWidget>(Template->Create(BP->WidgetTree));
 				NewSlotContents->SetDesignerFlags(BlueprintEditor->GetCurrentDesignerFlags());
 
+				BP->OnVariableAdded(NewSlotContents->GetFName());
+
 				FWidgetBlueprintEditorUtils::ReplaceNamedSlotHostContent(Widget, NamedSlotHost, NewSlotContents);
 
 				NewSlotContents->AddChild(Widget);
@@ -935,6 +961,8 @@ void FWidgetBlueprintEditorUtils::WrapWidgets(TSharedRef<FWidgetBlueprintEditor>
 				NewWrapperWidget = CastChecked<UPanelWidget>(Template->Create(BP->WidgetTree));
 				NewWrapperWidget->SetDesignerFlags(BlueprintEditor->GetCurrentDesignerFlags());
 
+				BP->OnVariableAdded(NewWrapperWidget->GetFName());
+
 				CurrentParent->SetFlags(RF_Transactional);
 				CurrentParent->Modify();
 				CurrentParent->ReplaceChildAt(OutIndex, NewWrapperWidget);
@@ -951,6 +979,8 @@ void FWidgetBlueprintEditorUtils::WrapWidgets(TSharedRef<FWidgetBlueprintEditor>
 			UPanelWidget* NewRootContents = CastChecked<UPanelWidget>(Template->Create(BP->WidgetTree));
 			NewRootContents->SetDesignerFlags(BlueprintEditor->GetCurrentDesignerFlags());
 
+			BP->OnVariableAdded(NewRootContents->GetFName());
+
 			BP->WidgetTree->RootWidget = NewRootContents;
 			NewRootContents->AddChild(Widget);
 		}
@@ -1074,10 +1104,29 @@ void FWidgetBlueprintEditorUtils::ReplaceWidgetWithSelectedTemplate(TSharedRef<F
 {
 	// @Todo: Needs to deal with bound object in animation tracks
 
-	const FScopedTransaction Transaction(LOCTEXT("ReplaceWidgets", "Replace Widgets"));
-	bool bIsUserWidget = false;
+	FScopedTransaction Transaction(LOCTEXT("ReplaceWidgets", "Replace Widgets"));
 
 	UWidget* WidgetToReplace = Widget.GetTemplate();
+
+	if (!WidgetToReplace)
+	{
+		Transaction.Cancel();
+		return;
+	}
+	
+	// If they are both panel widgets then call the generic ReplaceWidgets function
+	if (WidgetToReplace->IsA<UPanelWidget>())
+	{
+		UClass* ReplacementWidgetClass = BlueprintEditor->GetSelectedTemplate().Get();
+		if (ReplacementWidgetClass && ReplacementWidgetClass->IsChildOf<UPanelWidget>())
+		{
+			TSet<FWidgetReference> ChildWidgetToReplace;
+			ChildWidgetToReplace.Add(Widget);
+			ReplaceWidgets(BlueprintEditor, BP, ChildWidgetToReplace, ReplacementWidgetClass);
+			return;
+		}
+	}
+	
 	UWidget* NewReplacementWidget;
 	if (BlueprintEditor->GetSelectedTemplate().IsValid())
 	{
@@ -1087,7 +1136,6 @@ void FWidgetBlueprintEditorUtils::ReplaceWidgetWithSelectedTemplate(TSharedRef<F
 	}
 	else if (BlueprintEditor->GetSelectedUserWidget().GetSoftObjectPath().IsValid())
 	{
-		bIsUserWidget = true;
 		FAssetData WidgetAssetData = BlueprintEditor->GetSelectedUserWidget();
 		TSharedPtr<FWidgetTemplateBlueprintClass> Template = MakeShareable(new FWidgetTemplateBlueprintClass(WidgetAssetData));
 		NewReplacementWidget = Template->Create(BP->WidgetTree);
@@ -1099,19 +1147,6 @@ void FWidgetBlueprintEditorUtils::ReplaceWidgetWithSelectedTemplate(TSharedRef<F
 
 	NewReplacementWidget->SetFlags(RF_Transactional);
 	NewReplacementWidget->Modify();
-
-	if (UPanelWidget* ExistingPanel = Cast<UPanelWidget>(WidgetToReplace))
-	{
-		// if they are both panel widgets then call the existing replace function
-		UPanelWidget* ReplacementPanelWidget = Cast<UPanelWidget>(NewReplacementWidget);
-		if (ReplacementPanelWidget)
-		{
-			TSet<FWidgetReference> ChildWidgetToReplace;
-			ChildWidgetToReplace.Add(Widget);
-			ReplaceWidgets(BlueprintEditor, BP, ChildWidgetToReplace, ReplacementPanelWidget->GetClass());
-			return;
-		}
-	}
 	WidgetToReplace->SetFlags(RF_Transactional);
 	WidgetToReplace->Modify();
 
@@ -1119,7 +1154,12 @@ void FWidgetBlueprintEditorUtils::ReplaceWidgetWithSelectedTemplate(TSharedRef<F
 	BP->WidgetTree->Modify();
 
 	// We need to check before replacing because the Widget might be deleted, reseting the DesiredFocus
-	bool bReplacingDesiredFocus = IsDesiredFocusWiget(BlueprintEditor, WidgetToReplace);
+	const bool bReplacingDesiredFocus = IsDesiredFocusWiget(BlueprintEditor, WidgetToReplace);
+	// Attempt to keep the same widget name if the types match up
+	const bool bCanKeepName = !WidgetToReplace->IsGeneratedName()
+		&& ((WidgetToReplace->IsA<UPanelWidget>() && NewReplacementWidget->IsA<UPanelWidget>()) || (WidgetToReplace->IsA(NewReplacementWidget->GetClass()) || NewReplacementWidget->IsA(WidgetToReplace->GetClass())));
+	
+	const FName OriginalWidgetName = WidgetToReplace->GetFName();
 
 	// Look if the Widget to replace is a NamedSlot.
 	if (TScriptInterface<INamedSlotInterface> NamedSlotHost = FindNamedSlotHostForContent(WidgetToReplace, BP->WidgetTree))
@@ -1131,18 +1171,6 @@ void FWidgetBlueprintEditorUtils::ReplaceWidgetWithSelectedTemplate(TSharedRef<F
 		CurrentParent->SetFlags(RF_Transactional);
 		CurrentParent->Modify();
 		CurrentParent->ReplaceChild(WidgetToReplace, NewReplacementWidget);
-
-		FString ReplaceName = WidgetToReplace->GetName();
-		bool bIsGeneratedName = WidgetToReplace->IsGeneratedName();
-		// Rename the removed widget to the transient package so that it doesn't conflict with future widgets sharing the same name.
-		WidgetToReplace->Rename(nullptr, GetTransientPackage());
-
-		// Rename the new Widget to maintain the current name if it's not a generic name
-		if (!bIsGeneratedName)
-		{
-			ReplaceName = FindNextValidName(BP->WidgetTree, ReplaceName);
-			NewReplacementWidget->Rename(*ReplaceName, BP->WidgetTree);
-		}
 	}
 	else if (WidgetToReplace == BP->WidgetTree->RootWidget)
 	{
@@ -1153,16 +1181,34 @@ void FWidgetBlueprintEditorUtils::ReplaceWidgetWithSelectedTemplate(TSharedRef<F
 		return;
 	}
 
-	// Delete the widget that has been replaced
-	TSet<FWidgetReference> WidgetsToDelete;
-	WidgetsToDelete.Add(Widget);
-	DeleteWidgets(BlueprintEditor, BP, WidgetsToDelete);
+	FString NewWidgetName = WidgetToReplace->GetName();
+	// Rename the removed widget to the transient package so that it doesn't conflict with the new widget if we try to keep the same name.
+	WidgetToReplace->Rename(nullptr, GetTransientPackage());
+
+	if (bCanKeepName)
+	{
+		NewWidgetName = FindNextValidName(BP->WidgetTree, NewWidgetName);
+		NewReplacementWidget->Rename(*NewWidgetName, BP->WidgetTree);
+
+		// Since we're attempting to preserve the widget name, we'll also preserve its guid
+		BP->OnVariableRenamed(OriginalWidgetName, NewReplacementWidget->GetFName());
+	}
+	else if (NewReplacementWidget->GetFName() != OriginalWidgetName)
+	{
+		BP->OnVariableRemoved(OriginalWidgetName);
+		BP->OnVariableAdded(NewReplacementWidget->GetFName());
+	}
 
 	if (bReplacingDesiredFocus)
 	{
 		FWidgetBlueprintEditorUtils::SetDesiredFocus(BlueprintEditor, NewReplacementWidget->GetFName());
 	}
 
+	// Delete the widget that has been replaced
+	TSet<FWidgetReference> WidgetsToDelete;
+	WidgetsToDelete.Add(Widget);
+	DeleteWidgets(BlueprintEditor, BP, WidgetsToDelete);
+
 	FBlueprintEditorUtils::MarkBlueprintAsStructurallyModified(BP);
 }
 
@@ -1368,16 +1414,26 @@ void FWidgetBlueprintEditorUtils::ReplaceWidgets(TSharedRef<FWidgetBlueprintEdit
 		WidgetToReplace->SetFlags(RF_Transactional);
 		WidgetToReplace->Modify();
 
+		const FName OriginalWidgetName = WidgetToReplace->GetFName();
+
 		// Look if the Widget to replace is a NamedSlot.
 		if (TScriptInterface<INamedSlotInterface> NamedSlotHost = FindNamedSlotHostForContent(WidgetToReplace, BP->WidgetTree))
 		{
-			ReplaceNamedSlotHostContent(WidgetToReplace, NamedSlotHost, NewReplacementWidget);
+			const bool bDidReplace = ReplaceNamedSlotHostContent(WidgetToReplace, NamedSlotHost, NewReplacementWidget);
+			if (!bDidReplace)
+			{
+				continue;
+			}
 		}
 		else if (UPanelWidget* CurrentParent = WidgetToReplace->GetParent())
 		{
 			CurrentParent->SetFlags(RF_Transactional);
 			CurrentParent->Modify();
-			CurrentParent->ReplaceChild(WidgetToReplace, NewReplacementWidget);
+			const bool bDidReplace = CurrentParent->ReplaceChild(WidgetToReplace, NewReplacementWidget);
+			if (!bDidReplace)
+			{
+				continue;
+			}
 		}
 		else if (WidgetToReplace == BP->WidgetTree->RootWidget)
 		{
@@ -1408,17 +1464,28 @@ void FWidgetBlueprintEditorUtils::ReplaceWidgets(TSharedRef<FWidgetBlueprintEdit
 
 		FString ReplaceName = WidgetToReplace->GetName();
 		const bool bIsGeneratedName = WidgetToReplace->IsGeneratedName();
-
-		// Delete the widget that has been replaced
-		const bool bForceDelete = true;
-		DeleteWidgets(BlueprintEditor, BP, {Item}, bForceDelete);
+		
+		// Rename the removed widget to the transient package so that it doesn't conflict with the new widget if we try to keep the same name.
+		WidgetToReplace->Rename(nullptr, GetTransientPackage());
 
 		// Rename the new Widget to maintain the current name if it's not a generic name
 		if (!bIsGeneratedName)
 		{
 			ReplaceName = FindNextValidName(BP->WidgetTree, ReplaceName);
 			NewReplacementWidget->Rename(*ReplaceName, BP->WidgetTree);
+
+			// Since we're attempting to preserve the widget name, we'll also preserve its guid
+			BP->OnVariableRenamed(OriginalWidgetName, NewReplacementWidget->GetFName());
 		}
+		else if (NewReplacementWidget->GetFName() != OriginalWidgetName)
+		{
+			BP->OnVariableRemoved(OriginalWidgetName);
+			BP->OnVariableAdded(NewReplacementWidget->GetFName());
+		}
+
+		// Delete the widget that has been replaced
+		constexpr bool bForceDelete = true;
+		DeleteWidgets(BlueprintEditor, BP, {Item}, bForceDelete);
 
 		if (bReplacingDesiredFocus)
 		{
@@ -1867,6 +1934,7 @@ TArray<UWidget*> FWidgetBlueprintEditorUtils::PasteWidgetsInternal(TSharedRef<FW
 
 		for (UWidget* NewWidget : PastedWidgets)
 		{
+			BP->OnVariableAdded(NewWidget->GetFName());
 			// Widgets with a null parent mean that they were the root most widget of their selection set when
 			// they were copied and thus we need to paste only the root most widgets.  All their children will be added
 			// automatically.
@@ -1969,7 +2037,8 @@ TArray<UWidget*> FWidgetBlueprintEditorUtils::PasteWidgetsInternal(TSharedRef<FW
 
 			if (ChildWidgetToDelete)
 			{
-				DeleteWidgets(BlueprintEditor, BP, { BlueprintEditor->GetReferenceFromTemplate(ChildWidgetToDelete) });
+				constexpr bool bForceDelete = true;
+				DeleteWidgets(BlueprintEditor, BP, { BlueprintEditor->GetReferenceFromTemplate(ChildWidgetToDelete) }, bForceDelete);
 			}
 		}
 
@@ -1982,6 +2051,7 @@ TArray<UWidget*> FWidgetBlueprintEditorUtils::PasteWidgetsInternal(TSharedRef<FW
 			ParentWidget = BP->WidgetTree->ConstructWidget<UCanvasPanel>(UCanvasPanel::StaticClass());
 			BP->WidgetTree->Modify();
 			BP->WidgetTree->RootWidget = ParentWidget;
+			BP->OnVariableAdded(ParentWidget->GetFName());
 		}
 
 		if (ParentWidget)
diff --git a/Engine/Source/Editor/UMGEditor/Public/WidgetBlueprint.h b/Engine/Source/Editor/UMGEditor/Public/WidgetBlueprint.h
--- a/Engine/Source/Editor/UMGEditor/Public/WidgetBlueprint.h
+++ b/Engine/Source/Editor/UMGEditor/Public/WidgetBlueprint.h
@@ -222,12 +222,21 @@ class UMGEDITOR_API UWidgetBlueprint : public UBaseWidgetBlueprint
 public:
 
 #if WITH_EDITORONLY_DATA
-	
 	UPROPERTY()
 	TArray< FDelegateEditorBinding > Bindings;
 
 	UPROPERTY()
 	TArray<TObjectPtr<UWidgetAnimation>> Animations;
+	
+	/**
+	 * Saved Guids for variables generated by this blueprint (Widgets and Animations).
+	 *
+	 * Used to fixup external references to renamed widgets/animations in this blueprint
+	 *
+	 * See: OnVariableAdded, OnVariableRenamed, OnVariableRemoved
+	 */
+	UPROPERTY()
+	TMap<FName, FGuid> WidgetVariableNameToGuidMap;
 
 	/**
 	 * Don't directly modify this property to change the palette category.  The actual value is stored 
@@ -244,6 +253,10 @@ public:
 	 */
 	UPROPERTY(EditAnywhere, Category="Widget")
 	bool bCanCallInitializedWithoutPlayerContext;
+
+	void OnVariableAdded(const FName& VariableName);
+	void OnVariableRenamed(const FName& OldName, const FName& NewName);
+	void OnVariableRemoved(const FName& VariableName);
 #endif
 
 public:
diff --git a/Engine/Source/Editor/UMGEditor/Public/WidgetBlueprintCompiler.h b/Engine/Source/Editor/UMGEditor/Public/WidgetBlueprintCompiler.h
--- a/Engine/Source/Editor/UMGEditor/Public/WidgetBlueprintCompiler.h
+++ b/Engine/Source/Editor/UMGEditor/Public/WidgetBlueprintCompiler.h
@@ -59,8 +59,6 @@ public:
 	virtual ~FWidgetBlueprintCompilerContext();
 
 protected:
-	void ValidateWidgetNames();
-
 	/**
 	 * Checks if the animations' bindings are valid. Shows a compiler warning if any of the animations have a null track, urging the user
 	 * to delete said track to avoid performance hitches when playing said null tracks in large user widgets.
@@ -79,6 +77,7 @@ protected:
 	virtual void CleanAndSanitizeClass(UBlueprintGeneratedClass* ClassToClean, UObject*& InOutOldCDO) override;
 	virtual void SaveSubObjectsFromCleanAndSanitizeClass(FSubobjectCollection& SubObjectsToSave, UBlueprintGeneratedClass* ClassToClean) override;
 	virtual void EnsureProperGeneratedClass(UClass*& TargetClass) override;
+	virtual void PopulateBlueprintGeneratedVariables() override;
 	virtual void CreateClassVariablesFromBlueprint() override;
 	virtual void CopyTermDefaultsToDefaultObject(UObject* DefaultObject);
 	virtual void FinishCompilingClass(UClass* Class) override;
@@ -87,6 +86,7 @@ protected:
 	//~ End FKismetCompilerContext
 
 	void SanitizeBindings(UBlueprintGeneratedClass* Class);
+	void ValidateAndFixUpVariableGuids();
 
 	void VerifyEventReplysAreNotEmpty(FKismetFunctionContext& Context);
 	void VerifyFieldNotifyFunction(FKismetFunctionContext& Context);
@@ -95,6 +95,19 @@ public:
 	UWidgetBlueprint* WidgetBlueprint() const { return Cast<UWidgetBlueprint>(Blueprint); }
 
 	void AddExtension(UWidgetBlueprintGeneratedClass* Class, UWidgetBlueprintGeneratedClassExtension* Extension);
+	void AddGeneratedVariable(FBPVariableDescription&& VariableDescription) const;
+	
+	struct UMGEDITOR_API FPopulateGeneratedVariablesContext
+	{
+	public:
+		void AddGeneratedVariable(FBPVariableDescription&& VariableDescription) const;
+		UWidgetBlueprint* GetWidgetBlueprint() const;
+
+	private:
+		friend FWidgetBlueprintCompilerContext;
+		FPopulateGeneratedVariablesContext(FWidgetBlueprintCompilerContext& InContext);
+		FWidgetBlueprintCompilerContext& Context;
+	};
 
 	struct UMGEDITOR_API FCreateVariableContext
 	{
diff --git a/Engine/Source/Editor/UMGEditor/Public/WidgetBlueprintExtension.h b/Engine/Source/Editor/UMGEditor/Public/WidgetBlueprintExtension.h
--- a/Engine/Source/Editor/UMGEditor/Public/WidgetBlueprintExtension.h
+++ b/Engine/Source/Editor/UMGEditor/Public/WidgetBlueprintExtension.h
@@ -75,6 +75,7 @@ protected:
 	virtual void HandleCreateFunctionList(const FWidgetBlueprintCompilerContext::FCreateFunctionContext& InCreationContext) {}
 	virtual void HandleCleanAndSanitizeClass(UWidgetBlueprintGeneratedClass* ClassToClean, UObject* OldCDO) {}
 	virtual TArray<UObject*> HandleSaveSubObjectsFromCleanAndSanitizeClass(UWidgetBlueprintGeneratedClass* ClassToClean) { return TArray<UObject*>();  }
+	virtual void HandlePopulateGeneratedVariables(const FWidgetBlueprintCompilerContext::FPopulateGeneratedVariablesContext& Context) {}
 	virtual void HandleCreateClassVariablesFromBlueprint(const FWidgetBlueprintCompilerContext::FCreateVariableContext& Context) {}
 	virtual void HandleCopyTermDefaultsToDefaultObject(UObject* DefaultObject) {}
 	virtual void HandleFinishCompilingClass(UWidgetBlueprintGeneratedClass* Class) {}
@@ -110,6 +111,11 @@ private:
 		return HandleSaveSubObjectsFromCleanAndSanitizeClass(ClassToClean);
 	}
 
+	void PopulateGeneratedVariables(const FWidgetBlueprintCompilerContext::FPopulateGeneratedVariablesContext& Context)
+	{
+		HandlePopulateGeneratedVariables(Context);
+	}
+
 	void CreateClassVariablesFromBlueprint(const FWidgetBlueprintCompilerContext::FCreateVariableContext& Context)
 	{
 		HandleCreateClassVariablesFromBlueprint(Context);
