Exposed WidgetBlueprintGeneratedClass to the context menu.
Added support for injecting the "Run Editor Utility Widget" option for WidgetBlueprintGeneratedClass that are Editor Utility Widgets
Added the Blutility module to UMGEditor to support WidgetBlueprintGeneratedClass editor utility widgets being executed.

diff --git a/Engine/Plugins/Editor/EngineAssetDefinitions/Source/Private/AssetDefinition_EditorUtilityWidgetBlueprint.cpp b/Engine/Plugins/Editor/EngineAssetDefinitions/Source/Private/AssetDefinition_EditorUtilityWidgetBlueprint.cpp
--- a/Engine/Plugins/Editor/EngineAssetDefinitions/Source/Private/AssetDefinition_EditorUtilityWidgetBlueprint.cpp
+++ b/Engine/Plugins/Editor/EngineAssetDefinitions/Source/Private/AssetDefinition_EditorUtilityWidgetBlueprint.cpp
@@ -84,19 +84,20 @@ namespace MenuExtension_EditorUtilityWidgetBlueprint
 	{
 		if (const UContentBrowserAssetContextMenuContext* Context = UContentBrowserAssetContextMenuContext::FindContextWithAssets(InContext))
 		{
-			UEditorUtilitySubsystem* EditorUtilitySubsystem = GEditor->GetEditorSubsystem<UEditorUtilitySubsystem>();
-			
-			TSet<FName> LoadTags;
-			TArray<UWidgetBlueprint*> WidgetBlueprints = Context->LoadSelectedObjects<UWidgetBlueprint>(LoadTags);
-			
-			for (UWidgetBlueprint* WidgetBlueprint : WidgetBlueprints)
+			if (UEditorUtilitySubsystem* EditorUtilitySubsystem = GEditor->GetEditorSubsystem<UEditorUtilitySubsystem>())
 			{
-				ensureMsgf(WidgetBlueprint->GeneratedClass, TEXT("Null generated class for WidgetBlueprint [%s]"), *WidgetBlueprint->GetName());
-				if (WidgetBlueprint->GeneratedClass && WidgetBlueprint->GeneratedClass->IsChildOf(UEditorUtilityWidget::StaticClass()))
+				TSet<FName> LoadTags;
+				TArray<UWidgetBlueprint*> WidgetBlueprints = Context->LoadSelectedObjects<UWidgetBlueprint>(LoadTags);
+
+				for (UWidgetBlueprint* WidgetBlueprint : WidgetBlueprints)
 				{
-					if (UEditorUtilityWidgetBlueprint* EditorWidget = Cast<UEditorUtilityWidgetBlueprint>(WidgetBlueprint))
+					ensureMsgf(WidgetBlueprint->GeneratedClass, TEXT("Null generated class for WidgetBlueprint [%s]"), *WidgetBlueprint->GetName());
+					if (WidgetBlueprint->GeneratedClass && WidgetBlueprint->GeneratedClass->IsChildOf(UEditorUtilityWidget::StaticClass()))
 					{
-						EditorUtilitySubsystem->SpawnAndRegisterTab(EditorWidget);
+						if (UEditorUtilityWidgetBlueprint* EditorWidget = Cast<UEditorUtilityWidgetBlueprint>(WidgetBlueprint))
+						{
+							EditorUtilitySubsystem->SpawnAndRegisterTab(EditorWidget);
+						}
 					}
 				}
 			}
diff --git a/Engine/Source/Editor/UMGEditor/Private/AssetDefinition_WidgetBlueprintGeneratedClass.cpp b/Engine/Source/Editor/UMGEditor/Private/AssetDefinition_WidgetBlueprintGeneratedClass.cpp
--- a/Engine/Source/Editor/UMGEditor/Private/AssetDefinition_WidgetBlueprintGeneratedClass.cpp
+++ b/Engine/Source/Editor/UMGEditor/Private/AssetDefinition_WidgetBlueprintGeneratedClass.cpp
@@ -2,6 +2,16 @@
 
 #include "AssetDefinition_WidgetBlueprintGeneratedClass.h"
 
+#include "Algo/RemoveIf.h"
+#include "ContentBrowserMenuContexts.h"
+#include "EditorUtilityLibrary.h"
+#include "EditorUtilitySubsystem.h"
+#include "EditorUtilityWidget.h"
+#include "EditorUtilityWidgetBlueprint.h"
+#include "ToolMenus.h"
+
+#define LOCTEXT_NAMESPACE "AssetTypeActions"
+
 UAssetDefinition_WidgetBlueprintGeneratedClass::UAssetDefinition_WidgetBlueprintGeneratedClass() = default;
 
 UAssetDefinition_WidgetBlueprintGeneratedClass::~UAssetDefinition_WidgetBlueprintGeneratedClass() = default;
@@ -26,3 +36,78 @@ TConstArrayView<FAssetCategoryPath> UAssetDefinition_WidgetBlueprintGeneratedCla
 	static const auto Categories = { EAssetCategoryPaths::UI };
 	return Categories;
 }
+
+namespace MenuExtension_WidgetBlueprintGeneratedClass
+{
+	UClass* GetParentClass(const FAssetData& AssetData)
+	{
+		UClass* ParentClass = nullptr;
+		FString ParentClassName;
+		if (!AssetData.GetTagValue(FBlueprintTags::NativeParentClassPath, ParentClassName))
+		{
+			AssetData.GetTagValue(FBlueprintTags::ParentClassPath, ParentClassName);
+		}
+
+		if (!ParentClassName.IsEmpty())
+		{
+			ParentClass = UClass::TryFindTypeSlow<UClass>(FPackageName::ExportTextPathToObjectPath(ParentClassName));
+		}
+
+		return ParentClass;
+	}
+
+	void ExecuteEditorUtilityEdit(const FToolMenuContext& InContext, TArray<FAssetData> AssetData)
+	{
+		if (const UContentBrowserAssetContextMenuContext* Context = UContentBrowserAssetContextMenuContext::FindContextWithAssets(InContext))
+		{
+			if(UEditorUtilitySubsystem* EditorUtilitySubsystem = GEditor->GetEditorSubsystem<UEditorUtilitySubsystem>())
+			{
+				TArray<UWidgetBlueprintGeneratedClass*> WidgetBlueprints = Context->LoadSelectedObjects<UWidgetBlueprintGeneratedClass>();
+
+				for (UWidgetBlueprintGeneratedClass* WidgetBlueprint : WidgetBlueprints)
+				{
+					EditorUtilitySubsystem->SpawnAndRegisterTabGeneratedClass(WidgetBlueprint);
+				}
+			}
+		}
+	}
+
+	FDelayedAutoRegisterHelper DelayedAutoRegister(EDelayedRegisterRunPhase::EndOfEngineInit, []
+	{
+		UToolMenus::RegisterStartupCallback(FSimpleMulticastDelegate::FDelegate::CreateLambda([]()
+		{
+			FToolMenuOwnerScoped OwnerScoped(UE_MODULE_NAME);
+			UToolMenu* Menu = UE::ContentBrowser::ExtendToolMenu_AssetContextMenu(UWidgetBlueprintGeneratedClass::StaticClass());
+
+			FToolMenuSection& Section = Menu->FindOrAddSection("GetAssetActions");
+			Section.AddDynamicEntry(NAME_None, FNewToolMenuSectionDelegate::CreateLambda([](FToolMenuSection& InSection)
+			{
+				if (const UContentBrowserAssetContextMenuContext* Context = UContentBrowserAssetContextMenuContext::FindContextWithAssets(InSection))
+				{
+					TArray<FAssetData> SelectedAssets(Context->SelectedAssets);
+
+					const int32 EndOfRange = Algo::StableRemoveIf(SelectedAssets, [](const FAssetData& AssetData)
+					{
+						// Make sure this BP Generated Class is an actual EUW, and not just a UserWidget.
+						return !(GetParentClass(AssetData) == UEditorUtilityWidget::StaticClass());
+					});
+
+					SelectedAssets.SetNum(EndOfRange);
+
+					if (!SelectedAssets.IsEmpty())
+					{
+						const TAttribute<FText> Label = LOCTEXT("EditorUtilityWidget_Edit", "Run Editor Utility Widget");
+						const TAttribute<FText> ToolTip = LOCTEXT("EditorUtilityWidget_EditTooltip", "Opens the tab built by this Editor Utility Widget Blueprint.");
+						const FSlateIcon Icon = FSlateIcon();
+
+						FToolUIAction UIAction;
+						UIAction.ExecuteAction = FToolMenuExecuteAction::CreateStatic(&ExecuteEditorUtilityEdit, MoveTemp(SelectedAssets));
+						InSection.AddMenuEntry("EditorUtility_Run", Label, ToolTip, Icon, UIAction);
+					}
+				}
+			}));
+		}));
+	});
+}
+
+#undef LOCTEXT_NAMESPACE
diff --git a/Engine/Source/Editor/UMGEditor/UMGEditor.Build.cs b/Engine/Source/Editor/UMGEditor/UMGEditor.Build.cs
--- a/Engine/Source/Editor/UMGEditor/UMGEditor.Build.cs
+++ b/Engine/Source/Editor/UMGEditor/UMGEditor.Build.cs
@@ -37,6 +37,7 @@ public class UMGEditor : ModuleRules
 				"EditorWidgets",
 				"KismetCompiler",
 				"BlueprintGraph",
+				"Blutility",
 				"GraphEditor",
 				"Kismet",  // for FWorkflowCentricApplication
 				"Projects",
