UMGEditor: Add guids for widgets created by dragging them onto of the designer preview

- This was a missed path to adding widgets in the original change to track guids for UMG variables

diff --git a/Engine/Source/Editor/UMGEditor/Private/Designer/SDesignerView.cpp b/Engine/Source/Editor/UMGEditor/Private/Designer/SDesignerView.cpp
--- a/Engine/Source/Editor/UMGEditor/Private/Designer/SDesignerView.cpp
+++ b/Engine/Source/Editor/UMGEditor/Private/Designer/SDesignerView.cpp
@@ -2798,6 +2798,12 @@ void SDesignerView::ProcessDropAndAddWidget(const FGeometry& MyGeometry, const F
 				DropPreview.DragOperation = DragOperation;
 				DropPreviews.Add(DropPreview);
 			}
+			
+			if (!bIsPreview && BP->WidgetTree->RootWidget)
+			{
+				// We've added a new widget to the BP, notify the BP so it can create a GUID to track it
+				BP->OnVariableAdded(BP->WidgetTree->RootWidget->GetFName());
+			}
 		}
 		// If there's already a root widget we need to try and place our widget into a parent widget that we've picked against
 		else if ( Target && Target->IsA(UPanelWidget::StaticClass()) )
@@ -2830,6 +2836,12 @@ void SDesignerView::ProcessDropAndAddWidget(const FGeometry& MyGeometry, const F
 					DropPreview.Parent = Parent;
 					DropPreview.DragOperation = DragOperation;
 					DropPreviews.Add(DropPreview);
+					
+					if (!bIsPreview && Widget)
+					{
+						// We've added a new widget to the BP, notify the BP so it can create a GUID to track it
+						BP->OnVariableAdded(Widget->GetFName());
+					}
 				}
 				else
 				{
