UMGEditor: Clean up and expose Delete/Replace widget API to support operating on widget blueprints without opening the widget blueprint editor

- Allows external code to programatically delete/replace widgets in a blueprint without needing to manually perform the same operations
- Make an effort to persist data set on widgets to their replacement widgets (based on property name)
- Also fix renaming widgets not fixing up BP graph references due to needing to compile first

diff --git a/Engine/Plugins/Runtime/ModelViewViewModel/Source/ModelViewViewModelEditor/Private/Customizations/MVVMPanelWidgetExtensionCustomizationExtender.cpp b/Engine/Plugins/Runtime/ModelViewViewModel/Source/ModelViewViewModelEditor/Private/Customizations/MVVMPanelWidgetExtensionCustomizationExtender.cpp
--- a/Engine/Plugins/Runtime/ModelViewViewModel/Source/ModelViewViewModelEditor/Private/Customizations/MVVMPanelWidgetExtensionCustomizationExtender.cpp
+++ b/Engine/Plugins/Runtime/ModelViewViewModel/Source/ModelViewViewModelEditor/Private/Customizations/MVVMPanelWidgetExtensionCustomizationExtender.cpp
@@ -247,14 +247,14 @@ FReply FMVVMPanelWidgetExtensionCustomizationExtender::ModifyExtension()
 			{
 				if (UWidgetBlueprint* Blueprint = BPEditor->GetWidgetBlueprintObj())
 				{
-					TSet<FWidgetReference> ChildWidgets;
+					TSet<UWidget*> ChildWidgets;
 
 					for (UWidget* Child : WidgetPtr->GetAllChildren())
 					{
-						ChildWidgets.Add(BPEditor->GetReferenceFromPreview(Child));
+						ChildWidgets.Add(BPEditor->GetReferenceFromPreview(Child).GetTemplate());
 					}
 
-					FWidgetBlueprintEditorUtils::DeleteWidgets(BPEditor.ToSharedRef(), Blueprint, ChildWidgets);
+					FWidgetBlueprintEditorUtils::DeleteWidgets(Blueprint, ChildWidgets, FWidgetBlueprintEditorUtils::EDeleteWidgetWarningType::WarnAndAskUser);
 				}
 			}
 		}
diff --git a/Engine/Source/Editor/UMGEditor/Private/Hierarchy/SHierarchyView.cpp b/Engine/Source/Editor/UMGEditor/Private/Hierarchy/SHierarchyView.cpp
--- a/Engine/Source/Editor/UMGEditor/Private/Hierarchy/SHierarchyView.cpp
+++ b/Engine/Source/Editor/UMGEditor/Private/Hierarchy/SHierarchyView.cpp
@@ -302,7 +302,7 @@ FReply SHierarchyView::HandleDeleteSelected()
 {
 	TSet<FWidgetReference> SelectedWidgets = BlueprintEditor.Pin()->GetSelectedWidgets();
 	
-	FWidgetBlueprintEditorUtils::DeleteWidgets(BlueprintEditor.Pin().ToSharedRef(), GetBlueprint(), SelectedWidgets);
+	FWidgetBlueprintEditorUtils::DeleteWidgets(GetBlueprint(), FWidgetBlueprintEditorUtils::ResolveWidgetTemplates(SelectedWidgets), FWidgetBlueprintEditorUtils::EDeleteWidgetWarningType::WarnAndAskUser);
 
 	return FReply::Handled();
 }
diff --git a/Engine/Source/Editor/UMGEditor/Private/WidgetBlueprintEditor.cpp b/Engine/Source/Editor/UMGEditor/Private/WidgetBlueprintEditor.cpp
--- a/Engine/Source/Editor/UMGEditor/Private/WidgetBlueprintEditor.cpp
+++ b/Engine/Source/Editor/UMGEditor/Private/WidgetBlueprintEditor.cpp
@@ -823,7 +823,7 @@ bool FWidgetBlueprintEditor::CanDeleteSelectedWidgets()
 void FWidgetBlueprintEditor::DeleteSelectedWidgets()
 {
 	TSet<FWidgetReference> Widgets = GetSelectedWidgets();
-	FWidgetBlueprintEditorUtils::DeleteWidgets(SharedThis(this), GetWidgetBlueprintObj(), Widgets);
+	FWidgetBlueprintEditorUtils::DeleteWidgets(GetWidgetBlueprintObj(), FWidgetBlueprintEditorUtils::ResolveWidgetTemplates(Widgets), FWidgetBlueprintEditorUtils::EDeleteWidgetWarningType::WarnAndAskUser);
 
 	// Clear the selection now that the widget has been deleted.
 	TSet<FWidgetReference> Empty;
diff --git a/Engine/Source/Editor/UMGEditor/Private/WidgetBlueprintEditorUtils.cpp b/Engine/Source/Editor/UMGEditor/Private/WidgetBlueprintEditorUtils.cpp
--- a/Engine/Source/Editor/UMGEditor/Private/WidgetBlueprintEditorUtils.cpp
+++ b/Engine/Source/Editor/UMGEditor/Private/WidgetBlueprintEditorUtils.cpp
@@ -231,16 +231,29 @@ bool FWidgetBlueprintEditorUtils::VerifyWidgetRename(TSharedRef<class FWidgetBlu
 
 void FWidgetBlueprintEditorUtils::SetDesiredFocus(TSharedRef<FWidgetBlueprintEditor> BlueprintEditor, const FName DesiredFocusWidgetName)
 {
-	if (UUserWidget* PreviewWidget = BlueprintEditor->GetPreview())
+	SetDesiredFocus(BlueprintEditor->GetWidgetBlueprintObj(), DesiredFocusWidgetName);
+}
+
+void FWidgetBlueprintEditorUtils::SetDesiredFocus(UWidgetBlueprint* Blueprint, const FName& DesiredFocusWidgetName)
+{
+	if (Blueprint)
 	{
-		if (UWidgetBlueprintGeneratedClass* BGClass = PreviewWidget->GetWidgetTreeOwningClass())
+		if (Blueprint->GeneratedClass)
 		{
-			if (UUserWidget* WidgetCDO = BGClass->GetDefaultObject<UUserWidget>())
+			if (UUserWidget* WidgetCDO = Blueprint->GeneratedClass->GetDefaultObject<UUserWidget>())
 			{
 				WidgetCDO->SetFlags(RF_Transactional);
 				WidgetCDO->Modify();
 				WidgetCDO->SetDesiredFocusWidget(DesiredFocusWidgetName);
+			}
+		}
 
+		constexpr bool bFocusIfOpen = false;
+		FWidgetBlueprintEditor* BlueprintEditor = static_cast<FWidgetBlueprintEditor*>(GEditor->GetEditorSubsystem<UAssetEditorSubsystem>()->FindEditorForAsset(Blueprint, bFocusIfOpen));
+		if (BlueprintEditor)
+		{
+			if (UUserWidget* PreviewWidget = BlueprintEditor->GetPreview())
+			{
 				// We need to change the PreviewWidget to make sure the DetailPanel show the right value.
 				PreviewWidget->SetFlags(RF_Transactional);
 				PreviewWidget->Modify();
@@ -252,26 +265,36 @@ void FWidgetBlueprintEditorUtils::SetDesiredFocus(TSharedRef<FWidgetBlueprintEdi
 
 
 void FWidgetBlueprintEditorUtils::ReplaceDesiredFocus(TSharedRef<FWidgetBlueprintEditor> BlueprintEditor, const FName& OldName, const FName& NewName)
-{		
-	if (UUserWidget* PreviewWidget = BlueprintEditor->GetPreview())
+{
+	ReplaceDesiredFocus(BlueprintEditor->GetWidgetBlueprintObj(), OldName, NewName);
+}
+
+void FWidgetBlueprintEditorUtils::ReplaceDesiredFocus(UWidgetBlueprint* Blueprint, const FName& OldName, const FName& NewName)
+{
+	if (Blueprint && Blueprint->GeneratedClass)
 	{
-		if (UWidgetBlueprintGeneratedClass* BGClass = PreviewWidget->GetWidgetTreeOwningClass())
+		if (UUserWidget* WidgetCDO = Blueprint->GeneratedClass->GetDefaultObject<UUserWidget>())
 		{
-			if (UUserWidget* WidgetCDO = BGClass->GetDefaultObject<UUserWidget>())
+			// Verify if the Name changed is the Desired focus Widget name.
+			if (WidgetCDO->GetDesiredFocusWidgetName() == OldName)
 			{
-				// Verify if the Name changed is the Desired focus Widget name.
-				if (WidgetCDO->GetDesiredFocusWidgetName() == OldName)
+				WidgetCDO->SetFlags(RF_Transactional);
+				WidgetCDO->Modify();
+				WidgetCDO->SetDesiredFocusWidget(NewName);
+				
+				constexpr bool bFocusIfOpen = false;
+				FWidgetBlueprintEditor* BlueprintEditor = static_cast<FWidgetBlueprintEditor*>(GEditor->GetEditorSubsystem<UAssetEditorSubsystem>()->FindEditorForAsset(Blueprint, bFocusIfOpen));
+				if (BlueprintEditor)
 				{
-					WidgetCDO->SetFlags(RF_Transactional);
-					WidgetCDO->Modify();
-					WidgetCDO->SetDesiredFocusWidget(NewName);
-
-					ensure(PreviewWidget->GetDesiredFocusWidgetName() == OldName);
+					if (UUserWidget* PreviewWidget = BlueprintEditor->GetPreview())
+					{
+						ensure(PreviewWidget->GetDesiredFocusWidgetName() == OldName);
 
-					// We need to change the PreviewWidget to make sure the DetailPanel show the right value.
-					PreviewWidget->SetFlags(RF_Transactional);
-					PreviewWidget->Modify();
-					PreviewWidget->SetDesiredFocusWidget(NewName);
+						// We need to change the PreviewWidget to make sure the DetailPanel show the right value.
+						PreviewWidget->SetFlags(RF_Transactional);
+						PreviewWidget->Modify();
+						PreviewWidget->SetDesiredFocusWidget(NewName);
+					}
 				}
 			}
 		}
@@ -374,10 +397,6 @@ bool FWidgetBlueprintEditorUtils::RenameWidget(TSharedRef<FWidgetBlueprintEditor
 		}
 #endif
 
-		// Update Variable References and
-		// Update Event References to member variables
-		FBlueprintEditorUtils::ReplaceVariableReferences(Blueprint, OldObjectName, NewFName);
-
 		// Replace the Desired focus Widget name if it match the renamed widget
 		ReplaceDesiredFocus(BlueprintEditor, OldObjectName, NewFName);
 		
@@ -432,6 +451,11 @@ bool FWidgetBlueprintEditorUtils::RenameWidget(TSharedRef<FWidgetBlueprintEditor
 
 		// Refresh references and flush editors
 		FBlueprintEditorUtils::MarkBlueprintAsStructurallyModified(Blueprint);
+
+		// Update Variable References and
+		// Update Event References to member variables
+		FBlueprintEditorUtils::ReplaceVariableReferences(Blueprint, OldObjectName, NewFName);
+
 		bRenamed = true;
 	}
 
@@ -530,6 +554,11 @@ bool FWidgetBlueprintEditorUtils::CanOpenSelectedWidgetsForEdit( TSet<FWidgetRef
 }
 
 void FWidgetBlueprintEditorUtils::DeleteWidgets(TSharedRef<FWidgetBlueprintEditor> BlueprintEditor, UWidgetBlueprint* Blueprint, TSet<FWidgetReference> Widgets, bool bSilentDelete /*=false*/)
+{
+	DeleteWidgets(Blueprint, ResolveWidgetTemplates(Widgets), bSilentDelete ? EDeleteWidgetWarningType::DeleteSilently : EDeleteWidgetWarningType::WarnAndAskUser);
+}
+
+void FWidgetBlueprintEditorUtils::DeleteWidgets(UWidgetBlueprint* Blueprint, TSet<UWidget*> Widgets, EDeleteWidgetWarningType WarningType)
 {
 	if ( Widgets.Num() > 0 )
 	{
@@ -540,7 +569,7 @@ void FWidgetBlueprintEditorUtils::DeleteWidgets(TSharedRef<FWidgetBlueprintEdito
 		const bool bIncludeChildrenVariables = true;
 		FindUsedVariablesForWidgets(Widgets, Blueprint, UsedVariables, WidgetNames, bIncludeChildrenVariables);
 
-		if (!bSilentDelete && UsedVariables.Num()!= 0 && !ShouldContinueDeleteOperation(Blueprint, WidgetNames))
+		if (WarningType == EDeleteWidgetWarningType::WarnAndAskUser && UsedVariables.Num()!= 0 && !ShouldContinueDeleteOperation(Blueprint, WidgetNames))
 		{
 			Transaction.Cancel();
 			return;
@@ -551,9 +580,9 @@ void FWidgetBlueprintEditorUtils::DeleteWidgets(TSharedRef<FWidgetBlueprintEdito
 		Blueprint->Modify();
 
 		bool bRemoved = false;
-		for ( FWidgetReference& Item : Widgets )
+		for (UWidget* Item : Widgets)
 		{
-			UWidget* WidgetTemplate = Item.GetTemplate();
+			UWidget* WidgetTemplate = Item;
 			WidgetTemplate->SetFlags(RF_Transactional);
 			const FName WidgetName = WidgetTemplate->GetFName();
 			// Find and update all binding references in the widget blueprint
@@ -592,7 +621,7 @@ void FWidgetBlueprintEditorUtils::DeleteWidgets(TSharedRef<FWidgetBlueprintEdito
 			}
 
 			// Rename the Desired Focus that fit the Widget Deleted
-			ReplaceDesiredFocus(BlueprintEditor, WidgetTemplate->GetFName(), FName());
+			ReplaceDesiredFocus(Blueprint, WidgetTemplate->GetFName(), FName());
 
 			// Rename the removed widget to the transient package so that it doesn't conflict with future widgets sharing the same name.
 			WidgetTemplate->Rename(nullptr, GetTransientPackage());
@@ -642,17 +671,17 @@ void FWidgetBlueprintEditorUtils::DeleteWidgets(TSharedRef<FWidgetBlueprintEdito
 	}
 }
 
-void FWidgetBlueprintEditorUtils::FindUsedVariablesForWidgets(const TSet<FWidgetReference>& Widgets, const UWidgetBlueprint* BP, TArray<UWidget*>& UsedVariables, TArray<FText>& WidgetNames, bool bIncludeVariablesOnChildren)
+void FWidgetBlueprintEditorUtils::FindUsedVariablesForWidgets(const TSet<UWidget*>& Widgets, const UWidgetBlueprint* BP, TArray<UWidget*>& UsedVariables, TArray<FText>& WidgetNames, bool bIncludeVariablesOnChildren)
 {
 	TSet<UWidget*> AllWidgets;
 	AllWidgets.Reserve(Widgets.Num());
-	for (const FWidgetReference& Item : Widgets)
+	for (UWidget* Item : Widgets)
 	{
-		AllWidgets.Add(Item.GetTemplate());
+		AllWidgets.Add(Item);
 		if (bIncludeVariablesOnChildren)
 		{
 			TArray<UWidget*> ChildWidgets;
-			UWidgetTree::GetChildWidgets(Item.GetTemplate(), ChildWidgets);
+			UWidgetTree::GetChildWidgets(Item, ChildWidgets);
 			AllWidgets.Append(ChildWidgets);
 		}
 	}
@@ -1075,7 +1104,7 @@ void FWidgetBlueprintEditorUtils::BuildReplaceWithMenu(FMenuBuilder& Menu, TShar
 				FText::GetEmpty(),
 				FSlateIcon(),
 				FUIAction(
-					FExecuteAction::CreateStatic(&FWidgetBlueprintEditorUtils::ReplaceWidgets, BlueprintEditor, BP, Widgets, ReplacementClass)
+					FExecuteAction::CreateStatic(&FWidgetBlueprintEditorUtils::ReplaceWidgets, BP, ResolveWidgetTemplates(Widgets), ReplacementClass, EReplaceWidgetNamingMethod::MaintainNameAndReferences)
 				));
 		}
 	}
@@ -1083,20 +1112,21 @@ void FWidgetBlueprintEditorUtils::BuildReplaceWithMenu(FMenuBuilder& Menu, TShar
 }
 
 
-bool FWidgetBlueprintEditorUtils::IsDesiredFocusWiget(TSharedRef<FWidgetBlueprintEditor> BlueprintEditor, UWidget* Widget)
+bool FWidgetBlueprintEditorUtils::IsDesiredFocusWidget(TSharedRef<FWidgetBlueprintEditor> BlueprintEditor, UWidget* Widget)
 {
-	// Verify if the Name changed is the Desired focus Widget name.
-	if (UUserWidget* PrevieWidget = BlueprintEditor->GetPreview())
+	return IsDesiredFocusWidget(BlueprintEditor->GetWidgetBlueprintObj(), Widget);
+}
+
+bool FWidgetBlueprintEditorUtils::IsDesiredFocusWidget(UWidgetBlueprint* Blueprint, UWidget* Widget)
+{
+	if (Blueprint && Blueprint->GeneratedClass && Widget)
 	{
-		if (UWidgetBlueprintGeneratedClass* BGClass = PrevieWidget->GetWidgetTreeOwningClass())
+		if (UUserWidget* WidgetCDO = Blueprint->GeneratedClass->GetDefaultObject<UUserWidget>())
 		{
-			UUserWidget* WidgetCDO = BGClass->GetDefaultObject<UUserWidget>();
-			if (WidgetCDO && WidgetCDO->GetDesiredFocusWidgetName() == Widget->GetFName())
-			{
-				return true;
-			}
+			return (WidgetCDO && WidgetCDO->GetDesiredFocusWidgetName() == Widget->GetFName());
 		}
 	}
+
 	return false;
 }
 
@@ -1104,112 +1134,24 @@ void FWidgetBlueprintEditorUtils::ReplaceWidgetWithSelectedTemplate(TSharedRef<F
 {
 	// @Todo: Needs to deal with bound object in animation tracks
 
-	FScopedTransaction Transaction(LOCTEXT("ReplaceWidgets", "Replace Widgets"));
-
 	UWidget* WidgetToReplace = Widget.GetTemplate();
-
 	if (!WidgetToReplace)
 	{
-		Transaction.Cancel();
 		return;
 	}
-	
-	// If they are both panel widgets then call the generic ReplaceWidgets function
-	if (WidgetToReplace->IsA<UPanelWidget>())
+
+	UClass* ReplacementWidgetClass = BlueprintEditor->GetSelectedTemplate().Get();
+	if (!ReplacementWidgetClass)
 	{
-		UClass* ReplacementWidgetClass = BlueprintEditor->GetSelectedTemplate().Get();
-		if (ReplacementWidgetClass && ReplacementWidgetClass->IsChildOf<UPanelWidget>())
-		{
-			TSet<FWidgetReference> ChildWidgetToReplace;
-			ChildWidgetToReplace.Add(Widget);
-			ReplaceWidgets(BlueprintEditor, BP, ChildWidgetToReplace, ReplacementWidgetClass);
-			return;
-		}
+		ReplacementWidgetClass = BlueprintEditor->GetSelectedUserWidget().GetClass(EResolveClass::Yes);
 	}
 	
-	UWidget* NewReplacementWidget;
-	if (BlueprintEditor->GetSelectedTemplate().IsValid())
-	{
-		UClass* WidgetClass = BlueprintEditor->GetSelectedTemplate().Get();
-		TSharedPtr<FWidgetTemplateClass> Template = MakeShareable(new FWidgetTemplateClass(WidgetClass));
-		NewReplacementWidget = Template->Create(BP->WidgetTree);
-	}
-	else if (BlueprintEditor->GetSelectedUserWidget().GetSoftObjectPath().IsValid())
-	{
-		FAssetData WidgetAssetData = BlueprintEditor->GetSelectedUserWidget();
-		TSharedPtr<FWidgetTemplateBlueprintClass> Template = MakeShareable(new FWidgetTemplateBlueprintClass(WidgetAssetData));
-		NewReplacementWidget = Template->Create(BP->WidgetTree);
-	}
-	else
+	if (!ReplacementWidgetClass)
 	{
 		return;
 	}
-
-	NewReplacementWidget->SetFlags(RF_Transactional);
-	NewReplacementWidget->Modify();
-	WidgetToReplace->SetFlags(RF_Transactional);
-	WidgetToReplace->Modify();
-
-	BP->WidgetTree->SetFlags(RF_Transactional);
-	BP->WidgetTree->Modify();
-
-	// We need to check before replacing because the Widget might be deleted, reseting the DesiredFocus
-	const bool bReplacingDesiredFocus = IsDesiredFocusWiget(BlueprintEditor, WidgetToReplace);
-	// Attempt to keep the same widget name if the types match up
-	const bool bCanKeepName = !WidgetToReplace->IsGeneratedName()
-		&& ((WidgetToReplace->IsA<UPanelWidget>() && NewReplacementWidget->IsA<UPanelWidget>()) || (WidgetToReplace->IsA(NewReplacementWidget->GetClass()) || NewReplacementWidget->IsA(WidgetToReplace->GetClass())));
 	
-	const FName OriginalWidgetName = WidgetToReplace->GetFName();
-
-	// Look if the Widget to replace is a NamedSlot.
-	if (TScriptInterface<INamedSlotInterface> NamedSlotHost = FindNamedSlotHostForContent(WidgetToReplace, BP->WidgetTree))
-	{
-		ReplaceNamedSlotHostContent(WidgetToReplace, NamedSlotHost, NewReplacementWidget);
-	}
-	else if (UPanelWidget* CurrentParent = WidgetToReplace->GetParent())
-	{
-		CurrentParent->SetFlags(RF_Transactional);
-		CurrentParent->Modify();
-		CurrentParent->ReplaceChild(WidgetToReplace, NewReplacementWidget);
-	}
-	else if (WidgetToReplace == BP->WidgetTree->RootWidget)
-	{
-		BP->WidgetTree->RootWidget = NewReplacementWidget;
-	}
-	else
-	{
-		return;
-	}
-
-	FString NewWidgetName = WidgetToReplace->GetName();
-	// Rename the removed widget to the transient package so that it doesn't conflict with the new widget if we try to keep the same name.
-	WidgetToReplace->Rename(nullptr, GetTransientPackage());
-
-	if (bCanKeepName)
-	{
-		NewWidgetName = FindNextValidName(BP->WidgetTree, NewWidgetName);
-		NewReplacementWidget->Rename(*NewWidgetName, BP->WidgetTree);
-
-		// Since we're attempting to preserve the widget name, we'll also preserve its guid
-		BP->OnVariableRenamed(OriginalWidgetName, NewReplacementWidget->GetFName());
-	}
-	else if (NewReplacementWidget->GetFName() != OriginalWidgetName)
-	{
-		BP->OnVariableRemoved(OriginalWidgetName);
-		BP->OnVariableAdded(NewReplacementWidget->GetFName());
-	}
-
-	if (bReplacingDesiredFocus)
-	{
-		FWidgetBlueprintEditorUtils::SetDesiredFocus(BlueprintEditor, NewReplacementWidget->GetFName());
-	}
-
-	// Delete the widget that has been replaced
-	TSet<FWidgetReference> WidgetsToDelete;
-	WidgetsToDelete.Add(Widget);
-	DeleteWidgets(BlueprintEditor, BP, WidgetsToDelete);
-
-	FBlueprintEditorUtils::MarkBlueprintAsStructurallyModified(BP);
+	ReplaceWidgets(BP, {WidgetToReplace}, ReplacementWidgetClass, EReplaceWidgetNamingMethod::MaintainNameAndReferences);
 }
 
 bool FWidgetBlueprintEditorUtils::CanBeReplacedWithTemplate(TSharedRef<FWidgetBlueprintEditor> BlueprintEditor, UWidgetBlueprint* BP, FWidgetReference Widget)
@@ -1280,7 +1222,7 @@ void FWidgetBlueprintEditorUtils::ReplaceWidgetWithChildren(TSharedRef<FWidgetBl
 	TArray<FText> WidgetNames;
 
 	const bool bIncludeChildrenVariables = false;
-	FindUsedVariablesForWidgets(WidgetsToDelete, BP, UsedVariables, WidgetNames, bIncludeChildrenVariables);
+	FindUsedVariablesForWidgets(ResolveWidgetTemplates(WidgetsToDelete), BP, UsedVariables, WidgetNames, bIncludeChildrenVariables);
 
 	if (UsedVariables.Num() != 0 && !ShouldContinueReplaceOperation(BP, WidgetNames))
 	{
@@ -1324,8 +1266,7 @@ void FWidgetBlueprintEditorUtils::ReplaceWidgetWithChildren(TSharedRef<FWidgetBl
 		}
 
 		// Delete the widget that has been replaced
-		const bool bForceDelete = true;
-		DeleteWidgets(BlueprintEditor, BP, WidgetsToDelete, bForceDelete);
+		DeleteWidgets(BP, ResolveWidgetTemplates(WidgetsToDelete), EDeleteWidgetWarningType::DeleteSilently);
 
 		FBlueprintEditorUtils::MarkBlueprintAsStructurallyModified(BP);
 	}
@@ -1380,13 +1321,13 @@ void FWidgetBlueprintEditorUtils::ReplaceWidgetWithNamedSlot(TSharedRef<FWidgetB
 		}
 
 		// Remove the widget replaced
-		DeleteWidgets(BlueprintEditor, BP, {Widget});
+		DeleteWidgets(BP, {Widget.GetTemplate()}, EDeleteWidgetWarningType::WarnAndAskUser);
 
 		FBlueprintEditorUtils::MarkBlueprintAsStructurallyModified(BP);
 	}
 }
 
-void FWidgetBlueprintEditorUtils::ReplaceWidgets(TSharedRef<FWidgetBlueprintEditor> BlueprintEditor, UWidgetBlueprint* BP, TSet<FWidgetReference> Widgets, UClass* WidgetClass)
+void FWidgetBlueprintEditorUtils::ReplaceWidgets(UWidgetBlueprint* BP, TSet<UWidget*> Widgets, UClass* WidgetClass, EReplaceWidgetNamingMethod NewWidgetNamingMethod)
 {
 	FScopedTransaction Transaction(LOCTEXT("ReplaceWidgets", "Replace Widgets"));
 
@@ -1402,15 +1343,29 @@ void FWidgetBlueprintEditorUtils::ReplaceWidgets(TSharedRef<FWidgetBlueprintEdit
 	}
 
 	TSharedPtr<FWidgetTemplateClass> Template = MakeShareable(new FWidgetTemplateClass(WidgetClass));
-
-	for ( FWidgetReference& Item : Widgets )
+	TMap<FName, FName> RenamedWidgetMap;
+	
+	for (UWidget* Item : Widgets)
 	{
 		BP->WidgetTree->SetFlags(RF_Transactional);
 		BP->WidgetTree->Modify();
 
-		UPanelWidget* NewReplacementWidget = CastChecked<UPanelWidget>(Template->Create(BP->WidgetTree));
+		UWidget* NewReplacementWidget = Template->Create(BP->WidgetTree);
+		UWidget* WidgetToReplace = Item;
+
+		// If replacing a panel widget, then it must not have children or the replacement must also be a panel widget
+		if (UPanelWidget* ExistingPanel = Cast<UPanelWidget>(WidgetToReplace))
+		{
+			if (ExistingPanel->GetChildrenCount() > 0 && !NewReplacementWidget->IsA<UPanelWidget>())
+			{
+				continue;
+			}
+		}
+		
+		TMap<FName, FString> ExportedProperties;
+		ExportPropertiesToText(WidgetToReplace, ExportedProperties);
+		ImportPropertiesFromText(NewReplacementWidget, ExportedProperties);
 
-		UWidget* WidgetToReplace = Item.GetTemplate();
 		WidgetToReplace->SetFlags(RF_Transactional);
 		WidgetToReplace->Modify();
 
@@ -1449,33 +1404,47 @@ void FWidgetBlueprintEditorUtils::ReplaceWidgets(TSharedRef<FWidgetBlueprintEdit
 
 		if (UPanelWidget* ExistingPanel = Cast<UPanelWidget>(WidgetToReplace))
 		{
-			while (ExistingPanel->GetChildrenCount() > 0)
+			if (UPanelWidget* NewReplacementPanelWidget = Cast<UPanelWidget>(NewReplacementWidget))
 			{
-				UWidget* Widget = ExistingPanel->GetChildAt(0);
-				Widget->SetFlags(RF_Transactional);
-				Widget->Modify();
+				while (ExistingPanel->GetChildrenCount() > 0)
+				{
+					UWidget* Widget = ExistingPanel->GetChildAt(0);
+					Widget->SetFlags(RF_Transactional);
+					Widget->Modify();
 
-				NewReplacementWidget->AddChild(Widget);
+					NewReplacementPanelWidget->AddChild(Widget);
+				}
 			}
 		}
 
 		// We need to check before replacing because the Widget might be deleted, reseting the DesiredFocus
-		bool bReplacingDesiredFocus = IsDesiredFocusWiget(BlueprintEditor, WidgetToReplace);
+		const bool bReplacingDesiredFocus = IsDesiredFocusWidget(BP, WidgetToReplace);
 
 		FString ReplaceName = WidgetToReplace->GetName();
-		const bool bIsGeneratedName = WidgetToReplace->IsGeneratedName();
+		const bool bCanKeepName = !WidgetToReplace->IsGeneratedName() && NewWidgetNamingMethod == EReplaceWidgetNamingMethod::MaintainNameAndReferences
+			&& ((WidgetToReplace->IsA<UPanelWidget>() && NewReplacementWidget->IsA<UPanelWidget>())
+				|| WidgetToReplace->IsA(NewReplacementWidget->GetClass())
+				|| NewReplacementWidget->IsA(WidgetToReplace->GetClass()));
 		
 		// Rename the removed widget to the transient package so that it doesn't conflict with the new widget if we try to keep the same name.
-		WidgetToReplace->Rename(nullptr, GetTransientPackage());
+		FName TrashName = MakeUniqueObjectName(GetTransientPackage(), WidgetToReplace->GetClass(), *FString::Printf(TEXT("TRASH_%s"), *WidgetToReplace->GetName()));
+		WidgetToReplace->Rename(*TrashName.ToString(), GetTransientPackage());
 
 		// Rename the new Widget to maintain the current name if it's not a generic name
-		if (!bIsGeneratedName)
+		if (NewWidgetNamingMethod == EReplaceWidgetNamingMethod::MaintainNameAndReferences)
 		{
-			ReplaceName = FindNextValidName(BP->WidgetTree, ReplaceName);
-			NewReplacementWidget->Rename(*ReplaceName, BP->WidgetTree);
+			if (bCanKeepName)
+			{
+				ReplaceName = FindNextValidName(BP->WidgetTree, ReplaceName);
+				NewReplacementWidget->Rename(*ReplaceName, BP->WidgetTree);
+			}
 
-			// Since we're attempting to preserve the widget name, we'll also preserve its guid
-			BP->OnVariableRenamed(OriginalWidgetName, NewReplacementWidget->GetFName());
+			// Preserve references to the widget if we haven't kept the same name
+			if (OriginalWidgetName != NewReplacementWidget->GetFName())
+			{
+				BP->OnVariableRenamed(OriginalWidgetName, NewReplacementWidget->GetFName());
+				RenamedWidgetMap.Add(OriginalWidgetName, NewReplacementWidget->GetFName());
+			}
 		}
 		else if (NewReplacementWidget->GetFName() != OriginalWidgetName)
 		{
@@ -1484,22 +1453,26 @@ void FWidgetBlueprintEditorUtils::ReplaceWidgets(TSharedRef<FWidgetBlueprintEdit
 		}
 
 		// Delete the widget that has been replaced
-		constexpr bool bForceDelete = true;
-		DeleteWidgets(BlueprintEditor, BP, {Item}, bForceDelete);
+		DeleteWidgets(BP, {Item}, EDeleteWidgetWarningType::DeleteSilently);
 
 		if (bReplacingDesiredFocus)
 		{
-			FWidgetBlueprintEditorUtils::SetDesiredFocus(BlueprintEditor, NewReplacementWidget->GetFName());
+			SetDesiredFocus(BP, NewReplacementWidget->GetFName());
 		}
 	}
 
 	FBlueprintEditorUtils::MarkBlueprintAsStructurallyModified(BP);
+
+	for (const TPair<FName, FName>& RenamedWidgets : RenamedWidgetMap)
+	{
+		FBlueprintEditorUtils::ReplaceVariableReferences(BP, RenamedWidgets.Key, RenamedWidgets.Value);
+	}
 }
 
 void FWidgetBlueprintEditorUtils::CutWidgets(TSharedRef<FWidgetBlueprintEditor> BlueprintEditor, UWidgetBlueprint* BP, TSet<FWidgetReference> Widgets)
 {
 	CopyWidgets(BP, Widgets);
-	DeleteWidgets(BlueprintEditor, BP, Widgets);
+	DeleteWidgets(BP, ResolveWidgetTemplates(Widgets), EDeleteWidgetWarningType::WarnAndAskUser);
 }
 
 void FWidgetBlueprintEditorUtils::CopyWidgets(UWidgetBlueprint* BP, TSet<FWidgetReference> Widgets)
@@ -2037,8 +2010,7 @@ TArray<UWidget*> FWidgetBlueprintEditorUtils::PasteWidgetsInternal(TSharedRef<FW
 
 			if (ChildWidgetToDelete)
 			{
-				constexpr bool bForceDelete = true;
-				DeleteWidgets(BlueprintEditor, BP, { BlueprintEditor->GetReferenceFromTemplate(ChildWidgetToDelete) }, bForceDelete);
+				DeleteWidgets(BP, { ChildWidgetToDelete }, EDeleteWidgetWarningType::DeleteSilently);
 			}
 		}
 
@@ -2297,15 +2269,19 @@ void FWidgetBlueprintEditorUtils::ExportPropertiesToText(UObject* Object, TMap<F
 {
 	if ( Object )
 	{
-		for ( TFieldIterator<FProperty> PropertyIt(Object->GetClass(), EFieldIteratorFlags::ExcludeSuper); PropertyIt; ++PropertyIt )
+		static const TSet<FName> SpecialCaseProperties = { TEXT("bIsVariable") };
+		for ( TFieldIterator<FProperty> PropertyIt(Object->GetClass()); PropertyIt && PropertyIt.GetStruct() != UVisual::StaticClass(); ++PropertyIt )
 		{
 			FProperty* Property = *PropertyIt;
 
-			// Don't serialize out object properties, we just want value data.
-			if ( !Property->IsA<FObjectProperty>() )
+			// Skip transient and instanced properties, only include properties that the user can directly edit and some special cases
+			if (!Property->HasAnyPropertyFlags(CPF_TextExportTransient | CPF_Transient | CPF_DuplicateTransient) && (
+					Property->HasAllPropertyFlags(CPF_Edit) ||
+					Property->IsA<FDelegateProperty>() ||
+					SpecialCaseProperties.Contains(Property->GetFName())))
 			{
 				FString ValueText;
-				if ( Property->ExportText_InContainer(0, ValueText, Object, Object, Object, PPF_IncludeTransient) )
+				if (Property->ExportText_InContainer(0, ValueText, Object, Object, Object, PPF_Copy))
 				{
 					ExportedProperties.Add(Property->GetFName(), ValueText);
 				}
@@ -2326,7 +2302,7 @@ void FWidgetBlueprintEditorUtils::ImportPropertiesFromText(UObject* Object, cons
 				PropertyChain.AddHead(Property);
 				Object->PreEditChange(PropertyChain);
 
-				Property->ImportText_InContainer(*Entry.Value, Object, Object, 0);
+				Property->ImportText_InContainer(*Entry.Value, Object, Object, PPF_Copy);
 
 				FPropertyChangedEvent ChangedEvent(Property);
 				Object->PostEditChangeProperty(ChangedEvent);
@@ -3169,4 +3145,20 @@ UWidgetBlueprint* FWidgetBlueprintEditorUtils::GetWidgetBlueprintFromWidget(cons
 	return nullptr;
 }
 
+TSet<UWidget*> FWidgetBlueprintEditorUtils::ResolveWidgetTemplates(const TSet<FWidgetReference>& Widgets)
+{
+	TSet<UWidget*> Templates;
+	Algo::TransformIf(Widgets, Templates,
+		[&](const FWidgetReference& Widget)
+		{
+			return Widget.GetTemplate() != nullptr;
+		},
+		[&](const FWidgetReference& Widget)
+		{
+			return Widget.GetTemplate();
+		});
+
+	return Templates;
+}
+
 #undef LOCTEXT_NAMESPACE
diff --git a/Engine/Source/Editor/UMGEditor/Public/WidgetBlueprintEditorUtils.h b/Engine/Source/Editor/UMGEditor/Public/WidgetBlueprintEditorUtils.h
--- a/Engine/Source/Editor/UMGEditor/Public/WidgetBlueprintEditorUtils.h
+++ b/Engine/Source/Editor/UMGEditor/Public/WidgetBlueprintEditorUtils.h
@@ -42,8 +42,10 @@ public:
 	static bool RenameWidget(TSharedRef<class FWidgetBlueprintEditor> BlueprintEditor, const FName& OldObjectName, const FString& NewDisplayName);
 
 	static void ReplaceDesiredFocus(TSharedRef<FWidgetBlueprintEditor> BlueprintEditor, const FName& OldName, const FName& NewName);
+	static void ReplaceDesiredFocus(UWidgetBlueprint* Blueprint, const FName& OldName, const FName& NewName);
 	
 	static void SetDesiredFocus(TSharedRef<FWidgetBlueprintEditor> BlueprintEditor, const FName DesiredFocusWidgetName);
+	static void SetDesiredFocus(UWidgetBlueprint* Blueprint, const FName& DesiredFocusWidgetName);
 
 	static void CreateWidgetContextMenu(FMenuBuilder& MenuBuilder, TSharedRef<FWidgetBlueprintEditor> BlueprintEditor, FVector2D TargetLocation);
 
@@ -51,8 +53,29 @@ public:
 
 	static TArray<UWidget*> PasteWidgets(TSharedRef<FWidgetBlueprintEditor> BlueprintEditor, UWidgetBlueprint* BP, FWidgetReference ParentWidget, FName SlotName, FVector2D PasteLocation);
 
+	enum class EReplaceWidgetNamingMethod
+	{
+		// Will give the new widget the same name as the replaced widget
+		// If it's using a generated name, the new generated name will be used and FBlueprintEditorUtils::ReplaceVariableReferences will be called
+		MaintainNameAndReferences,
+
+		// Will use the new widget's generated name and not make an effort to maintain references
+		UseNewGeneratedName
+	};
+	static void ReplaceWidgets(UWidgetBlueprint* BP, TSet<UWidget*> Widgets, UClass* WidgetClass, EReplaceWidgetNamingMethod NewWidgetNamingMethod);
+
+	UE_DEPRECATED(5.6, "FWidgetBlueprintEditorUtils::DeleteWidgets no longer takes in the BlueprintEditor, use the version without it instead")
 	static void DeleteWidgets(TSharedRef<FWidgetBlueprintEditor> BlueprintEditor, UWidgetBlueprint* BP, TSet<FWidgetReference> Widgets, bool bSilentDelete = false);
 
+	enum class EDeleteWidgetWarningType
+	{
+		// If the widget being deleted is referenced in the graph, ask the user if the deletion should continue
+		WarnAndAskUser,
+		// Don't notify the user at all and delete the widget even if it is referenced
+		DeleteSilently,
+	};
+	static void DeleteWidgets(UWidgetBlueprint* BP, TSet<UWidget*> Widgets, EDeleteWidgetWarningType WarningType);
+
 	static void CutWidgets(TSharedRef<FWidgetBlueprintEditor> BlueprintEditor, UWidgetBlueprint* BP, TSet<FWidgetReference> Widgets);
 
 	static TArray<UWidget*> DuplicateWidgets(TSharedRef<FWidgetBlueprintEditor> BlueprintEditor, UWidgetBlueprint* BP, TSet<FWidgetReference> Widgets);
@@ -141,6 +164,8 @@ public:
 	static const UWidgetEditingProjectSettings* GetRelevantSettings(TWeakPtr<FWidgetBlueprintEditor> CurrentEditor);
 
 	static UWidgetBlueprint* GetWidgetBlueprintFromWidget(const UWidget* Widget);
+	
+	static TSet<UWidget*> ResolveWidgetTemplates(const TSet<FWidgetReference>& Widgets);
 
 private:
 
@@ -170,11 +195,9 @@ private:
 
 	static void ReplaceWidgetWithNamedSlot(TSharedRef<FWidgetBlueprintEditor> BlueprintEditor, UWidgetBlueprint* BP, FWidgetReference Widget, FName NamedSlot);
 
-	static void ReplaceWidgets(TSharedRef<FWidgetBlueprintEditor> BlueprintEditor, UWidgetBlueprint* BP, TSet<FWidgetReference> Widgets, UClass* WidgetClass);
-
 	static FString FindNextValidName(UWidgetTree* WidgetTree, const FString& Name);
 
-	static void FindUsedVariablesForWidgets(const TSet<FWidgetReference>& Widgets, const UWidgetBlueprint* BP, TArray<UWidget*>& UsedVariables, TArray<FText>& WidgetNames, bool bIncludeVariablesOnChildren);
+	static void FindUsedVariablesForWidgets(const TSet<UWidget*>& Widgets, const UWidgetBlueprint* BP, TArray<UWidget*>& UsedVariables, TArray<FText>& WidgetNames, bool bIncludeVariablesOnChildren);
 
 	static bool ShouldContinueDeleteOperation(UWidgetBlueprint* BP, const TArray<FText>& WidgetNames);
 
@@ -182,7 +205,8 @@ private:
 
 	static TOptional<FWidgetThumbnailProperties> DrawSWidgetInRenderTargetInternal(UUserWidget* WidgetInstance, FRenderTarget* RenderTarget2D, UTextureRenderTarget2D* TextureRenderTarget,FVector2D ThumbnailSize, bool bIsForThumbnail, TOptional<FVector2D> ThumbnailCustomSize, EThumbnailPreviewSizeMode ThumbnailSizeMode);
 	
-	static bool IsDesiredFocusWiget(TSharedRef<FWidgetBlueprintEditor> BlueprintEditor, UWidget* Widget);
+	static bool IsDesiredFocusWidget(TSharedRef<FWidgetBlueprintEditor> BlueprintEditor, UWidget* Widget);
+	static bool IsDesiredFocusWidget(UWidgetBlueprint* Blueprint, UWidget* Widget);
 
 	static FWidgetObjectTextFactory ProcessImportedText(UWidgetBlueprint* BP, const FString& TextToImport, /*out*/ UPackage*& TempPackage);
 };
