Fix missing resource warning in UMG Editor & Project Launcher

#11037

diff --git a/Engine/Source/Developer/ProjectLauncher/Private/Widgets/Cook/SProjectLauncherCookedPlatforms.h b/Engine/Source/Developer/ProjectLauncher/Private/Widgets/Cook/SProjectLauncherCookedPlatforms.h
--- a/Engine/Source/Developer/ProjectLauncher/Private/Widgets/Cook/SProjectLauncherCookedPlatforms.h
+++ b/Engine/Source/Developer/ProjectLauncher/Private/Widgets/Cook/SProjectLauncherCookedPlatforms.h
@@ -108,7 +108,7 @@ private:
 			{
 //				if (TargetPlatform->HasValidBuild(SelectedProfile->GetProjectPath(), SelectedProfile->GetBuildConfiguration()))
 				{
-					return FAppStyle::GetColor("Foreground");
+					return FAppStyle::GetColor("Colors.Foreground");
 				}
 			}
 		}
diff --git a/Engine/Source/Editor/UMGEditor/Private/Hierarchy/SHierarchyViewItem.cpp b/Engine/Source/Editor/UMGEditor/Private/Hierarchy/SHierarchyViewItem.cpp
--- a/Engine/Source/Editor/UMGEditor/Private/Hierarchy/SHierarchyViewItem.cpp
+++ b/Engine/Source/Editor/UMGEditor/Private/Hierarchy/SHierarchyViewItem.cpp
@@ -1562,7 +1562,7 @@ void SHierarchyViewItem::Construct(const FArguments& InArgs, const TSharedRef< S
 				// TODO Tooltip should tell you what the widget is setup to do
 				.ToolTipText(LOCTEXT("NavigationHierarchyToolTip", "This widget overrides the navigation preference."))
 				.Visibility_Lambda([InModel] { return InModel->DoesWidgetOverrideNavigation() ? EVisibility::Visible : EVisibility::Collapsed; })
-				.ColorAndOpacity(FCoreStyle::Get().GetSlateColor("Foreground"))
+				.ColorAndOpacity(FAppStyle::Get().GetSlateColor("Colors.Foreground"))
 				.Font(FAppStyle::Get().GetFontStyle("FontAwesome.10"))
 				.Text(FEditorFontGlyphs::Arrows)
 			]
@@ -1576,7 +1576,7 @@ void SHierarchyViewItem::Construct(const FArguments& InArgs, const TSharedRef< S
 				// TODO Tooltip should tell you what the widget is setup to do
 				.ToolTipText(LOCTEXT("FlowDirectionHierarchyToolTip", "This widget overrides the culture/localization flow direction preference."))
 				.Visibility_Lambda([InModel] { return InModel->DoesWidgetOverrideFlowDirection() ? EVisibility::Visible : EVisibility::Collapsed; })
-				.ColorAndOpacity(FCoreStyle::Get().GetSlateColor("Foreground"))
+				.ColorAndOpacity(FAppStyle::Get().GetSlateColor("Colors.Foreground"))
 				.Font(FAppStyle::Get().GetFontStyle("FontAwesome.10"))
 				.Text(FEditorFontGlyphs::Exchange)
 			]
@@ -1589,7 +1589,7 @@ void SHierarchyViewItem::Construct(const FArguments& InArgs, const TSharedRef< S
 				SNew(SButton)
 				.ContentPadding(FMargin(3, 1))
 				.ButtonStyle(FAppStyle::Get(), "HoverHintOnly")
-				.ForegroundColor(FCoreStyle::Get().GetSlateColor("Foreground"))
+				.ForegroundColor(FAppStyle::Get().GetSlateColor("Colors.Foreground"))
 				.OnClicked(this, &SHierarchyViewItem::OnToggleLockedInDesigner)
 				.Visibility(Model->CanControlLockedInDesigner() ? EVisibility::Visible : EVisibility::Hidden)
 				.ToolTipText(LOCTEXT("WidgetLockedButtonToolTip", "Locks or Unlocks this widget and all children.  Locking a widget prevents it from being selected in the designer view by clicking on them.\n\nHolding [Shift] will only affect this widget and no children."))
@@ -1615,7 +1615,7 @@ void SHierarchyViewItem::Construct(const FArguments& InArgs, const TSharedRef< S
 				SNew(SButton)
 				.ContentPadding(FMargin(3, 1))
 				.ButtonStyle(FAppStyle::Get(), "HoverHintOnly")
-				.ForegroundColor(FCoreStyle::Get().GetSlateColor("Foreground"))
+				.ForegroundColor(FAppStyle::Get().GetSlateColor("Colors.Foreground"))
 				.OnClicked(this, &SHierarchyViewItem::OnToggleVisibility)
 				.Visibility(Model->CanControlVisibility() ? EVisibility::Visible : EVisibility::Hidden)
 				.ToolTipText(LOCTEXT("WidgetVisibilityButtonToolTip", "Toggle Widget's Editor Visibility"))
