Improved error message to help track down build errors faster.

#rnx

diff --git a/Engine/Source/Editor/UMGEditor/Private/WidgetBlueprintCompiler.cpp b/Engine/Source/Editor/UMGEditor/Private/WidgetBlueprintCompiler.cpp
--- a/Engine/Source/Editor/UMGEditor/Private/WidgetBlueprintCompiler.cpp
+++ b/Engine/Source/Editor/UMGEditor/Private/WidgetBlueprintCompiler.cpp
@@ -825,7 +825,7 @@ void FWidgetBlueprintCompilerContext::ValidateAndFixUpVariableGuids()
 				GuidToVariableNameMap.Add(It.Value(), It.Key());
 			}
 
-			if (!ensureAlwaysMsgf(SeenVariableNames.Contains(It.Key()), TEXT("Variable [%s] was deleted but still has a GUID"), *It.Key().ToString()))
+			if (!ensureAlwaysMsgf(SeenVariableNames.Contains(It.Key()), TEXT("Variable [%s] was deleted but still has a GUID referenced by WidgetBlueprint [%s]"), *It.Key().ToString(), *WidgetBP->GetName()))
 			{
 				It.RemoveCurrent();
 			}
