UMGEditor: Don't copy EditDefaultOnly properties when replacing widgets and allow forcing to keep the same name

#rnx

diff --git a/Engine/Source/Editor/UMGEditor/Private/WidgetBlueprintEditorUtils.cpp b/Engine/Source/Editor/UMGEditor/Private/WidgetBlueprintEditorUtils.cpp
--- a/Engine/Source/Editor/UMGEditor/Private/WidgetBlueprintEditorUtils.cpp
+++ b/Engine/Source/Editor/UMGEditor/Private/WidgetBlueprintEditorUtils.cpp
@@ -1343,7 +1343,7 @@ void FWidgetBlueprintEditorUtils::ReplaceWidgets(UWidgetBlueprint* BP, TSet<UWid
 	}
 
 	TSharedPtr<FWidgetTemplateClass> Template = MakeShareable(new FWidgetTemplateClass(WidgetClass));
-	TMap<FName, FName> RenamedWidgetMap;
+	TMap<FName, FName> ReplacedWidgetMap;
 	
 	for (UWidget* Item : Widgets)
 	{
@@ -1421,17 +1421,18 @@ void FWidgetBlueprintEditorUtils::ReplaceWidgets(UWidgetBlueprint* BP, TSet<UWid
 		const bool bReplacingDesiredFocus = IsDesiredFocusWidget(BP, WidgetToReplace);
 
 		FString ReplaceName = WidgetToReplace->GetName();
-		const bool bCanKeepName = !WidgetToReplace->IsGeneratedName() && NewWidgetNamingMethod == EReplaceWidgetNamingMethod::MaintainNameAndReferences
-			&& ((WidgetToReplace->IsA<UPanelWidget>() && NewReplacementWidget->IsA<UPanelWidget>())
-				|| WidgetToReplace->IsA(NewReplacementWidget->GetClass())
-				|| NewReplacementWidget->IsA(WidgetToReplace->GetClass()));
+		const bool bCanKeepName = (NewWidgetNamingMethod == EReplaceWidgetNamingMethod::MaintainNameAndReferencesForUnmatchingClass) ||
+			(!WidgetToReplace->IsGeneratedName() && NewWidgetNamingMethod == EReplaceWidgetNamingMethod::MaintainNameAndReferences
+				&& ((WidgetToReplace->IsA<UPanelWidget>() && NewReplacementWidget->IsA<UPanelWidget>())
+					|| WidgetToReplace->IsA(NewReplacementWidget->GetClass())
+					|| NewReplacementWidget->IsA(WidgetToReplace->GetClass())));
 		
 		// Rename the removed widget to the transient package so that it doesn't conflict with the new widget if we try to keep the same name.
 		FName TrashName = MakeUniqueObjectName(GetTransientPackage(), WidgetToReplace->GetClass(), *FString::Printf(TEXT("TRASH_%s"), *WidgetToReplace->GetName()));
 		WidgetToReplace->Rename(*TrashName.ToString(), GetTransientPackage());
 
 		// Rename the new Widget to maintain the current name if it's not a generic name
-		if (NewWidgetNamingMethod == EReplaceWidgetNamingMethod::MaintainNameAndReferences)
+		if (NewWidgetNamingMethod == EReplaceWidgetNamingMethod::MaintainNameAndReferences || NewWidgetNamingMethod == EReplaceWidgetNamingMethod::MaintainNameAndReferencesForUnmatchingClass)
 		{
 			if (bCanKeepName)
 			{
@@ -1443,8 +1444,10 @@ void FWidgetBlueprintEditorUtils::ReplaceWidgets(UWidgetBlueprint* BP, TSet<UWid
 			if (OriginalWidgetName != NewReplacementWidget->GetFName())
 			{
 				BP->OnVariableRenamed(OriginalWidgetName, NewReplacementWidget->GetFName());
-				RenamedWidgetMap.Add(OriginalWidgetName, NewReplacementWidget->GetFName());
 			}
+
+			// Even if the name hasn't changed, we need to replace references since the type might have changed
+			ReplacedWidgetMap.Add(OriginalWidgetName, NewReplacementWidget->GetFName());
 		}
 		else if (NewReplacementWidget->GetFName() != OriginalWidgetName)
 		{
@@ -1463,7 +1466,7 @@ void FWidgetBlueprintEditorUtils::ReplaceWidgets(UWidgetBlueprint* BP, TSet<UWid
 
 	FBlueprintEditorUtils::MarkBlueprintAsStructurallyModified(BP);
 
-	for (const TPair<FName, FName>& RenamedWidgets : RenamedWidgetMap)
+	for (const TPair<FName, FName>& RenamedWidgets : ReplacedWidgetMap)
 	{
 		FBlueprintEditorUtils::ReplaceVariableReferences(BP, RenamedWidgets.Key, RenamedWidgets.Value);
 	}
@@ -2274,10 +2277,10 @@ void FWidgetBlueprintEditorUtils::ExportPropertiesToText(UObject* Object, TMap<F
 		{
 			FProperty* Property = *PropertyIt;
 
-			// Skip transient and instanced properties, only include properties that the user can directly edit and some special cases
-			if (!Property->HasAnyPropertyFlags(CPF_TextExportTransient | CPF_Transient | CPF_DuplicateTransient) && (
+			// Skip EditDefaultOnly, transient, and instanced properties, only include properties that the user can directly edit and some special cases
+			if (!Property->HasAnyPropertyFlags(CPF_TextExportTransient | CPF_Transient | CPF_DuplicateTransient | CPF_DisableEditOnInstance) && (
 					Property->HasAllPropertyFlags(CPF_Edit) ||
-					Property->IsA<FDelegateProperty>() ||
+					Property->IsA<FMulticastDelegateProperty>() ||
 					SpecialCaseProperties.Contains(Property->GetFName())))
 			{
 				FString ValueText;
diff --git a/Engine/Source/Editor/UMGEditor/Public/WidgetBlueprintEditorUtils.h b/Engine/Source/Editor/UMGEditor/Public/WidgetBlueprintEditorUtils.h
--- a/Engine/Source/Editor/UMGEditor/Public/WidgetBlueprintEditorUtils.h
+++ b/Engine/Source/Editor/UMGEditor/Public/WidgetBlueprintEditorUtils.h
@@ -55,10 +55,13 @@ public:
 
 	enum class EReplaceWidgetNamingMethod
 	{
-		// Will give the new widget the same name as the replaced widget
+		// Will give the new widget the same name as the replaced widget if the widget classes are compatible
 		// If it's using a generated name, the new generated name will be used and FBlueprintEditorUtils::ReplaceVariableReferences will be called
 		MaintainNameAndReferences,
 
+		// Same as MaintainNameAndReferences but doesn't check for matching classes or generated names
+		MaintainNameAndReferencesForUnmatchingClass,
+
 		// Will use the new widget's generated name and not make an effort to maintain references
 		UseNewGeneratedName
 	};
