Avoid crash when parent class is null (e.g. skeleton compile when parent class is missing)

#rnx

diff --git a/Engine/Source/Editor/UMGEditor/Private/WidgetBlueprintCompiler.cpp b/Engine/Source/Editor/UMGEditor/Private/WidgetBlueprintCompiler.cpp
--- a/Engine/Source/Editor/UMGEditor/Private/WidgetBlueprintCompiler.cpp
+++ b/Engine/Source/Editor/UMGEditor/Private/WidgetBlueprintCompiler.cpp
@@ -1464,7 +1464,7 @@ void FWidgetBlueprintCompilerContext::PopulateBlueprintGeneratedVariables()
 			}
 
 			// Look in the Parent class properties to find a property with the BindWidget meta tag of the same name and Type.
-			FObjectPropertyBase* ExistingProperty = CastField<FObjectPropertyBase>(WidgetBP->ParentClass->FindPropertyByName(Widget->GetFName()));
+			FObjectPropertyBase* ExistingProperty = WidgetBP->ParentClass ? CastField<FObjectPropertyBase>(WidgetBP->ParentClass->FindPropertyByName(Widget->GetFName())) : nullptr;
 			if (ExistingProperty && 
 				FWidgetBlueprintEditorUtils::IsBindWidgetProperty(ExistingProperty) && 
 				Widget->IsA(ExistingProperty->PropertyClass))
@@ -1508,7 +1508,7 @@ void FWidgetBlueprintCompilerContext::PopulateBlueprintGeneratedVariables()
 		for (UWidgetAnimation* Animation : WidgetBP->Animations)
 		{
 			// BindWidgetAnims already have properties
-			FObjectPropertyBase* ExistingProperty = CastField<FObjectPropertyBase>(WidgetBP->ParentClass->FindPropertyByName(Animation->GetFName()));
+			FObjectPropertyBase* ExistingProperty = WidgetBP->ParentClass ? CastField<FObjectPropertyBase>(WidgetBP->ParentClass->FindPropertyByName(Animation->GetFName())) : nullptr;
 			if (ExistingProperty &&
 				FWidgetBlueprintEditorUtils::IsBindWidgetAnimProperty(ExistingProperty) &&
 				ExistingProperty->PropertyClass->IsChildOf(UWidgetAnimation::StaticClass()))
