Adding Remove UI Component functionnality to UMG Editor (experimental feature)
- Added TransactionScope to support Undo/Redo
- Small refactor to move add/remove component widget blueprint editor logic  UIComponentUtils class

diff --git a/Engine/Source/Editor/UMGEditor/Private/Customizations/UIComponentCustomizationExtender.cpp b/Engine/Source/Editor/UMGEditor/Private/Customizations/UIComponentCustomizationExtender.cpp
--- a/Engine/Source/Editor/UMGEditor/Private/Customizations/UIComponentCustomizationExtender.cpp
+++ b/Engine/Source/Editor/UMGEditor/Private/Customizations/UIComponentCustomizationExtender.cpp
@@ -4,11 +4,15 @@
 
 #include "DetailCategoryBuilder.h"
 #include "DetailLayoutBuilder.h"
+#include "DetailWidgetRow.h"
 #include "Extensions/UIComponent.h"
 #include "Extensions/UIComponentUserWidgetExtension.h"
+#include "PropertyCustomizationHelpers.h"
 #include "PropertyHandle.h"
+#include "UIComponentUtils.h"
 #include "WidgetBlueprint.h"
 #include "WidgetBlueprintEditor.h"
+#include "Widgets/Layout/SBox.h"
 
 #define LOCTEXT_NAMESPACE "UIComponentCustomizationExtender"
 
@@ -39,13 +43,51 @@ void FUIComponentCustomizationExtender::CustomizeDetails(IDetailLayoutBuilder& I
 		{
 			if (UUIComponent* Component = ComponentsOnWidget[Index])
 			{
-				IDetailCategoryBuilder& ComponentCategory = InDetailLayout.EditCategory(Component->GetFName(), FText::GetEmpty(), ECategoryPriority::Important);
+				IDetailCategoryBuilder& ComponentCategory = InDetailLayout.EditCategory(Component->GetClass()->GetFName(), FText::GetEmpty(), ECategoryPriority::Important);
 
 				IDetailPropertyRow* PropertyRow = ComponentCategory.AddExternalObjects({ Component }, EPropertyLocation::Default,
 					FAddPropertyParams()
 					.CreateCategoryNodes(false)
 					.AllowChildren(true)
 					.HideRootObjectNode(true));
+
+				const FName WidgetName = Widget->GetFName();
+				const UClass* ComponentClass = Component->GetClass();
+				
+				TSharedRef<SWidget> RemoveButton =
+					PropertyCustomizationHelpers::MakeDeleteButton(
+					FSimpleDelegate::CreateLambda([this, WidgetName, ComponentClass]()
+					{
+						if (TSharedPtr<FWidgetBlueprintEditor> Editor = WidgetBlueprintEditor.Pin())
+						{
+							FUIComponentUtils::RemoveComponent(Editor.ToSharedRef(),ComponentClass, WidgetName);
+						}
+					})
+					, LOCTEXT("RemoveUIComponent", "Remove Component from Widget.")
+				);				
+				
+				TSharedRef<SWidget> ComponentHeaderContent =
+					SNew(SHorizontalBox)
+					+ SHorizontalBox::Slot()
+					.FillWidth(1.0f)
+					.HAlign(HAlign_Left)
+					.Padding(4.0f, 4.0f)
+					.VAlign(VAlign_Center) 
+					[
+						SNew(STextBlock)
+						.Text(Component->GetClass()->GetDisplayNameText())
+						.Font(FAppStyle::Get().GetFontStyle("DetailsView.CategoryFontStyle"))
+						.TextStyle(FAppStyle::Get(), "DetailsView.CategoryTextStyle")
+					]
+					+ SHorizontalBox::Slot()					
+					.AutoWidth()
+					.HAlign(HAlign_Right)
+					.Padding(4.0f, 4.0f)
+					.VAlign(VAlign_Center) 
+					[
+						RemoveButton
+					];				
+				ComponentCategory.HeaderContent(ComponentHeaderContent, true);				
 			}
 		}
 	}
diff --git a/Engine/Source/Editor/UMGEditor/Private/Details/SWidgetDetailsView.cpp b/Engine/Source/Editor/UMGEditor/Private/Details/SWidgetDetailsView.cpp
--- a/Engine/Source/Editor/UMGEditor/Private/Details/SWidgetDetailsView.cpp
+++ b/Engine/Source/Editor/UMGEditor/Private/Details/SWidgetDetailsView.cpp
@@ -661,15 +661,10 @@ FReply SWidgetDetailsView::OnAddComponentButtonClicked(TSharedPtr<FWidgetBluepri
 		const bool bPressedOk = SClassPickerDialog::PickClass(TitleText, FUIComponentUtils::CreateClassViewerInitializationOptions(), ChosenClass, UUIComponent::StaticClass());
 		if (ChosenClass)
 		{
-			UWidgetBlueprint* WidgetBlueprint = InBlueprintEditor->GetWidgetBlueprintObj();
 			const UWidget* SelectedWidget = Cast<UWidget>(SelectedObjects[0].Get());
-			if (WidgetBlueprint && SelectedWidget)
+			if (SelectedWidget)
 			{
-				if (UUIComponentWidgetBlueprintExtension* Extension = UUIComponentWidgetBlueprintExtension::RequestExtension<UUIComponentWidgetBlueprintExtension>(WidgetBlueprint))
-				{
-					Extension->AddComponent(InBlueprintEditor.ToSharedRef(), ChosenClass, SelectedWidget->GetFName());
-					FBlueprintEditorUtils::MarkBlueprintAsStructurallyModified(WidgetBlueprint);
-				}
+				FUIComponentUtils::AddComponent(InBlueprintEditor.ToSharedRef(), ChosenClass, SelectedWidget->GetFName());
 			}
 		}
 	}
diff --git a/Engine/Source/Editor/UMGEditor/Private/Extensions/UIComponentContainerDesignerExtension.cpp b/Engine/Source/Editor/UMGEditor/Private/Extensions/UIComponentContainerDesignerExtension.cpp
--- a/Engine/Source/Editor/UMGEditor/Private/Extensions/UIComponentContainerDesignerExtension.cpp
+++ b/Engine/Source/Editor/UMGEditor/Private/Extensions/UIComponentContainerDesignerExtension.cpp
@@ -24,9 +24,10 @@ void FUIComponentContainerDesignerExtension::PreviewContentCreated(UUserWidget*
 		return;
 	}
 	
+	// Here we force to recreate the PreviewWidget Extension from the WidgetBlueprint Extension to make sure they are in sync
 	if (UUIComponentWidgetBlueprintExtension* Extension = UUIComponentWidgetBlueprintExtension::GetExtension<UUIComponentWidgetBlueprintExtension>(WidgetBlueprint))
 	{
-		UUIComponentUserWidgetExtension* UserWidgetExtension = Extension->GetOrCreateExtension(PreviewWidget);
+		Extension->GetOrCreateExtension(PreviewWidget);
 	}
 }
 
diff --git a/Engine/Source/Editor/UMGEditor/Private/UIComponentUtils.cpp b/Engine/Source/Editor/UMGEditor/Private/UIComponentUtils.cpp
--- a/Engine/Source/Editor/UMGEditor/Private/UIComponentUtils.cpp
+++ b/Engine/Source/Editor/UMGEditor/Private/UIComponentUtils.cpp
@@ -9,6 +9,10 @@
 #include "Extensions/UIComponent.h"
 #include "WidgetBlueprintEditor.h"
 #include "Extensions/UIComponentUserWidgetExtension.h"
+#include "Kismet2/BlueprintEditorUtils.h"
+#include "ScopedTransaction.h"
+
+#define LOCTEXT_NAMESPACE "UMG"
 
 FClassViewerInitializationOptions FUIComponentUtils::CreateClassViewerInitializationOptions()
 {
@@ -24,8 +28,9 @@ FClassViewerInitializationOptions FUIComponentUtils::CreateClassViewerInitializa
 }
 
 
-void FUIComponentUtils::RenameWidgetName(TSharedRef<FWidgetBlueprintEditor> BlueprintEditor, UWidgetBlueprint* WidgetBlueprint, const FName& OldVarName, const FName& NewVarName)
+void FUIComponentUtils::OnWidgetRenamed(const TSharedRef<FWidgetBlueprintEditor>& BlueprintEditor, UWidgetBlueprint* WidgetBlueprint, const FName& OldVarName, const FName& NewVarName)
 {
+	// On a Widget rename in the Editor we update the Widget names in UI Components extensions	
 	if (UUIComponentWidgetBlueprintExtension* Extension = UWidgetBlueprintExtension::GetExtension<UUIComponentWidgetBlueprintExtension>(WidgetBlueprint))
 	{
 		Extension->RenameWidget(OldVarName, NewVarName);
@@ -40,6 +45,49 @@ void FUIComponentUtils::RenameWidgetName(TSharedRef<FWidgetBlueprintEditor> Blue
 	}	
 }
 
+void FUIComponentUtils::AddComponent(const TSharedRef<FWidgetBlueprintEditor>& BlueprintEditor, const UClass* ComponentClass, const FName WidgetName)
+{	
+	UWidgetBlueprint* WidgetBlueprint = BlueprintEditor->GetWidgetBlueprintObj();
+
+	const FScopedTransaction Transaction(LOCTEXT("AddComponent", "Add Component"));
+	if (UUIComponentWidgetBlueprintExtension* WidgetBlueprintExtension = UUIComponentWidgetBlueprintExtension::RequestExtension<UUIComponentWidgetBlueprintExtension>(WidgetBlueprint))
+	{
+		UUIComponent* ComponentArchetype = WidgetBlueprintExtension->AddComponent(ComponentClass, WidgetName);		
+		UUserWidget* PreviewWidget = BlueprintEditor->GetPreview();
+		
+		if (ComponentArchetype && PreviewWidget)
+		{
+			// If the extension do not exist, we create it which will create a copy of the component we just added.
+			if (UUIComponentUserWidgetExtension* UserWidgetExtension = WidgetBlueprintExtension->GetOrCreateExtension(PreviewWidget))			
+			{
+				UserWidgetExtension->CreateAndAddComponent(ComponentArchetype, WidgetName);
+			}	
+		}
+
+		FBlueprintEditorUtils::MarkBlueprintAsStructurallyModified(WidgetBlueprint);
+	}
+}
+
+void FUIComponentUtils::RemoveComponent(const TSharedRef<FWidgetBlueprintEditor>& BlueprintEditor, const UClass* ComponentClass, const FName WidgetName)
+{
+	UWidgetBlueprint* WidgetBlueprint = BlueprintEditor->GetWidgetBlueprintObj();
+	const FScopedTransaction Transaction(LOCTEXT("RemoveComponent", "Remove Component"));
+	if (UUIComponentWidgetBlueprintExtension* WidgetBlueprintExtension = UWidgetBlueprintExtension::GetExtension<UUIComponentWidgetBlueprintExtension>(WidgetBlueprint))
+	{
+		WidgetBlueprintExtension->RemoveComponent(ComponentClass, WidgetName);
+
+		// Also Remove it from the Preview
+		if (const UUserWidget* PreviewWidget = BlueprintEditor->GetPreview())
+		{
+			if (UUIComponentUserWidgetExtension* UserWidgetExtension = PreviewWidget->GetExtension<UUIComponentUserWidgetExtension>())
+			{
+				UserWidgetExtension->RemoveComponent(ComponentClass, WidgetName);
+			}
+		}
+		FBlueprintEditorUtils::MarkBlueprintAsStructurallyModified(WidgetBlueprint);
+	}	
+}
+
 bool FUIComponentUtils::FUIComponentClassFilter::IsClassAllowed(const FClassViewerInitializationOptions& InInitOptions, const UClass* InClass, TSharedRef< FClassViewerFilterFuncs > InFilterFuncs)
 {
 	return !InClass->HasAnyClassFlags(DisallowedClassFlags)
@@ -50,4 +98,6 @@ bool FUIComponentUtils::FUIComponentClassFilter::IsUnloadedClassAllowed(const FC
 {
 	return !InUnloadedClassData->HasAnyClassFlags(DisallowedClassFlags)
 		&& InFilterFuncs->IfInChildOfClassesSet(AllowedChildrenOfClasses, InUnloadedClassData) != EFilterReturn::Failed;
-}
\ No newline at end of file
+}
+
+#undef LOCTEXT_NAMESPACE
\ No newline at end of file
diff --git a/Engine/Source/Editor/UMGEditor/Private/UIComponentWidgetBlueprintExtension.cpp b/Engine/Source/Editor/UMGEditor/Private/UIComponentWidgetBlueprintExtension.cpp
--- a/Engine/Source/Editor/UMGEditor/Private/UIComponentWidgetBlueprintExtension.cpp
+++ b/Engine/Source/Editor/UMGEditor/Private/UIComponentWidgetBlueprintExtension.cpp
@@ -26,22 +26,13 @@ UUIComponentContainer* UUIComponentWidgetBlueprintExtension::DuplicateContainer(
 	return NewContainer;
 }
 
-void UUIComponentWidgetBlueprintExtension::RemoveComponent(TSharedRef<FWidgetBlueprintEditor> Editor, UClass* ComponentClass, FName OwnerName)
+void UUIComponentWidgetBlueprintExtension::RemoveComponent(const UClass* ComponentClass, FName OwnerName)
 {
 	// Remove the Component
 	if(ensure(ComponentContainer))
 	{
 		ComponentContainer->RemoveAllComponentsOfType(ComponentClass, OwnerName);
 	}
-
-	//Also remove it from the PreviewWidget
-	if (const UUserWidget* PreviewWidget = Editor->GetPreview())
-	{
-		if(UUIComponentUserWidgetExtension* UserWidgetExtension = PreviewWidget->GetExtension<UUIComponentUserWidgetExtension>())
-		{
-			UserWidgetExtension->RemoveComponent(ComponentClass, OwnerName);
-		}
-	}	
 }
 
 TArray<UUIComponent*> UUIComponentWidgetBlueprintExtension::GetComponentsFor(UWidget* Target) const
@@ -120,17 +111,7 @@ UUIComponentUserWidgetExtension* UUIComponentWidgetBlueprintExtension::GetOrCrea
 	return UserWidgetExtension;
 }
 
-void UUIComponentWidgetBlueprintExtension::AddComponentToPreview(UUserWidget* PreviewWidget, FName OwnerName, UUIComponent* ArchetypeComponent)
-{	
-	UUIComponentUserWidgetExtension* UserWidgetExtension = GetOrCreateExtension(PreviewWidget);
-	// If the extension do not exist, we create it which will create a copy of the component we just added.
-	if (UserWidgetExtension)
-	{
-		UserWidgetExtension->CreateAndAddComponent(ArchetypeComponent, OwnerName);
-	}	
-}
-
-void UUIComponentWidgetBlueprintExtension::AddComponent(TSharedRef<FWidgetBlueprintEditor> Editor, UClass* ComponentClass, FName OwnerName)
+UUIComponent* UUIComponentWidgetBlueprintExtension::AddComponent( const UClass* ComponentClass, FName OwnerName)
 {
 	// Add the Component to the WidgetBlueprint
 	check(ComponentContainer);
@@ -138,15 +119,16 @@ void UUIComponentWidgetBlueprintExtension::AddComponent(TSharedRef<FWidgetBluepr
 	// If we already have a component of that tyype on the Widget early out
 	if (ComponentContainer->GetComponent(ComponentClass, OwnerName))
 	{
-		return;
+		return nullptr;
 	}
 
 	if (UUIComponent* NewComponent = NewObject<UUIComponent>(ComponentContainer, ComponentClass, NAME_None, RF_ArchetypeObject))
 	{
 		Modify();
 		ComponentContainer->AddComponent(OwnerName, NewComponent);
-		AddComponentToPreview(Editor->GetPreview(), OwnerName, NewComponent);
+		return NewComponent;
 	}
+	return nullptr;
 }
 
 void UUIComponentWidgetBlueprintExtension::HandleBeginCompilation(FWidgetBlueprintCompilerContext& InCreationContext)
diff --git a/Engine/Source/Editor/UMGEditor/Private/UIComponentWidgetBlueprintExtension.h b/Engine/Source/Editor/UMGEditor/Private/UIComponentWidgetBlueprintExtension.h
--- a/Engine/Source/Editor/UMGEditor/Private/UIComponentWidgetBlueprintExtension.h
+++ b/Engine/Source/Editor/UMGEditor/Private/UIComponentWidgetBlueprintExtension.h
@@ -13,7 +13,6 @@
 
 class UUIComponentUserWidgetExtension;
 class UUIComponentContainer;
-class FWidgetBlueprintEditor;
 
 UCLASS()
 /***
@@ -26,8 +25,8 @@ class UMGEDITOR_API UUIComponentWidgetBlueprintExtension : public UWidgetBluepri
 	explicit UUIComponentWidgetBlueprintExtension(const FObjectInitializer& ObjectInitializer = FObjectInitializer());
 
 public:
-	void AddComponent(TSharedRef<FWidgetBlueprintEditor> Editor, UClass* ComponentClass, FName OwnerName);
-	void RemoveComponent(TSharedRef<FWidgetBlueprintEditor> Editor, UClass* ComponentClass, FName OwnerName);
+	UUIComponent* AddComponent(const UClass* ComponentClass, FName OwnerName);
+	void RemoveComponent(const UClass* ComponentClass, FName OwnerName);
 	TArray<UUIComponent*> GetComponentsFor(UWidget* Target) const;
 	UUIComponent* GetComponent(const UClass* ComponentClass, FName OwnerName) const;
 
@@ -35,9 +34,6 @@ public:
 	void RenameWidget(const FName& OldVarName, const FName& NewVarName);
 
 	bool VerifyContainer(UUserWidget* UserWidget) const;
-	
-private:
-	void AddComponentToPreview(UUserWidget* PreviewWidget, FName OwnerName, UUIComponent* ArchetypeComponent);
 
 protected:
 	virtual void HandleBeginCompilation(FWidgetBlueprintCompilerContext& InCreationContext) override;	
diff --git a/Engine/Source/Editor/UMGEditor/Private/WidgetBlueprintEditorUtils.cpp b/Engine/Source/Editor/UMGEditor/Private/WidgetBlueprintEditorUtils.cpp
--- a/Engine/Source/Editor/UMGEditor/Private/WidgetBlueprintEditorUtils.cpp
+++ b/Engine/Source/Editor/UMGEditor/Private/WidgetBlueprintEditorUtils.cpp
@@ -448,7 +448,7 @@ bool FWidgetBlueprintEditorUtils::RenameWidget(TSharedRef<FWidgetBlueprintEditor
 		});
 
 		// If we use Components, make sure to remane the target.
-		FUIComponentUtils::RenameWidgetName(BlueprintEditor, Blueprint, OldObjectName, NewFName);		
+		FUIComponentUtils::OnWidgetRenamed(BlueprintEditor, Blueprint, OldObjectName, NewFName);		
 
 		// Validate child blueprints and adjust variable names to avoid a potential name collision
 		FBlueprintEditorUtils::ValidateBlueprintChildVariables(Blueprint, NewFName);
diff --git a/Engine/Source/Editor/UMGEditor/Public/UIComponentUtils.h b/Engine/Source/Editor/UMGEditor/Public/UIComponentUtils.h
--- a/Engine/Source/Editor/UMGEditor/Public/UIComponentUtils.h
+++ b/Engine/Source/Editor/UMGEditor/Public/UIComponentUtils.h
@@ -6,9 +6,7 @@
 #include "ClassViewerFilter.h"
 #include "UObject/ObjectMacros.h"
 
-class UUserWidget;
 class UUIComponent;
-class UUIComponentContainer;
 class UWidgetBlueprint;
 class FWidgetBlueprintEditor;
 /**  */
@@ -33,5 +31,8 @@ public:
 
 	static FClassViewerInitializationOptions CreateClassViewerInitializationOptions();
 
-	static void RenameWidgetName(TSharedRef<FWidgetBlueprintEditor> BlueprintEditor, UWidgetBlueprint* WidgetBlueprint, const FName& OldVarName, const FName& NewVarName);
+	static void OnWidgetRenamed(const TSharedRef<FWidgetBlueprintEditor>& BlueprintEditor, UWidgetBlueprint* WidgetBlueprint, const FName& OldVarName, const FName& NewVarName);
+	static void AddComponent(const TSharedRef<FWidgetBlueprintEditor>& BlueprintEditor, const UClass* ComponentClass, const FName WidgetName);	
+	static void RemoveComponent(const TSharedRef<FWidgetBlueprintEditor>& BlueprintEditor, const UClass* ComponentClass, const FName WidgetName);
+	
 };
\ No newline at end of file
diff --git a/Engine/Source/Runtime/UMG/Private/Extensions/UIComponentContainer.cpp b/Engine/Source/Runtime/UMG/Private/Extensions/UIComponentContainer.cpp
--- a/Engine/Source/Runtime/UMG/Private/Extensions/UIComponentContainer.cpp
+++ b/Engine/Source/Runtime/UMG/Private/Extensions/UIComponentContainer.cpp
@@ -46,6 +46,16 @@ void UUIComponentContainer::AddComponent(FName TargetName, UUIComponent* Compone
 
 void UUIComponentContainer::RemoveComponent(FName TargetName, UUIComponent* Component)
 {
+	bool bModified = false;
+	auto UpdateModify = [&bModified, Self = this]()
+	{
+		if (!bModified)
+		{
+			bModified = true;
+			Self->Modify();
+		}
+	};
+	
 	if(ensure(!TargetName.IsNone() && Component != nullptr))
 	{
 		for (int i = 0; i < Components.Num(); i++)
@@ -54,6 +64,7 @@ void UUIComponentContainer::RemoveComponent(FName TargetName, UUIComponent* Comp
 			if (Target.GetComponent() == Component)
 			{
 				ensure(Target.GetTargetName() == TargetName);
+				UpdateModify();
 				Components.RemoveAtSwap(i);
 				return;
 			}
@@ -65,11 +76,22 @@ void UUIComponentContainer::RemoveAllComponentsOfType(const UClass* ComponentCla
 {
 	ensure(!TargetName.IsNone());
 
+	bool bModified = false;
+	auto UpdateModify = [&bModified, Self = this]()
+	{
+		if (!bModified)
+		{
+			bModified = true;
+			Self->Modify();
+		}
+	};
+	
 	for (int32 Index = Components.Num() - 1; Index >= 0; --Index)
 	{
 		FUIComponentTarget& Target = Components[Index];
 		if (Target.GetComponent()->GetClass() == ComponentClass && Target.GetTargetName() == TargetName)
 		{
+			UpdateModify();
 			Components.RemoveAt(Index);
 		}
 	}
@@ -77,11 +99,22 @@ void UUIComponentContainer::RemoveAllComponentsOfType(const UClass* ComponentCla
 
 void UUIComponentContainer::RemoveAllComponentsFor(FName TargetName)
 {
+	bool bModified = false;
+	auto UpdateModify = [&bModified, Self = this]()
+	{
+		if (!bModified)
+		{
+			bModified = true;
+			Self->Modify();
+		}
+	};
+
 	for (int32 Index = Components.Num() - 1; Index >= 0; --Index)
 	{
 		FUIComponentTarget& Target = Components[Index];
 		if (Target.GetTargetName() == TargetName)
 		{
+			UpdateModify();
 			Components.RemoveAtSwap(Index);
 		}
 	}
