Minor change to UMG Designer to only show AddComponent when a single Widget is selected as we do not support adding Widget to multiple Widget at the same time yet.

diff --git a/Engine/Source/Editor/UMGEditor/Private/Details/SWidgetDetailsView.cpp b/Engine/Source/Editor/UMGEditor/Private/Details/SWidgetDetailsView.cpp
--- a/Engine/Source/Editor/UMGEditor/Private/Details/SWidgetDetailsView.cpp
+++ b/Engine/Source/Editor/UMGEditor/Private/Details/SWidgetDetailsView.cpp
@@ -440,7 +440,12 @@ EVisibility SWidgetDetailsView::GetCategoryAreaVisibility() const
 
 EVisibility SWidgetDetailsView::GetAddComponentAreaVisibility() const
 {
-	return IsWidgetCDOSelected() ? EVisibility::Collapsed : EVisibility::Visible;
+	if (IsWidgetCDOSelected() || SelectedObjects.Num() > 1)
+	{
+		return EVisibility::Collapsed;
+	}
+	
+	return EVisibility::Visible;
 }
 
 void SWidgetDetailsView::HandleCategoryTextCommitted(const FText& Text, ETextCommit::Type CommitType)
@@ -654,15 +659,14 @@ void SWidgetDetailsView::HandleIsVariableChanged(ECheckBoxState CheckState)
 
 FReply SWidgetDetailsView::OnAddComponentButtonClicked(TSharedPtr<FWidgetBlueprintEditor> InBlueprintEditor)
 {
-	if (SelectedObjects.Num() > 0)
+	UClass* ChosenClass = nullptr;
+	const FText TitleText = LOCTEXT("AddExtension", "Pick a component to add to the widget.");
+	const bool bPressedOk = SClassPickerDialog::PickClass(TitleText, FUIComponentUtils::CreateClassViewerInitializationOptions(), ChosenClass, UUIComponent::StaticClass());
+	if (ChosenClass)
 	{
-		UClass* ChosenClass = nullptr;
-		const FText TitleText = LOCTEXT("AddExtension", "Pick a component to add to the widget.");
-		const bool bPressedOk = SClassPickerDialog::PickClass(TitleText, FUIComponentUtils::CreateClassViewerInitializationOptions(), ChosenClass, UUIComponent::StaticClass());
-		if (ChosenClass)
+		for(TWeakObjectPtr<UObject> SelectedObject : SelectedObjects)
 		{
-			const UWidget* SelectedWidget = Cast<UWidget>(SelectedObjects[0].Get());
-			if (SelectedWidget)
+			if (const UWidget* SelectedWidget = Cast<UWidget>(SelectedObject.Get()))
 			{
 				FUIComponentUtils::AddComponent(InBlueprintEditor.ToSharedRef(), ChosenClass, SelectedWidget->GetFName());
 			}
