New setting for drawing the grid in the umg editor and the node graph using antialiasing
#12764

#tests Checked old look on nvidia drivers before and after the change. Tested switching the setting to see the difference
[RN] minor ui

diff --git a/Engine/Source/Editor/GraphEditor/Private/SNodePanel.cpp b/Engine/Source/Editor/GraphEditor/Private/SNodePanel.cpp
--- a/Engine/Source/Editor/GraphEditor/Private/SNodePanel.cpp
+++ b/Engine/Source/Editor/GraphEditor/Private/SNodePanel.cpp
@@ -1345,7 +1345,7 @@ inline float FancyMod(float Value, float Size)
 
 void SNodePanel::PaintBackgroundAsLines(const FSlateBrush* BackgroundImage, const FGeometry& AllottedGeometry, const FSlateRect& MyCullingRect, FSlateWindowElementList& OutDrawElements, int32& DrawLayerId) const
 {
-	const bool bAntialias = false;
+	const bool bAntialias = GetDefault<UEditorStyleSettings>()->bAntiAliasGrid;
 
 	const int32 RulePeriod = (int32)FAppStyle::GetFloat("Graph.Panel.GridRulePeriod");
 	check(RulePeriod > 0);
diff --git a/Engine/Source/Editor/UMGEditor/Private/Designer/SDesignSurface.cpp b/Engine/Source/Editor/UMGEditor/Private/Designer/SDesignSurface.cpp
--- a/Engine/Source/Editor/UMGEditor/Private/Designer/SDesignSurface.cpp
+++ b/Engine/Source/Editor/UMGEditor/Private/Designer/SDesignSurface.cpp
@@ -5,6 +5,7 @@
 #include "Framework/Application/SlateApplication.h"
 
 #if WITH_EDITOR
+	#include "Settings/EditorStyleSettings.h"
 	#include "Settings/LevelEditorViewportSettings.h"
 #endif // WITH_EDITOR
 
@@ -556,7 +557,7 @@ float SDesignSurface::GetGridScaleAmount() const
 
 void SDesignSurface::PaintBackgroundAsLines(const FSlateBrush* BackgroundImage, const FGeometry& AllottedGeometry, const FSlateRect& MyCullingRect, FSlateWindowElementList& OutDrawElements, int32& DrawLayerId) const
 {
-	const bool bAntialias = false;
+	const bool bAntialias = GetDefault<UEditorStyleSettings>()->bAntiAliasGrid;
 
 	const int32 RulePeriod = GetGraphRulePeriod();
 	check(RulePeriod > 0);
diff --git a/Engine/Source/Editor/UnrealEd/Classes/Settings/EditorStyleSettings.h b/Engine/Source/Editor/UnrealEd/Classes/Settings/EditorStyleSettings.h
--- a/Engine/Source/Editor/UnrealEd/Classes/Settings/EditorStyleSettings.h
+++ b/Engine/Source/Editor/UnrealEd/Classes/Settings/EditorStyleSettings.h
@@ -119,6 +119,10 @@ public:
 	UPROPERTY(EditAnywhere, config, Category = Graphs, meta = (DisplayName = "Use Grids In The Material And Blueprint Editor"))
 	uint32 bUseGrid : 1;
 
+	/** Whether grid lines should be anti-aliased */
+	UPROPERTY(EditAnywhere, config, Category = Graphs, meta = (DisplayName = "Grid Anti-Alias"))
+	uint32 bAntiAliasGrid : 1;
+
 	/** The color used to represent regular grid lines */
 	UPROPERTY(EditAnywhere, config, Category = Graphs, meta = (DisplayName = "Grid Regular Color"))
 	FLinearColor RegularColor;
diff --git a/Engine/Source/Editor/UnrealEd/Private/Settings/EditorStyleClasses.cpp b/Engine/Source/Editor/UnrealEd/Private/Settings/EditorStyleClasses.cpp
--- a/Engine/Source/Editor/UnrealEd/Private/Settings/EditorStyleClasses.cpp
+++ b/Engine/Source/Editor/UnrealEd/Private/Settings/EditorStyleClasses.cpp
@@ -44,6 +44,7 @@ UEditorStyleSettings::UEditorStyleSettings( const FObjectInitializer& ObjectInit
 	bEnableColorizedEditorTabs = true;
 	
 	bUseGrid = true;
+	bAntiAliasGrid = true;
 
 	RegularColor = FLinearColor(0.024f, 0.024, 0.024f);
 	RuleColor = FLinearColor(0.010f, 0.010f, 0.010f);
