UI Component: refresh details panel to show the component properties when the widget blueprint is compiled

diff --git a/Engine/Source/Editor/UMGEditor/Private/Customizations/UIComponentCustomizationExtender.cpp b/Engine/Source/Editor/UMGEditor/Private/Customizations/UIComponentCustomizationExtender.cpp
--- a/Engine/Source/Editor/UMGEditor/Private/Customizations/UIComponentCustomizationExtender.cpp
+++ b/Engine/Source/Editor/UMGEditor/Private/Customizations/UIComponentCustomizationExtender.cpp
@@ -7,6 +7,7 @@
 #include "DetailWidgetRow.h"
 #include "Extensions/UIComponent.h"
 #include "Extensions/UIComponentUserWidgetExtension.h"
+#include "IDetailsView.h"
 #include "PropertyCustomizationHelpers.h"
 #include "PropertyHandle.h"
 #include "UIComponentUtils.h"
@@ -91,5 +92,22 @@ void FUIComponentCustomizationExtender::CustomizeDetails(IDetailLayoutBuilder& I
 			}
 		}
 	}
+	else
+	{
+		TWeakPtr<IDetailsView> DetailsView = InDetailLayout.GetDetailsViewSharedPtr();
+		if (!UpdateQueuedForDetailsView.Contains(DetailsView))
+		{
+			TFunction<void()> UpdateDetailsPanel = [this, DetailsView]()
+				{
+					if (TSharedPtr<IDetailsView> PinnedDetailsView = DetailsView.Pin())
+					{
+						PinnedDetailsView->ForceRefresh();
+					}
+					UpdateQueuedForDetailsView.Remove(DetailsView);
+				};
+			InWidgetBlueprintEditor->AddPostDesignerLayoutAction(UpdateDetailsPanel);
+			UpdateQueuedForDetailsView.Add(DetailsView);
+		}
+	}
 }
 #undef LOCTEXT_NAMESPACE
\ No newline at end of file
diff --git a/Engine/Source/Editor/UMGEditor/Private/Customizations/UIComponentCustomizationExtender.h b/Engine/Source/Editor/UMGEditor/Private/Customizations/UIComponentCustomizationExtender.h
--- a/Engine/Source/Editor/UMGEditor/Private/Customizations/UIComponentCustomizationExtender.h
+++ b/Engine/Source/Editor/UMGEditor/Private/Customizations/UIComponentCustomizationExtender.h
@@ -7,6 +7,7 @@
 #include "UObject/WeakObjectPtr.h"
 
 class SWidget;
+class IDetailsView;
 
 class FUIComponentCustomizationExtender : public IBlueprintWidgetCustomizationExtender
 {
@@ -18,5 +19,5 @@ private:
 	/** The selected widget in the details panel. */
 	TWeakObjectPtr<class UWidget> Widget;
 	TWeakPtr<FWidgetBlueprintEditor> WidgetBlueprintEditor;
-
+	TArray<TWeakPtr<IDetailsView>> UpdateQueuedForDetailsView;
 };
