UMG: PR #10920: Fix not being able to use "Wrap With..." on the first widget in an inherited Named Slot

#rnx

diff --git a/Engine/Source/Editor/UMGEditor/Private/WidgetBlueprintEditorUtils.cpp b/Engine/Source/Editor/UMGEditor/Private/WidgetBlueprintEditorUtils.cpp
--- a/Engine/Source/Editor/UMGEditor/Private/WidgetBlueprintEditorUtils.cpp
+++ b/Engine/Source/Editor/UMGEditor/Private/WidgetBlueprintEditorUtils.cpp
@@ -749,6 +749,21 @@ bool FWidgetBlueprintEditorUtils::ShouldContinueReplaceOperation(UWidgetBlueprin
 
 TScriptInterface<INamedSlotInterface> FWidgetBlueprintEditorUtils::FindNamedSlotHostForContent(UWidget* WidgetTemplate, UWidgetTree* WidgetTree)
 {
+	// If the named slot comes from a parent widget class, the WidgetTree will be the SlotHost 
+	TArray<FName> SlotNames;
+	WidgetTree->GetSlotNames(SlotNames);
+
+	for (FName SlotName : SlotNames)
+	{
+		if (UWidget* SlotContent = WidgetTree->GetContentForSlot(SlotName))
+		{
+			if (SlotContent == WidgetTemplate)
+			{
+				return WidgetTree;
+			}
+		}
+	}
+	
 	return TScriptInterface<INamedSlotInterface>(FindNamedSlotHostWidgetForContent(WidgetTemplate, WidgetTree));
 }
 
@@ -955,10 +970,10 @@ void FWidgetBlueprintEditorUtils::WrapWidgets(TSharedRef<FWidgetBlueprintEditor>
 		int32 OutIndex;
 		UWidget* Widget = Item.GetTemplate();
 		UPanelWidget* CurrentParent = BP->WidgetTree->FindWidgetParent(Widget, OutIndex);
-		UWidget* CurrentSlot = FindNamedSlotHostWidgetForContent(Widget, BP->WidgetTree);
+		TScriptInterface<INamedSlotInterface> NamedSlotHost = FindNamedSlotHostForContent(Widget, BP->WidgetTree);
 
 		// If the widget doesn't currently have a slot or parent, and isn't the root, ignore it.
-		if (CurrentSlot == nullptr && CurrentParent == nullptr && Widget != BP->WidgetTree->RootWidget)
+		if (NamedSlotHost == nullptr && CurrentParent == nullptr && Widget != BP->WidgetTree->RootWidget)
 		{
 			continue;
 		}
@@ -967,13 +982,13 @@ void FWidgetBlueprintEditorUtils::WrapWidgets(TSharedRef<FWidgetBlueprintEditor>
 		BP->WidgetTree->SetFlags(RF_Transactional);
 		BP->WidgetTree->Modify();
 
-		if (CurrentSlot)
+		if (NamedSlotHost)
 		{
 			// If this is a named slot, we need to properly remove and reassign the slot content
-			if (TScriptInterface<INamedSlotInterface> NamedSlotHost = TScriptInterface<INamedSlotInterface>(CurrentSlot))
+			if (UObject* NamedSlotObject = NamedSlotHost.GetObject())
 			{
-				CurrentSlot->SetFlags(RF_Transactional);
-				CurrentSlot->Modify();
+				NamedSlotObject->SetFlags(RF_Transactional);
+				NamedSlotObject->Modify();
 
 				UPanelWidget* NewSlotContents = CastChecked<UPanelWidget>(Template->Create(BP->WidgetTree));
 				NewSlotContents->SetDesignerFlags(BlueprintEditor->GetCurrentDesignerFlags());
