UMGEditor: Fix not being able to save navigation settings on the root userwidget in widget blueprints

#rnx

diff --git a/Engine/Plugins/Tests/EditorTests/Source/EditorTests/Private/UnrealEd/EditorTestsUtilityLibrary.cpp b/Engine/Plugins/Tests/EditorTests/Source/EditorTests/Private/UnrealEd/EditorTestsUtilityLibrary.cpp
--- a/Engine/Plugins/Tests/EditorTests/Source/EditorTests/Private/UnrealEd/EditorTestsUtilityLibrary.cpp
+++ b/Engine/Plugins/Tests/EditorTests/Source/EditorTests/Private/UnrealEd/EditorTestsUtilityLibrary.cpp
@@ -232,6 +232,11 @@ void UEditorTestsUtilityLibrary::SetEditorWidgetNavigationRule(UWidget* Widget,
 		Widget->Navigation = NewObject<UWidgetNavigation>(Widget);
 		WidgetNavigation = Widget->Navigation;
 		WidgetNavigation->SetFlags(RF_Transactional);
+		if (Widget->IsTemplate())
+		{
+			EObjectFlags TemplateFlags = Widget->GetMaskedFlags(RF_PropagateToSubObjects) | RF_DefaultSubObject;
+			WidgetNavigation->SetFlags(TemplateFlags);
+		}
 	}
 
 	FWidgetNavigationData* DirectionNavigation = nullptr;
diff --git a/Engine/Source/Editor/UMGEditor/Private/Customizations/WidgetNavigationCustomization.cpp b/Engine/Source/Editor/UMGEditor/Private/Customizations/WidgetNavigationCustomization.cpp
--- a/Engine/Source/Editor/UMGEditor/Private/Customizations/WidgetNavigationCustomization.cpp
+++ b/Engine/Source/Editor/UMGEditor/Private/Customizations/WidgetNavigationCustomization.cpp
@@ -15,6 +15,7 @@
 #include "Details/SFunctionSelector.h"
 #include "Framework/Application/SlateApplication.h"
 #include "Blueprint/WidgetTree.h"
+#include "Kismet2/BlueprintEditorUtils.h"
 #include "WidgetBlueprint.h"
 #include "WidgetBlueprintEditor.h"
 
@@ -154,10 +155,25 @@ void FWidgetNavigationCustomization::OnWidgetSelectedForExplicitNavigation(FName
 	{
 		if ( UWidget* Widget = Cast<UWidget>(OuterObject) )
 		{
-			FWidgetReference WidgetReference = Editor.Pin()->GetReferenceFromPreview(Widget);
+			const TSharedPtr<FWidgetBlueprintEditor> WidgetBlueprintEditor = Editor.Pin();
+			if (Widget == WidgetBlueprintEditor->GetPreview())
+			{
+				SetNav(Widget, Nav, TOptional<EUINavigationRule>(), ExplictWidgetOrFunction);
+
+				UWidgetBlueprint* WidgetBlueprint = WidgetBlueprintEditor->GetWidgetBlueprintObj();
+				if (WidgetBlueprint && WidgetBlueprint->GeneratedClass)
+				{
+					SetNav(WidgetBlueprint->GeneratedClass->GetDefaultObject<UUserWidget>(), Nav, TOptional<EUINavigationRule>(), ExplictWidgetOrFunction);
+					FBlueprintEditorUtils::MarkBlueprintAsStructurallyModified(WidgetBlueprint);
+				}
+			}
+			else
+			{
+				FWidgetReference WidgetReference = WidgetBlueprintEditor->GetReferenceFromPreview(Widget);
 
-			SetNav(WidgetReference.GetPreview(), Nav, TOptional<EUINavigationRule>(), ExplictWidgetOrFunction);
-			SetNav(WidgetReference.GetTemplate(), Nav, TOptional<EUINavigationRule>(), ExplictWidgetOrFunction);
+				SetNav(WidgetReference.GetPreview(), Nav, TOptional<EUINavigationRule>(), ExplictWidgetOrFunction);
+				SetNav(WidgetReference.GetTemplate(), Nav, TOptional<EUINavigationRule>(), ExplictWidgetOrFunction);
+			}
 		}
 	}
 }
@@ -356,10 +372,25 @@ void FWidgetNavigationCustomization::HandleNavMenuEntryClicked(TWeakPtr<IPropert
 	{
 		if (UWidget* Widget = Cast<UWidget>(OuterObject))
 		{
-			FWidgetReference WidgetReference = Editor.Pin()->GetReferenceFromPreview(Widget);
+			const TSharedPtr<FWidgetBlueprintEditor> WidgetBlueprintEditor = Editor.Pin();
+			if (Widget == WidgetBlueprintEditor->GetPreview())
+			{
+				SetNav(Widget, Nav, Rule, FName(NAME_None));
 
-			SetNav(WidgetReference.GetPreview(), Nav, Rule, FName(NAME_None));
-			SetNav(WidgetReference.GetTemplate(), Nav, Rule, FName(NAME_None));
+				UWidgetBlueprint* WidgetBlueprint = WidgetBlueprintEditor->GetWidgetBlueprintObj();
+				if (WidgetBlueprint && WidgetBlueprint->GeneratedClass)
+				{
+					SetNav(WidgetBlueprint->GeneratedClass->GetDefaultObject<UUserWidget>(), Nav, Rule, FName(NAME_None));
+					FBlueprintEditorUtils::MarkBlueprintAsStructurallyModified(WidgetBlueprint);
+				}
+			}
+			else
+			{
+				FWidgetReference WidgetReference = WidgetBlueprintEditor->GetReferenceFromPreview(Widget);
+
+				SetNav(WidgetReference.GetPreview(), Nav, Rule, FName(NAME_None));
+				SetNav(WidgetReference.GetTemplate(), Nav, Rule, FName(NAME_None));
+			}
 		}
 	}
 }
@@ -378,6 +409,11 @@ void FWidgetNavigationCustomization::SetNav(UWidget* Widget, EUINavigation Nav,
 	{
 		WidgetNavigation = NewObject<UWidgetNavigation>(Widget);
 		WidgetNavigation->SetFlags(RF_Transactional);
+		if (Widget->IsTemplate())
+		{
+			EObjectFlags TemplateFlags = Widget->GetMaskedFlags(RF_PropagateToSubObjects) | RF_DefaultSubObject;
+			WidgetNavigation->SetFlags(TemplateFlags);
+		}
 	}
 
 	FWidgetNavigationData* DirectionNavigation = nullptr;
diff --git a/Engine/Source/Runtime/UMG/Private/UserWidget.cpp b/Engine/Source/Runtime/UMG/Private/UserWidget.cpp
--- a/Engine/Source/Runtime/UMG/Private/UserWidget.cpp
+++ b/Engine/Source/Runtime/UMG/Private/UserWidget.cpp
@@ -1133,6 +1133,7 @@ void UUserWidget::OnWidgetRebuilt()
 {
 	// When a user widget is rebuilt we can safely initialize the navigation now since all the slate
 	// widgets should be held onto by a smart pointer at this point.
+	BuildNavigation();
 	WidgetTree->ForEachWidget([&] (UWidget* Widget) {
 		Widget->BuildNavigation();
 	});
diff --git a/Engine/Source/Runtime/UMG/Private/WidgetBlueprintGeneratedClass.cpp b/Engine/Source/Runtime/UMG/Private/WidgetBlueprintGeneratedClass.cpp
--- a/Engine/Source/Runtime/UMG/Private/WidgetBlueprintGeneratedClass.cpp
+++ b/Engine/Source/Runtime/UMG/Private/WidgetBlueprintGeneratedClass.cpp
@@ -245,6 +245,13 @@ void UWidgetBlueprintGeneratedClass::InitializeWidgetStatic(UUserWidget* UserWid
 
 		BindAnimationsStatic(UserWidget, InAnimations, ObjectPropertiesMap);
 
+		// Support navigation rules set on dynamically created user widgets by resolving the navigation rules against itself
+		// if this widget is pre-placed in another widget, the navigation rules will be resolved afterward against the outer widget below in CreatedWidgetTree->ForEachWidget
+		if (UserWidget->Navigation)
+		{
+			UserWidget->Navigation->ResolveRules(UserWidget, CreatedWidgetTree);
+		}
+
 		CreatedWidgetTree->ForEachWidget([&](UWidget* Widget) {
 			// Not fatal if NULL, but shouldn't happen
 			if (!ensure(Widget != nullptr))
