Fix wrong style for widget child type customization in the details view.

#tests Opened the details view, validated that the combobox was looking as expected, both in its collapse and expanded state.
#rnx

diff --git a/Engine/Source/Editor/UMGEditor/Private/Customizations/WidgetChildTypeCustomization.cpp b/Engine/Source/Editor/UMGEditor/Private/Customizations/WidgetChildTypeCustomization.cpp
--- a/Engine/Source/Editor/UMGEditor/Private/Customizations/WidgetChildTypeCustomization.cpp
+++ b/Engine/Source/Editor/UMGEditor/Private/Customizations/WidgetChildTypeCustomization.cpp
@@ -38,8 +38,7 @@ void FWidgetChildTypeCustomization::CustomizeHeader(TSharedRef<IPropertyHandle>
 	.MaxDesiredWidth(FDetailWidgetRow::DefaultValueMaxWidth * 2)
 	[
 		SAssignNew(WidgetListComboButton, SComboButton)
-		.ButtonStyle(FAppStyle::Get(), "PropertyEditor.AssetComboStyle")
-		.ForegroundColor(FAppStyle::GetColor("PropertyEditor.AssetName.ColorAndOpacity"))
+		.ForegroundColor(FSlateColor::UseForeground())
 		.OnGetMenuContent(this, &FWidgetChildTypeCustomization::GetPopupContent)
 		.ContentPadding(2.0f)
 		.ButtonContent()
