Fixed a bug where the Desired Focus Widget setting in UMG User Widget blueprints was not showing the contents of Named Slots within children User Widgets as an option.

#tests created a User Widgets which contained 1 or 2 named slots, and put them as child widgets inside of another User Widget blueprint (including nested within itself), and the Desired Focus Widget setting was able to correctly select all widgets within the name slots, and widgets not in named slots. I then tried setting different widgets (those in named slots, and those not in named slots) as the desired focus widget, and tested it in runtime (to do this, in the level blueprint, on Event BeginPlay, create your widget, add it to viewport, and set it as the "In Widget to Focus" input in a "Set Input Mode Game and UI" node. Also pass "Get Player Controller"'s return value into the "Target" input in that node), and it worked in runtime.
[RN] fix creator/ui

diff --git a/Engine/Source/Editor/UMGEditor/Private/Hierarchy/SReadOnlyHierarchyView.cpp b/Engine/Source/Editor/UMGEditor/Private/Hierarchy/SReadOnlyHierarchyView.cpp
--- a/Engine/Source/Editor/UMGEditor/Private/Hierarchy/SReadOnlyHierarchyView.cpp
+++ b/Engine/Source/Editor/UMGEditor/Private/Hierarchy/SReadOnlyHierarchyView.cpp
@@ -4,6 +4,7 @@
 
 #include "Algo/Transform.h"
 #include "Blueprint/WidgetTree.h"
+#include "Components/NamedSlotInterface.h"
 #include "Styling/SlateIconFinder.h"
 #include "Widgets/Input/SSearchBox.h"
 
@@ -297,7 +298,24 @@ void SReadOnlyHierarchyView::BuildWidgetChildren(const UWidget* Widget, TSharedP
 		Parent = WidgetItem;
 	}
 
-	// even if we didn't create an item for this widget, we still want to iterate over its children
+	// Search for content within a named slot that we need to recurse into.
+	// We still want to iterate over its children even if we didn't create an item for this widget.
+	if (const INamedSlotInterface* NamedSlotHost = Cast<INamedSlotInterface>(Widget))
+	{
+		TArray<FName> SlotNames;
+		NamedSlotHost->GetSlotNames(SlotNames);
+
+		for (FName SlotName : SlotNames)
+		{
+			if (UWidget* SlotContent = NamedSlotHost->GetContentForSlot(SlotName))
+			{
+				BuildWidgetChildren(SlotContent, Parent);
+			}
+		}
+	}
+
+	// If the current widget is a panel, then we want to iterate through its children.
+	// We still want to iterate over its children even if we didn't create an item for this widget.
 	if (const UPanelWidget* PanelWidget = Cast<UPanelWidget>(Widget))
 	{
 		for (const UWidget* Child : PanelWidget->GetAllChildren())
