Updated headers using UnrealCodeFixup to make sure dllstorage is on methods/staticvars instead of types. Part 5/n

diff --git a/Engine/Source/Editor/UMGEditor/Public/BlueprintModes/WidgetBlueprintApplicationMode.h b/Engine/Source/Editor/UMGEditor/Public/BlueprintModes/WidgetBlueprintApplicationMode.h
--- a/Engine/Source/Editor/UMGEditor/Public/BlueprintModes/WidgetBlueprintApplicationMode.h
+++ b/Engine/Source/Editor/UMGEditor/Public/BlueprintModes/WidgetBlueprintApplicationMode.h
@@ -7,6 +7,8 @@
 #include "WidgetBlueprintEditor.h"
 #include "BlueprintEditorModes.h"
 
+#define UE_API UMGEDITOR_API
+
 class UWidgetBlueprint;
 
 /////////////////////////////////////////////////////
@@ -14,14 +16,14 @@ class UWidgetBlueprint;
 
 DECLARE_MULTICAST_DELEGATE_OneParam(FOnWidgetBlueprintModeTransition, class FWidgetBlueprintApplicationMode&);
 
-class UMGEDITOR_API FWidgetBlueprintApplicationMode : public FBlueprintEditorApplicationMode
+class FWidgetBlueprintApplicationMode : public FBlueprintEditorApplicationMode
 {
 public:
-	FWidgetBlueprintApplicationMode(TSharedPtr<class FWidgetBlueprintEditor> InWidgetEditor, FName InModeName);
+	UE_API FWidgetBlueprintApplicationMode(TSharedPtr<class FWidgetBlueprintEditor> InWidgetEditor, FName InModeName);
 
 	// FApplicationMode interface
-	virtual void PreDeactivateMode() override;
-	virtual void PostActivateMode() override;
+	UE_API virtual void PreDeactivateMode() override;
+	UE_API virtual void PostActivateMode() override;
 	// End of FApplicationMode interface
 
 	/** Called at start of PostActivateMode */
@@ -31,8 +33,8 @@ public:
 	mutable FOnWidgetBlueprintModeTransition OnPreDeactivateMode;
 
 public:
-	TSharedPtr<class FWidgetBlueprintEditor> GetBlueprintEditor() const;
-	UWidgetBlueprint* GetBlueprint() const;
+	UE_API TSharedPtr<class FWidgetBlueprintEditor> GetBlueprintEditor() const;
+	UE_API UWidgetBlueprint* GetBlueprint() const;
 
 protected:
 	TWeakPtr<class FWidgetBlueprintEditor> MyWidgetBlueprintEditor;
@@ -40,3 +42,5 @@ protected:
 	// Set of spawnable tabs in the mode
 	FWorkflowAllowedTabSet TabFactories;
 };
+
+#undef UE_API
diff --git a/Engine/Source/Editor/UMGEditor/Public/BlueprintModes/WidgetBlueprintApplicationModes.h b/Engine/Source/Editor/UMGEditor/Public/BlueprintModes/WidgetBlueprintApplicationModes.h
--- a/Engine/Source/Editor/UMGEditor/Public/BlueprintModes/WidgetBlueprintApplicationModes.h
+++ b/Engine/Source/Editor/UMGEditor/Public/BlueprintModes/WidgetBlueprintApplicationModes.h
@@ -6,25 +6,29 @@
 #include "Internationalization/Text.h"
 #include "UObject/NameTypes.h"
 
+#define UE_API UMGEDITOR_API
+
 /////////////////////////////////////////////////////
 // FWidgetBlueprintApplicationModes
 
 // This is the list of IDs for widget blueprint editor modes
-struct UMGEDITOR_API FWidgetBlueprintApplicationModes
+struct FWidgetBlueprintApplicationModes
 {
 	// Mode constants
-	static const FName DesignerMode;
-	static const FName GraphMode;
+	static UE_API const FName DesignerMode;
+	static UE_API const FName GraphMode;
 	UE_DEPRECATED(5.3, "DebugMode is deprecated. Use PreviewMode instead.")
-	static const FName DebugMode;
-	static const FName PreviewMode;
+	static UE_API const FName DebugMode;
+	static UE_API const FName PreviewMode;
 
-	static FText GetLocalizedMode(const FName InMode);
+	static UE_API FText GetLocalizedMode(const FName InMode);
 
 	UE_DEPRECATED(5.3, "IsDebugModeEnabled is deprecated. Use IsPreviewModeEnabled instead.")
-	static bool IsDebugModeEnabled();
-	static bool IsPreviewModeEnabled();	
+	static UE_API bool IsDebugModeEnabled();
+	static UE_API bool IsPreviewModeEnabled();	
 
 private:
 	FWidgetBlueprintApplicationModes() = delete;
 };
+
+#undef UE_API
diff --git a/Engine/Source/Editor/UMGEditor/Public/Components/AssetThumbnailWidget.h b/Engine/Source/Editor/UMGEditor/Public/Components/AssetThumbnailWidget.h
--- a/Engine/Source/Editor/UMGEditor/Public/Components/AssetThumbnailWidget.h
+++ b/Engine/Source/Editor/UMGEditor/Public/Components/AssetThumbnailWidget.h
@@ -6,6 +6,8 @@
 #include "AssetRegistry/AssetData.h"
 #include "AssetThumbnailWidget.generated.h"
 
+#define UE_API UMGEDITOR_API
+
 class FAssetThumbnail;
 class SBorder;
 class SWidget;
@@ -32,7 +34,7 @@ DECLARE_DYNAMIC_DELEGATE_RetVal(FText, FGetHighlightTextDelegate);
 
 /** Copied mostly from FAssetThumbnailConfig */
 USTRUCT(BlueprintType)
-struct UMGEDITOR_API FAssetThumbnailWidgetSettings
+struct FAssetThumbnailWidgetSettings
 {
 	GENERATED_BODY()
 	
@@ -74,7 +76,7 @@ struct UMGEDITOR_API FAssetThumbnailWidgetSettings
 	UPROPERTY(EditAnywhere, BlueprintReadWrite, Category = "Appearance")
 	EThumbnailColorStripOrientation_BlueprintType ColorStripOrientation = EThumbnailColorStripOrientation_BlueprintType::HorizontalBottomEdge;
 
-	FAssetThumbnailConfig ToThumbnailConfig() const;
+	UE_API FAssetThumbnailConfig ToThumbnailConfig() const;
 
 	PRAGMA_DISABLE_DEPRECATION_WARNINGS
 	FAssetThumbnailWidgetSettings() = default;
@@ -86,34 +88,34 @@ struct UMGEDITOR_API FAssetThumbnailWidgetSettings
 };
 
 /** This widget can be given an asset and it will render its thumbnail. Editor-only. */
-UCLASS()
-class UMGEDITOR_API UAssetThumbnailWidget : public UWidget
+UCLASS(MinimalAPI)
+class UAssetThumbnailWidget : public UWidget
 {
 	GENERATED_BODY()
 public:
 
 	UFUNCTION(BlueprintCallable, Category = "Appearance")
-	void SetThumbnailSettings(const FAssetThumbnailWidgetSettings& InThumbnailSettings);
+	UE_API void SetThumbnailSettings(const FAssetThumbnailWidgetSettings& InThumbnailSettings);
 	FAssetThumbnailWidgetSettings GetThumbnailSettings() const { return ThumbnailSettings; }
 
 	UFUNCTION(BlueprintCallable, Category = "Appearance")
-	void SetAsset(const FAssetData& AssetData);
+	UE_API void SetAsset(const FAssetData& AssetData);
 	UFUNCTION(BlueprintCallable, Category = "Appearance")
-	void SetAssetByObject(UObject* Object);
+	UE_API void SetAssetByObject(UObject* Object);
 	
 	/** Sets the resolution for the rendered thumbnail. */
 	UFUNCTION(Blueprintcallable, Category = "Appearance")
-	void SetResolution(const FIntPoint& InResolution);
+	UE_API void SetResolution(const FIntPoint& InResolution);
 
 	/** Gets the resolution of the rendered thumbnail. */
 	UFUNCTION(BlueprintPure, Category = "Appearance")
 	FIntPoint GetResolution() const { return Resolution; }
 
 	//~ Begin UWidget Interface
-	virtual void SynchronizeProperties() override;
-	virtual void ReleaseSlateResources(bool bReleaseChildren) override;
-	virtual TSharedRef<SWidget> RebuildWidget() override;
-	virtual const FText GetPaletteCategory() override;
+	UE_API virtual void SynchronizeProperties() override;
+	UE_API virtual void ReleaseSlateResources(bool bReleaseChildren) override;
+	UE_API virtual TSharedRef<SWidget> RebuildWidget() override;
+	UE_API virtual const FText GetPaletteCategory() override;
 	//~ End UWidget Interface
 	
 private:
@@ -137,5 +139,7 @@ private:
 	TSharedPtr<SBorder> DisplayedWidget;
 
 	/** Updates or creates ThumbnailRenderer and sets the contents of DisplayedWidget. */
-	void UpdateNativeAssetThumbnailWidget();
+	UE_API void UpdateNativeAssetThumbnailWidget();
 };
+
+#undef UE_API
diff --git a/Engine/Source/Editor/UMGEditor/Public/Customizations/IBlueprintWidgetCustomizationExtender.h b/Engine/Source/Editor/UMGEditor/Public/Customizations/IBlueprintWidgetCustomizationExtender.h
--- a/Engine/Source/Editor/UMGEditor/Public/Customizations/IBlueprintWidgetCustomizationExtender.h
+++ b/Engine/Source/Editor/UMGEditor/Public/Customizations/IBlueprintWidgetCustomizationExtender.h
@@ -11,7 +11,7 @@ class FWidgetBlueprintEditor;
 /**
  * Details customization extender for UWidgets.
  */
-class UMGEDITOR_API IBlueprintWidgetCustomizationExtender : public TSharedFromThis<IBlueprintWidgetCustomizationExtender>
+class IBlueprintWidgetCustomizationExtender : public TSharedFromThis<IBlueprintWidgetCustomizationExtender>
 {
 public:
 	virtual ~IBlueprintWidgetCustomizationExtender() = default;
@@ -23,4 +23,4 @@ public:
 	 * @param InWidgetBlueprintEditor		widget blueprint editor that is used for widgets being customized.
 	 */
 	virtual void CustomizeDetails(IDetailLayoutBuilder& InDetailLayout, const TArrayView<UWidget*> InWidgets, const TSharedRef<FWidgetBlueprintEditor>& InWidgetBlueprintEditor) = 0;
-};
\ No newline at end of file
+};
diff --git a/Engine/Source/Editor/UMGEditor/Public/Designer/SDesignSurface.h b/Engine/Source/Editor/UMGEditor/Public/Designer/SDesignSurface.h
--- a/Engine/Source/Editor/UMGEditor/Public/Designer/SDesignSurface.h
+++ b/Engine/Source/Editor/UMGEditor/Public/Designer/SDesignSurface.h
@@ -20,6 +20,8 @@
 #include "Widgets/DeclarativeSyntaxSupport.h"
 #include "Widgets/SCompoundWidget.h"
 
+#define UE_API UMGEDITOR_API
+
 class FActiveTimerHandle;
 class FPaintArgs;
 class FSlateWindowElementList;
@@ -29,7 +31,7 @@ struct FGeometry;
 struct FPointerEvent;
 struct FSlateBrush;
 
-class UMGEDITOR_API SDesignSurface : public SCompoundWidget
+class SDesignSurface : public SCompoundWidget
 {
 public:
 
@@ -43,53 +45,53 @@ public:
 		SLATE_ATTRIBUTE(bool, AllowContinousZoomInterpolation)
 	SLATE_END_ARGS()
 
-	void Construct(const FArguments& InArgs);
+	UE_API void Construct(const FArguments& InArgs);
 
 	// SWidget interface
-	virtual void Tick(const FGeometry& AllottedGeometry, const double InCurrentTime, const float InDeltaTime) override;
-	virtual FCursorReply OnCursorQuery(const FGeometry& MyGeometry, const FPointerEvent& CursorEvent) const override;
-	virtual FReply OnMouseButtonDown(const FGeometry& MyGeometry, const FPointerEvent& MouseEvent) override;
-	virtual FReply OnMouseButtonUp(const FGeometry& MyGeometry, const FPointerEvent& MouseEvent) override;
-	virtual FReply OnMouseMove(const FGeometry& MyGeometry, const FPointerEvent& MouseEvent) override;
-	virtual FReply OnMouseWheel(const FGeometry& MyGeometry, const FPointerEvent& MouseEvent) override;
-	virtual FReply OnTouchGesture(const FGeometry& MyGeometry, const FPointerEvent& GestureEvent) override;
-	virtual FReply OnTouchEnded(const FGeometry& MyGeometry, const FPointerEvent& InTouchEvent) override;
-	virtual int32 OnPaint(const FPaintArgs& Args, const FGeometry& AllottedGeometry, const FSlateRect& MyCullingRect, FSlateWindowElementList& OutDrawElements, int32 LayerId, const FWidgetStyle& InWidgetStyle, bool bParentEnabled) const override;
+	UE_API virtual void Tick(const FGeometry& AllottedGeometry, const double InCurrentTime, const float InDeltaTime) override;
+	UE_API virtual FCursorReply OnCursorQuery(const FGeometry& MyGeometry, const FPointerEvent& CursorEvent) const override;
+	UE_API virtual FReply OnMouseButtonDown(const FGeometry& MyGeometry, const FPointerEvent& MouseEvent) override;
+	UE_API virtual FReply OnMouseButtonUp(const FGeometry& MyGeometry, const FPointerEvent& MouseEvent) override;
+	UE_API virtual FReply OnMouseMove(const FGeometry& MyGeometry, const FPointerEvent& MouseEvent) override;
+	UE_API virtual FReply OnMouseWheel(const FGeometry& MyGeometry, const FPointerEvent& MouseEvent) override;
+	UE_API virtual FReply OnTouchGesture(const FGeometry& MyGeometry, const FPointerEvent& GestureEvent) override;
+	UE_API virtual FReply OnTouchEnded(const FGeometry& MyGeometry, const FPointerEvent& InTouchEvent) override;
+	UE_API virtual int32 OnPaint(const FPaintArgs& Args, const FGeometry& AllottedGeometry, const FSlateRect& MyCullingRect, FSlateWindowElementList& OutDrawElements, int32 LayerId, const FWidgetStyle& InWidgetStyle, bool bParentEnabled) const override;
 	virtual bool SupportsKeyboardFocus() const override { return true; }
 	// End of Swidget interface
 
 	/** Gets the current zoom factor. */
-	float GetZoomAmount() const;
+	UE_API float GetZoomAmount() const;
 
-	FVector2D GraphCoordToPanelCoord(const FVector2D& GraphSpaceCoordinate) const;
-	FVector2D PanelCoordToGraphCoord(const FVector2D& PanelSpaceCoordinate) const;
+	UE_API FVector2D GraphCoordToPanelCoord(const FVector2D& GraphSpaceCoordinate) const;
+	UE_API FVector2D PanelCoordToGraphCoord(const FVector2D& PanelSpaceCoordinate) const;
 
 protected:
-	virtual void OnPaintBackground(const FGeometry& AllottedGeometry, const FSlateRect& MyCullingRect, FSlateWindowElementList& OutDrawElements, int32 LayerId) const;
+	UE_API virtual void OnPaintBackground(const FGeometry& AllottedGeometry, const FSlateRect& MyCullingRect, FSlateWindowElementList& OutDrawElements, int32 LayerId) const;
 
-	void PaintBackgroundAsLines(const FSlateBrush* BackgroundImage, const FGeometry& AllottedGeometry, const FSlateRect& MyCullingRect, FSlateWindowElementList& OutDrawElements, int32& DrawLayerId) const;
+	UE_API void PaintBackgroundAsLines(const FSlateBrush* BackgroundImage, const FGeometry& AllottedGeometry, const FSlateRect& MyCullingRect, FSlateWindowElementList& OutDrawElements, int32& DrawLayerId) const;
 
-	void ChangeZoomLevel(int32 ZoomLevelDelta, const FVector2D& WidgetSpaceZoomOrigin, bool bOverrideZoomLimiting);
+	UE_API void ChangeZoomLevel(int32 ZoomLevelDelta, const FVector2D& WidgetSpaceZoomOrigin, bool bOverrideZoomLimiting);
 	
-	void PostChangedZoom();
+	UE_API void PostChangedZoom();
 
-	bool ScrollToLocation(const FGeometry& MyGeometry, FVector2D DesiredCenterPosition, const float InDeltaTime);
+	UE_API bool ScrollToLocation(const FGeometry& MyGeometry, FVector2D DesiredCenterPosition, const float InDeltaTime);
 
-	bool ZoomToLocation(const FVector2D& CurrentSizeWithoutZoom, const FVector2D& DesiredSize, bool bDoneScrolling);
+	UE_API bool ZoomToLocation(const FVector2D& CurrentSizeWithoutZoom, const FVector2D& DesiredSize, bool bDoneScrolling);
 
-	void ZoomToFit(bool bInstantZoom);
+	UE_API void ZoomToFit(bool bInstantZoom);
 
-	FText GetZoomText() const;
-	FSlateColor GetZoomTextColorAndOpacity() const;
+	UE_API FText GetZoomText() const;
+	UE_API FSlateColor GetZoomTextColorAndOpacity() const;
 
-	FVector2D GetViewOffset() const;
+	UE_API FVector2D GetViewOffset() const;
 
-	FSlateRect ComputeSensibleBounds() const;
+	UE_API FSlateRect ComputeSensibleBounds() const;
 
 protected:
-	virtual FSlateRect ComputeAreaBounds() const;
-	virtual float GetGridScaleAmount() const;
-	virtual int32 GetGraphRulePeriod() const;
+	UE_API virtual FSlateRect ComputeAreaBounds() const;
+	UE_API virtual float GetGridScaleAmount() const;
+	UE_API virtual int32 GetGraphRulePeriod() const;
 	virtual int32 GetSnapGridSize() const = 0;
 
 protected:
@@ -154,7 +156,7 @@ protected:
 
 private:
 	/** Active timer that handles deferred zooming until the target zoom is reached */
-	EActiveTimerReturnType HandleZoomToFit(double InCurrentTime, float InDeltaTime);
+	UE_API EActiveTimerReturnType HandleZoomToFit(double InCurrentTime, float InDeltaTime);
 
 	/** The handle to the active timer */
 	TWeakPtr<FActiveTimerHandle> ActiveTimerHandle;
@@ -162,3 +164,5 @@ private:
 	// A flag noting if we have a pending zoom to extents operation to perform next tick.
 	bool bDeferredZoomToExtents;
 };
+
+#undef UE_API
diff --git a/Engine/Source/Editor/UMGEditor/Public/Designer/SDesignerView.h b/Engine/Source/Editor/UMGEditor/Public/Designer/SDesignerView.h
--- a/Engine/Source/Editor/UMGEditor/Public/Designer/SDesignerView.h
+++ b/Engine/Source/Editor/UMGEditor/Public/Designer/SDesignerView.h
@@ -26,6 +26,8 @@
 #include "Designer/SDesignSurface.h"
 #include "UMGEditorProjectSettings.h"
 
+#define UE_API UMGEDITOR_API
+
 class FMenuBuilder;
 class FScopedTransaction;
 class SBox;
@@ -43,60 +45,60 @@ class UWidgetEditingProjectSettings;
 /**
  * The designer for widgets.  Allows for laying out widgets in a drag and drop environment.
  */
-class UMGEDITOR_API SDesignerView : public SDesignSurface, public FGCObject, public IUMGDesigner
+class SDesignerView : public SDesignSurface, public FGCObject, public IUMGDesigner
 {
 public:
 
 	SLATE_BEGIN_ARGS( SDesignerView ) {}
 	SLATE_END_ARGS()
 
-	void Construct(const FArguments& InArgs, TSharedPtr<class FWidgetBlueprintEditor> InBlueprintEditor);
-	virtual ~SDesignerView();
+	UE_API void Construct(const FArguments& InArgs, TSharedPtr<class FWidgetBlueprintEditor> InBlueprintEditor);
+	UE_API virtual ~SDesignerView();
 
-	EActiveTimerReturnType EnsureTick(double InCurrentTime, float InDeltaTime);
+	UE_API EActiveTimerReturnType EnsureTick(double InCurrentTime, float InDeltaTime);
 
-	TSharedRef<SWidget> CreateOverlayUI();
+	UE_API TSharedRef<SWidget> CreateOverlayUI();
 
 	// SWidget interface
-	virtual FReply OnMouseButtonDown(const FGeometry& MyGeometry, const FPointerEvent& MouseEvent) override;
-	virtual FReply OnMouseButtonUp(const FGeometry& MyGeometry, const FPointerEvent& MouseEvent) override;
-	virtual FReply OnMouseMove(const FGeometry& MyGeometry, const FPointerEvent& MouseEvent) override;
-	virtual void OnMouseEnter(const FGeometry& MyGeometry, const FPointerEvent& MouseEvent) override;
-	virtual void OnMouseLeave(const FPointerEvent& MouseEvent) override;
-
-	virtual FReply OnKeyDown(const FGeometry& MyGeometry, const FKeyEvent& InKeyEvent) override;
-	virtual FReply OnKeyUp(const FGeometry& MyGeometry, const FKeyEvent& InKeyEvent) override;
-
-	virtual void Tick( const FGeometry& AllottedGeometry, const double InCurrentTime, const float InDeltaTime ) override;
-
-	virtual FReply OnDragDetected(const FGeometry& MyGeometry, const FPointerEvent& MouseEvent) override;
-	virtual void OnDragEnter(const FGeometry& MyGeometry, const FDragDropEvent& DragDropEvent) override;
-	virtual void OnDragLeave(const FDragDropEvent& DragDropEvent) override;
-	virtual FReply OnDragOver(const FGeometry& MyGeometry, const FDragDropEvent& DragDropEvent) override;
-	virtual FReply OnDrop(const FGeometry& MyGeometry, const FDragDropEvent& DragDropEvent) override;
+	UE_API virtual FReply OnMouseButtonDown(const FGeometry& MyGeometry, const FPointerEvent& MouseEvent) override;
+	UE_API virtual FReply OnMouseButtonUp(const FGeometry& MyGeometry, const FPointerEvent& MouseEvent) override;
+	UE_API virtual FReply OnMouseMove(const FGeometry& MyGeometry, const FPointerEvent& MouseEvent) override;
+	UE_API virtual void OnMouseEnter(const FGeometry& MyGeometry, const FPointerEvent& MouseEvent) override;
+	UE_API virtual void OnMouseLeave(const FPointerEvent& MouseEvent) override;
+
+	UE_API virtual FReply OnKeyDown(const FGeometry& MyGeometry, const FKeyEvent& InKeyEvent) override;
+	UE_API virtual FReply OnKeyUp(const FGeometry& MyGeometry, const FKeyEvent& InKeyEvent) override;
+
+	UE_API virtual void Tick( const FGeometry& AllottedGeometry, const double InCurrentTime, const float InDeltaTime ) override;
+
+	UE_API virtual FReply OnDragDetected(const FGeometry& MyGeometry, const FPointerEvent& MouseEvent) override;
+	UE_API virtual void OnDragEnter(const FGeometry& MyGeometry, const FDragDropEvent& DragDropEvent) override;
+	UE_API virtual void OnDragLeave(const FDragDropEvent& DragDropEvent) override;
+	UE_API virtual FReply OnDragOver(const FGeometry& MyGeometry, const FDragDropEvent& DragDropEvent) override;
+	UE_API virtual FReply OnDrop(const FGeometry& MyGeometry, const FDragDropEvent& DragDropEvent) override;
 	// End of SWidget interface
 
-	void Register(TSharedRef<FDesignerExtension> Extension);
-	void Unregister(TSharedRef<FDesignerExtension> Extension);
+	UE_API void Register(TSharedRef<FDesignerExtension> Extension);
+	UE_API void Unregister(TSharedRef<FDesignerExtension> Extension);
 
 	// IUMGDesigner interface
-	virtual float GetPreviewScale() const override;
-	virtual const TSet<FWidgetReference>& GetSelectedWidgets() const override;
-	virtual FWidgetReference GetSelectedWidget() const override;
-	virtual ETransformMode::Type GetTransformMode() const override;
-	virtual FGeometry GetDesignerGeometry() const override;
-	virtual FVector2D GetWidgetOriginAbsolute() const;
-	virtual bool GetWidgetGeometry(const FWidgetReference& Widget, FGeometry& Geometry) const override;
-	virtual bool GetWidgetGeometry(const UWidget* PreviewWidget, FGeometry& Geometry) const override;
-	virtual bool GetWidgetParentGeometry(const FWidgetReference& Widget, FGeometry& Geometry) const override;
-	virtual FGeometry MakeGeometryWindowLocal(const FGeometry& WidgetGeometry) const override;
-	virtual void MarkDesignModifed(bool bRequiresRecompile) override;
-	virtual void PushDesignerMessage(const FText& Message) override;
-	virtual void PopDesignerMessage() override;
+	UE_API virtual float GetPreviewScale() const override;
+	UE_API virtual const TSet<FWidgetReference>& GetSelectedWidgets() const override;
+	UE_API virtual FWidgetReference GetSelectedWidget() const override;
+	UE_API virtual ETransformMode::Type GetTransformMode() const override;
+	UE_API virtual FGeometry GetDesignerGeometry() const override;
+	UE_API virtual FVector2D GetWidgetOriginAbsolute() const;
+	UE_API virtual bool GetWidgetGeometry(const FWidgetReference& Widget, FGeometry& Geometry) const override;
+	UE_API virtual bool GetWidgetGeometry(const UWidget* PreviewWidget, FGeometry& Geometry) const override;
+	UE_API virtual bool GetWidgetParentGeometry(const FWidgetReference& Widget, FGeometry& Geometry) const override;
+	UE_API virtual FGeometry MakeGeometryWindowLocal(const FGeometry& WidgetGeometry) const override;
+	UE_API virtual void MarkDesignModifed(bool bRequiresRecompile) override;
+	UE_API virtual void PushDesignerMessage(const FText& Message) override;
+	UE_API virtual void PopDesignerMessage() override;
 	// End of IUMGDesigner interface
 
 	// FGCObject interface
-	virtual void AddReferencedObjects(FReferenceCollector& Collector) override;
+	UE_API virtual void AddReferencedObjects(FReferenceCollector& Collector) override;
 	virtual FString GetReferencerName() const override
 	{
 		return TEXT("SDesignerView");
@@ -106,143 +108,143 @@ public:
 public:
 
 	/** The width of the preview screen for the UI */
-	FOptionalSize GetPreviewAreaWidth() const;
+	UE_API FOptionalSize GetPreviewAreaWidth() const;
 
 	/** The height of the preview screen for the UI */
-	FOptionalSize GetPreviewAreaHeight() const;
+	UE_API FOptionalSize GetPreviewAreaHeight() const;
 
 	/** The width of the preview widget for the UI */
-	FOptionalSize GetPreviewSizeWidth() const;
+	UE_API FOptionalSize GetPreviewSizeWidth() const;
 
 	/** The height of the preview widget for the UI */
-	FOptionalSize GetPreviewSizeHeight() const;
+	UE_API FOptionalSize GetPreviewSizeHeight() const;
 
 	/** Gets the DPI scale that would be applied given the current preview width and height */
-	float GetPreviewDPIScale() const;
+	UE_API float GetPreviewDPIScale() const;
 
 	/** Set the size of the preview screen for the UI */
-	void SetPreviewAreaSize(int32 Width, int32 Height);
+	UE_API void SetPreviewAreaSize(int32 Width, int32 Height);
 
-	void BeginResizingArea();
-	void EndResizingArea();
+	UE_API void BeginResizingArea();
+	UE_API void EndResizingArea();
 	
-	const UWidgetEditingProjectSettings* GetRelevantSettings() const;
+	UE_API const UWidgetEditingProjectSettings* GetRelevantSettings() const;
 
 protected:
-	virtual void OnPaintBackground(const FGeometry& AllottedGeometry, const FSlateRect& MyCullingRect, FSlateWindowElementList& OutDrawElements, int32 LayerId) const override;
-	void DrawResolution(const FDebugResolution& Resolution, const FGeometry& AllottedGeometry, const FSlateRect& MyCullingRect, FSlateWindowElementList& OutDrawElements, int32 LayerId) const;
+	UE_API virtual void OnPaintBackground(const FGeometry& AllottedGeometry, const FSlateRect& MyCullingRect, FSlateWindowElementList& OutDrawElements, int32 LayerId) const override;
+	UE_API void DrawResolution(const FDebugResolution& Resolution, const FGeometry& AllottedGeometry, const FSlateRect& MyCullingRect, FSlateWindowElementList& OutDrawElements, int32 LayerId) const;
 
-	virtual FSlateRect ComputeAreaBounds() const override;
-	virtual int32 GetGraphRulePeriod() const override;
-	virtual float GetGridScaleAmount() const override;
-	virtual int32 GetSnapGridSize() const override;
+	UE_API virtual FSlateRect ComputeAreaBounds() const override;
+	UE_API virtual int32 GetGraphRulePeriod() const override;
+	UE_API virtual float GetGridScaleAmount() const override;
+	UE_API virtual int32 GetSnapGridSize() const override;
 
 private:
-	void RegisterExtensions();
+	UE_API void RegisterExtensions();
 
 	/** Establishes the resolution and aspect ratio to use on construction from config settings */
-	void SetStartupResolution();
+	UE_API void SetStartupResolution();
 
-	FVector2D GetAreaResizeHandlePosition() const;
-	EVisibility GetAreaResizeHandleVisibility() const;
+	UE_API FVector2D GetAreaResizeHandlePosition() const;
+	UE_API EVisibility GetAreaResizeHandleVisibility() const;
 
-	const FSlateBrush* GetPreviewBackground() const;
+	UE_API const FSlateBrush* GetPreviewBackground() const;
 
 	/** Adds any pending selected widgets to the selection set */
-	void ResolvePendingSelectedWidgets();
+	UE_API void ResolvePendingSelectedWidgets();
 
 	/** Updates the designer to display the latest preview widget */
-	void UpdatePreviewWidget(bool bForceUpdate);
+	UE_API void UpdatePreviewWidget(bool bForceUpdate);
 
-	void BroadcastDesignerChanged();
+	UE_API void BroadcastDesignerChanged();
 
-	void ClearExtensionWidgets();
-	void CreateExtensionWidgetsForSelection();
+	UE_API void ClearExtensionWidgets();
+	UE_API void CreateExtensionWidgetsForSelection();
 
-	EVisibility GetInfoBarVisibility() const;
-	FText GetInfoBarText() const;
+	UE_API EVisibility GetInfoBarVisibility() const;
+	UE_API FText GetInfoBarText() const;
 
 	/** Displays the context menu when you right click */
-	void ShowContextMenu(const FGeometry& MyGeometry, const FPointerEvent& MouseEvent);
+	UE_API void ShowContextMenu(const FGeometry& MyGeometry, const FPointerEvent& MouseEvent);
 
-	void OnEditorSelectionChanged();
+	UE_API void OnEditorSelectionChanged();
 
 	/** Called when a new widget is being hovered */
-	void OnHoveredWidgetSet(const FWidgetReference& InHoveredWidget);
+	UE_API void OnHoveredWidgetSet(const FWidgetReference& InHoveredWidget);
 
 	/** Called when a widget is no longer being hovered */
-	void OnHoveredWidgetCleared();
+	UE_API void OnHoveredWidgetCleared();
 
 	/** Gets the blueprint being edited by the designer */
-	UWidgetBlueprint* GetBlueprint() const;
+	UE_API UWidgetBlueprint* GetBlueprint() const;
 
 	/** Called whenever the blueprint is recompiled or the preview is updated */
-	void OnPreviewNeedsRecreation();
+	UE_API void OnPreviewNeedsRecreation();
 
-	void PopulateWidgetGeometryCache(FArrangedWidget& Root);
-	void PopulateWidgetGeometryCache_Loop(FArrangedWidget& Parent);
+	UE_API void PopulateWidgetGeometryCache(FArrangedWidget& Root);
+	UE_API void PopulateWidgetGeometryCache_Loop(FArrangedWidget& Parent);
 
 	/** @return Formatted text for the given resolution params */
-	FText GetResolutionText(int32 Width, int32 Height, const FString& AspectRatio) const;
+	UE_API FText GetResolutionText(int32 Width, int32 Height, const FString& AspectRatio) const;
 
 	/** Move the selected widget by the nudge amount. */
-	FReply NudgeSelectedWidget(FVector2D Nudge);
-
-	FText GetCurrentResolutionText() const;
-	FText GetCurrentDPIScaleText() const;
-	FSlateColor GetCurrentDPIScaleColor() const;
-	FText GetCurrentScaleFactorText() const;
-	FText GetCurrentSafeZoneText() const;
-	FSlateColor GetResolutionTextColorAndOpacity() const;
-	EVisibility GetResolutionTextVisibility() const;
-
-	TOptional<int32> GetCustomResolutionWidth() const;
-	TOptional<int32> GetCustomResolutionHeight() const;
-	void OnCustomResolutionWidthChanged(int32 InValue);
-	void OnCustomResolutionHeightChanged(int32 InValue);
-	EVisibility GetCustomResolutionEntryVisibility() const;
+	UE_API FReply NudgeSelectedWidget(FVector2D Nudge);
+
+	UE_API FText GetCurrentResolutionText() const;
+	UE_API FText GetCurrentDPIScaleText() const;
+	UE_API FSlateColor GetCurrentDPIScaleColor() const;
+	UE_API FText GetCurrentScaleFactorText() const;
+	UE_API FText GetCurrentSafeZoneText() const;
+	UE_API FSlateColor GetResolutionTextColorAndOpacity() const;
+	UE_API EVisibility GetResolutionTextVisibility() const;
+
+	UE_API TOptional<int32> GetCustomResolutionWidth() const;
+	UE_API TOptional<int32> GetCustomResolutionHeight() const;
+	UE_API void OnCustomResolutionWidthChanged(int32 InValue);
+	UE_API void OnCustomResolutionHeightChanged(int32 InValue);
+	UE_API EVisibility GetCustomResolutionEntryVisibility() const;
 	
 	// Handles selecting a common screen resolution.
-	void HandleOnCommonResolutionSelected(const FPlayScreenResolution InResolution);
-	bool HandleIsCommonResolutionSelected(const FPlayScreenResolution InResolution) const;
-	FUIAction GetResolutionMenuAction( const FPlayScreenResolution& ScreenResolution );
-	TSharedRef<SWidget> GetResolutionsMenu();
-
-	TSharedRef<SWidget> GetScreenSizingFillMenu();
-	void CreateScreenFillEntry(FMenuBuilder& MenuBuilder, EDesignPreviewSizeMode SizeMode);
-	FText GetScreenSizingFillText() const;
-	bool GetIsScreenFillRuleSelected(EDesignPreviewSizeMode SizeMode) const;
-	void OnScreenFillRuleSelected(EDesignPreviewSizeMode SizeMode);
-
-	const FSlateBrush* GetAspectRatioSwitchImage() const;
-	bool GetAspectRatioSwitchEnabled() const;
-	bool GetFlipDeviceEnabled() const;
-	EVisibility GetDesignerOutlineVisibility() const;
-	FSlateColor GetDesignerOutlineColor() const;
-	FText GetDesignerOutlineText() const;
-
-	FText GetCursorPositionText() const;
-	EVisibility GetCursorPositionTextVisibility() const;
+	UE_API void HandleOnCommonResolutionSelected(const FPlayScreenResolution InResolution);
+	UE_API bool HandleIsCommonResolutionSelected(const FPlayScreenResolution InResolution) const;
+	UE_API FUIAction GetResolutionMenuAction( const FPlayScreenResolution& ScreenResolution );
+	UE_API TSharedRef<SWidget> GetResolutionsMenu();
+
+	UE_API TSharedRef<SWidget> GetScreenSizingFillMenu();
+	UE_API void CreateScreenFillEntry(FMenuBuilder& MenuBuilder, EDesignPreviewSizeMode SizeMode);
+	UE_API FText GetScreenSizingFillText() const;
+	UE_API bool GetIsScreenFillRuleSelected(EDesignPreviewSizeMode SizeMode) const;
+	UE_API void OnScreenFillRuleSelected(EDesignPreviewSizeMode SizeMode);
+
+	UE_API const FSlateBrush* GetAspectRatioSwitchImage() const;
+	UE_API bool GetAspectRatioSwitchEnabled() const;
+	UE_API bool GetFlipDeviceEnabled() const;
+	UE_API EVisibility GetDesignerOutlineVisibility() const;
+	UE_API FSlateColor GetDesignerOutlineColor() const;
+	UE_API FText GetDesignerOutlineText() const;
+
+	UE_API FText GetCursorPositionText() const;
+	UE_API EVisibility GetCursorPositionTextVisibility() const;
 	
-	FText GetSelectedWidgetDimensionsText() const;
-	EVisibility GetSelectedWidgetDimensionsVisibility() const;
+	UE_API FText GetSelectedWidgetDimensionsText() const;
+	UE_API EVisibility GetSelectedWidgetDimensionsVisibility() const;
 
 	// Handles drawing selection and other effects a SPaintSurface widget injected into the hierarchy.
-	int32 HandleEffectsPainting(const FOnPaintHandlerParams& PaintArgs);
-	void DrawSelectionAndHoverOutline(const FOnPaintHandlerParams& PaintArgs);
-	void DrawSafeZone(const FOnPaintHandlerParams& PaintArgs);
+	UE_API int32 HandleEffectsPainting(const FOnPaintHandlerParams& PaintArgs);
+	UE_API void DrawSelectionAndHoverOutline(const FOnPaintHandlerParams& PaintArgs);
+	UE_API void DrawSafeZone(const FOnPaintHandlerParams& PaintArgs);
 
-	FReply HandleDPISettingsClicked();
+	UE_API FReply HandleDPISettingsClicked();
 
-	UUserWidget* GetDefaultWidget() const;
+	UE_API UUserWidget* GetDefaultWidget() const;
 
-	void BeginTransaction(const FText& SessionName);
-	bool InTransaction() const;
-	void EndTransaction(bool bCancel);
+	UE_API void BeginTransaction(const FText& SessionName);
+	UE_API bool InTransaction() const;
+	UE_API void EndTransaction(bool bCancel);
 
-	UWidget* GetWidgetInDesignScopeFromSlateWidget(TSharedRef<SWidget>& InWidget);
+	UE_API UWidget* GetWidgetInDesignScopeFromSlateWidget(TSharedRef<SWidget>& InWidget);
 
-	EVisibility GetExtensionCanvasVisibility() const;
+	UE_API EVisibility GetExtensionCanvasVisibility() const;
 
 private:
 	struct FWidgetHitResult
@@ -258,56 +260,56 @@ private:
 	};
 
 	/** @returns Gets the widget under the cursor based on a mouse pointer event. */
-	bool FindWidgetUnderCursor(const FGeometry& MyGeometry, const FPointerEvent& MouseEvent, TSubclassOf<UWidget> FindType, FWidgetHitResult& HitResult);
+	UE_API bool FindWidgetUnderCursor(const FGeometry& MyGeometry, const FPointerEvent& MouseEvent, TSubclassOf<UWidget> FindType, FWidgetHitResult& HitResult);
 
 private:
-	FReply HandleZoomToFitClicked();
-	FReply HandleSwapAspectRatioClicked();
-	FReply HandleFlipSafeZonesClicked();
-	EVisibility GetRulerVisibility() const;
+	UE_API FReply HandleZoomToFitClicked();
+	UE_API FReply HandleSwapAspectRatioClicked();
+	UE_API FReply HandleFlipSafeZonesClicked();
+	UE_API EVisibility GetRulerVisibility() const;
 
 private:
-	static const FString ConfigSectionName;
-	static const FString DefaultPreviewOverrideName;
+	static UE_API const FString ConfigSectionName;
+	static UE_API const FString DefaultPreviewOverrideName;
 
 	/** Extensions for the designer to allow for custom widgets to be inserted onto the design surface as selection changes. */
 	TArray< TSharedRef<FDesignerExtension> > DesignerExtensions;
 
 private:
-	void BindCommands();
+	UE_API void BindCommands();
 
-	void SetTransformMode(ETransformMode::Type InTransformMode);
-	bool CanSetTransformMode(ETransformMode::Type InTransformMode) const;
-	bool IsTransformModeActive(ETransformMode::Type InTransformMode) const;
+	UE_API void SetTransformMode(ETransformMode::Type InTransformMode);
+	UE_API bool CanSetTransformMode(ETransformMode::Type InTransformMode) const;
+	UE_API bool IsTransformModeActive(ETransformMode::Type InTransformMode) const;
 
-	void ToggleShowingOutlines();
-	bool IsShowingOutlines() const;
+	UE_API void ToggleShowingOutlines();
+	UE_API bool IsShowingOutlines() const;
 
-	void ToggleRespectingLocks();
-	bool IsRespectingLocks() const;
+	UE_API void ToggleRespectingLocks();
+	UE_API bool IsRespectingLocks() const;
 
-	void ProcessDropAndAddWidget(const FGeometry& MyGeometry, const FDragDropEvent& DragDropEvent, const bool bIsPreview);
-	void MoveWidgets(const FGeometry& MyGeometry, const FDragDropEvent& DragDropEvent, const bool bIsPreview, UWidget* Target, const bool bAnyWidgetChangingParent);
+	UE_API void ProcessDropAndAddWidget(const FGeometry& MyGeometry, const FDragDropEvent& DragDropEvent, const bool bIsPreview);
+	UE_API void MoveWidgets(const FGeometry& MyGeometry, const FDragDropEvent& DragDropEvent, const bool bIsPreview, UWidget* Target, const bool bAnyWidgetChangingParent);
 
-	FVector2D GetExtensionPosition(TSharedRef<FDesignerSurfaceElement> ExtensionElement) const;
+	UE_API FVector2D GetExtensionPosition(TSharedRef<FDesignerSurfaceElement> ExtensionElement) const;
 
-	FVector2D GetExtensionSize(TSharedRef<FDesignerSurfaceElement> ExtensionElement) const;
+	UE_API FVector2D GetExtensionSize(TSharedRef<FDesignerSurfaceElement> ExtensionElement) const;
 	
-	void ClearDropPreviews();
+	UE_API void ClearDropPreviews();
 
-	void DetermineDragDropPreviewWidgets(TArray<UWidget*>& OutWidgets, const FDragDropEvent& DragDropEvent, UWidgetTree* RootWidgetTree);
+	UE_API void DetermineDragDropPreviewWidgets(TArray<UWidget*>& OutWidgets, const FDragDropEvent& DragDropEvent, UWidgetTree* RootWidgetTree);
 
-	void SwapSafeZoneTypes();
+	UE_API void SwapSafeZoneTypes();
 
-	bool IsSelectableInSequencer(UWidget* const InWidget) const;
+	UE_API bool IsSelectableInSequencer(UWidget* const InWidget) const;
 
-	void OnSelectedAnimationChanged();
+	UE_API void OnSelectedAnimationChanged();
 
-	void OnSelectionLimitedChanged(const bool bInEnabled);
+	UE_API void OnSelectionLimitedChanged(const bool bInEnabled);
 
-	void DeselectNonSequencerWidgets();
+	UE_API void DeselectNonSequencerWidgets();
 
-	EVisibility GetSelectionLimitedTextVisibility() const;
+	UE_API EVisibility GetSelectionLimitedTextVisibility() const;
 
 private:
 	/** A reference to the BP Editor that owns this designer */
@@ -437,3 +439,5 @@ private:
 
 	FMargin DesignerSafeZoneOverride;
 };
+
+#undef UE_API
diff --git a/Engine/Source/Editor/UMGEditor/Public/DesignerExtension.h b/Engine/Source/Editor/UMGEditor/Public/DesignerExtension.h
--- a/Engine/Source/Editor/UMGEditor/Public/DesignerExtension.h
+++ b/Engine/Source/Editor/UMGEditor/Public/DesignerExtension.h
@@ -14,6 +14,8 @@
 #include "WidgetReference.h"
 #include "Widgets/SWidget.h"
 
+#define UE_API UMGEDITOR_API
+
 class FScopedTransaction;
 class FSlateRect;
 class FSlateWindowElementList;
@@ -50,7 +52,7 @@ enum class EExtensionLayoutLocation : uint8
 /**
  * The basic element returned for extending the design surface.
  */
-class UMGEDITOR_API FDesignerSurfaceElement : public TSharedFromThis < FDesignerSurfaceElement >
+class FDesignerSurfaceElement : public TSharedFromThis < FDesignerSurfaceElement >
 {
 public:
 	FDesignerSurfaceElement(TSharedRef<SWidget> InWidget, EExtensionLayoutLocation InLocation, TAttribute<FVector2D> InOffset = FVector2D(0, 0), TAttribute<FVector2D> InAlignment = FVector2D(0, 0))
@@ -111,18 +113,18 @@ protected:
  * The Designer extension allows developers to provide additional widgets and custom painting to the designer surface for
  * specific widgets.  Allowing for a more customized and specific editors for the different widgets.
  */
-class UMGEDITOR_API FDesignerExtension : public TSharedFromThis<FDesignerExtension>
+class FDesignerExtension : public TSharedFromThis<FDesignerExtension>
 {
 public:
 	/** Constructor */
-	FDesignerExtension();
-	virtual ~FDesignerExtension();
+	UE_API FDesignerExtension();
+	UE_API virtual ~FDesignerExtension();
 
 	/** Initializes the designer extension, this is called the first time a designer extension is registered */
-	virtual void Initialize(IUMGDesigner* InDesigner, UWidgetBlueprint* InBlueprint);
+	UE_API virtual void Initialize(IUMGDesigner* InDesigner, UWidgetBlueprint* InBlueprint);
 
 	/** Uninitializes the designer extension, this is called when the designer is removed. */
-	virtual void Uninitialize();
+	UE_API virtual void Uninitialize();
 
 	/** Returns true if the designer extension can extend the current selection. */
 	virtual bool CanExtendSelection(const TArray< FWidgetReference >& Selection) const
@@ -156,7 +158,7 @@ public:
 	}
 
 	/** Gets the ID identifying this extension. */
-	FName GetExtensionId() const;
+	UE_API FName GetExtensionId() const;
 
 	bool operator==(const FDesignerExtension& Other) const
 	{
@@ -164,9 +166,9 @@ public:
 	}
 
 protected:
-	void BeginTransaction(const FText& SessionName);
+	UE_API void BeginTransaction(const FText& SessionName);
 
-	void EndTransaction();
+	UE_API void EndTransaction();
 
 protected:
 	FName ExtensionId;
@@ -178,3 +180,5 @@ protected:
 private:
 	FScopedTransaction* ScopedTransaction;
 };
+
+#undef UE_API
diff --git a/Engine/Source/Editor/UMGEditor/Public/DragDrop/WidgetTemplateDragDropOp.h b/Engine/Source/Editor/UMGEditor/Public/DragDrop/WidgetTemplateDragDropOp.h
--- a/Engine/Source/Editor/UMGEditor/Public/DragDrop/WidgetTemplateDragDropOp.h
+++ b/Engine/Source/Editor/UMGEditor/Public/DragDrop/WidgetTemplateDragDropOp.h
@@ -8,13 +8,15 @@
 #include "Input/DragAndDrop.h"
 #include "Templates/SharedPointer.h"
 
+#define UE_API UMGEDITOR_API
+
 class FWidgetTemplate;
 
 /**
  * This drag drop operation allows widget templates from the palate to be dragged and dropped into the designer
  * or the widget hierarchy in order to spawn new widgets.
  */
-class UMGEDITOR_API FWidgetTemplateDragDropOp : public FDecoratedDragDropOp
+class FWidgetTemplateDragDropOp : public FDecoratedDragDropOp
 {
 public:
 	DRAG_DROP_OPERATOR_TYPE(FWidgetTemplateDragDropOp, FDecoratedDragDropOp)
@@ -23,5 +25,7 @@ public:
 	TSharedPtr<FWidgetTemplate> Template;
 
 	/** Constructs the drag drop operation */
-	static TSharedRef<FWidgetTemplateDragDropOp> New(const TSharedPtr<FWidgetTemplate>& InTemplate);
+	static UE_API TSharedRef<FWidgetTemplateDragDropOp> New(const TSharedPtr<FWidgetTemplate>& InTemplate);
 };
+
+#undef UE_API
diff --git a/Engine/Source/Editor/UMGEditor/Public/FieldNotification/CustomizationHelper.h b/Engine/Source/Editor/UMGEditor/Public/FieldNotification/CustomizationHelper.h
--- a/Engine/Source/Editor/UMGEditor/Public/FieldNotification/CustomizationHelper.h
+++ b/Engine/Source/Editor/UMGEditor/Public/FieldNotification/CustomizationHelper.h
@@ -8,6 +8,8 @@
 #include "UObject/WeakFieldPtr.h"
 #include "WidgetBlueprint.h"
 
+#define UE_API UMGEDITOR_API
+
 class IDetailLayoutBuilder;
 class IBlueprintEditor;
 class UWidgetBlueprint;
@@ -16,25 +18,25 @@ class UK2Node_FunctionEntry;
 namespace UE::FieldNotification
 {
 
-class UMGEDITOR_API FCustomizationHelper
+class FCustomizationHelper
 {
 public:
-	static const FName MetaData_FieldNotify;
+	static UE_API const FName MetaData_FieldNotify;
 
 	FCustomizationHelper(UWidgetBlueprint* InBlueprint)
 		: Blueprint(InBlueprint)
 	{}
 
-	void CustomizeVariableDetails(IDetailLayoutBuilder& DetailLayout);
-	void CustomizeFunctionDetails(IDetailLayoutBuilder& DetailLayout);
+	UE_API void CustomizeVariableDetails(IDetailLayoutBuilder& DetailLayout);
+	UE_API void CustomizeFunctionDetails(IDetailLayoutBuilder& DetailLayout);
 
 private:
-	ECheckBoxState IsPropertyFieldNotifyChecked() const;
-	void HandlePropertyFieldNotifyCheckStateChanged(ECheckBoxState CheckBoxState);
-	ECheckBoxState IsFunctionFieldNotifyChecked() const;
-	void HandleFunctionFieldNotifyCheckStateChanged(ECheckBoxState CheckBoxState);
+	UE_API ECheckBoxState IsPropertyFieldNotifyChecked() const;
+	UE_API void HandlePropertyFieldNotifyCheckStateChanged(ECheckBoxState CheckBoxState);
+	UE_API ECheckBoxState IsFunctionFieldNotifyChecked() const;
+	UE_API void HandleFunctionFieldNotifyCheckStateChanged(ECheckBoxState CheckBoxState);
 
-	UK2Node_FunctionEntry* FindFunctionEntry(UObject* Obj) const;
+	UE_API UK2Node_FunctionEntry* FindFunctionEntry(UObject* Obj) const;
 
 private:
 	/** The blueprint we are editing */
@@ -47,4 +49,6 @@ private:
 	TArray<TWeakObjectPtr<UK2Node_FunctionEntry>> FunctionsBeingCustomized;
 };
 
-} //namespace
\ No newline at end of file
+} //namespace
+
+#undef UE_API
diff --git a/Engine/Source/Editor/UMGEditor/Public/Hierarchy/SReadOnlyHierarchyView.h b/Engine/Source/Editor/UMGEditor/Public/Hierarchy/SReadOnlyHierarchyView.h
--- a/Engine/Source/Editor/UMGEditor/Public/Hierarchy/SReadOnlyHierarchyView.h
+++ b/Engine/Source/Editor/UMGEditor/Public/Hierarchy/SReadOnlyHierarchyView.h
@@ -12,6 +12,8 @@
 
 #include "SReadOnlyHierarchyView.generated.h"
 
+#define UE_API UMGEDITOR_API
+
 class ITableRow;
 class SSearchBox;
 class STableViewBase;
@@ -33,7 +35,7 @@ enum class ERootSelectionMode : uint8
 	Self
 };
 
-class UMGEDITOR_API SReadOnlyHierarchyView : public SCompoundWidget
+class SReadOnlyHierarchyView : public SCompoundWidget
 {
 private:
 
@@ -61,37 +63,37 @@ public:
 		SLATE_ARGUMENT_DEFAULT(bool, ExpandAll) = true;
 	SLATE_END_ARGS()
 
-	virtual ~SReadOnlyHierarchyView();
+	UE_API virtual ~SReadOnlyHierarchyView();
 
-	void Construct(const FArguments& InArgs, const UWidgetBlueprint* WidgetBlueprint);
-	void Refresh();
+	UE_API void Construct(const FArguments& InArgs, const UWidgetBlueprint* WidgetBlueprint);
+	UE_API void Refresh();
 
-	void SetSelectedWidget(FName WidgetName);
-	void SetRawFilterText(const FText& Text);
+	UE_API void SetSelectedWidget(FName WidgetName);
+	UE_API void SetRawFilterText(const FText& Text);
 
-	TArray<FName> GetSelectedWidgets() const;
-	void ClearSelection();
+	UE_API TArray<FName> GetSelectedWidgets() const;
+	UE_API void ClearSelection();
 
 private:
 
-	TSharedRef<ITableRow> GenerateRow(TSharedPtr<FItem> Item, const TSharedRef<STableViewBase>& OwnerTable) const;
+	UE_API TSharedRef<ITableRow> GenerateRow(TSharedPtr<FItem> Item, const TSharedRef<STableViewBase>& OwnerTable) const;
 
-	void GetItemChildren(TSharedPtr<FItem> Item, TArray<TSharedPtr<FItem>>& OutChildren) const;
-	FText GetItemText(TSharedPtr<FItem> Item) const;
-	const FSlateBrush* GetIconBrush(TSharedPtr<FItem> Item) const;
-	FText GetIconToolTipText(TSharedPtr<FItem> Item) const;
-	FText GetWidgetToolTipText(TSharedPtr<FItem> Item) const;
-	void GetFilterStringsForItem(TSharedPtr<FItem> Item, TArray<FString>& OutStrings) const;
-	TSharedPtr<FItem> FindItem(const TArray<TSharedPtr<FItem>>& RootItems, FName Name) const;
-	void OnSelectionChanged(TSharedPtr<FItem> Selected, ESelectInfo::Type SelectionType);
+	UE_API void GetItemChildren(TSharedPtr<FItem> Item, TArray<TSharedPtr<FItem>>& OutChildren) const;
+	UE_API FText GetItemText(TSharedPtr<FItem> Item) const;
+	UE_API const FSlateBrush* GetIconBrush(TSharedPtr<FItem> Item) const;
+	UE_API FText GetIconToolTipText(TSharedPtr<FItem> Item) const;
+	UE_API FText GetWidgetToolTipText(TSharedPtr<FItem> Item) const;
+	UE_API void GetFilterStringsForItem(TSharedPtr<FItem> Item, TArray<FString>& OutStrings) const;
+	UE_API TSharedPtr<FItem> FindItem(const TArray<TSharedPtr<FItem>>& RootItems, FName Name) const;
+	UE_API void OnSelectionChanged(TSharedPtr<FItem> Selected, ESelectInfo::Type SelectionType);
 
-	void RebuildTree();
-	void BuildWidgetChildren(const UWidget* Widget, TSharedPtr<FItem> Parent);
+	UE_API void RebuildTree();
+	UE_API void BuildWidgetChildren(const UWidget* Widget, TSharedPtr<FItem> Parent);
 
-	void SetItemExpansionRecursive(TSharedPtr<FItem> Item, bool bShouldBeExpanded);
-	void ExpandAll();
+	UE_API void SetItemExpansionRecursive(TSharedPtr<FItem> Item, bool bShouldBeExpanded);
+	UE_API void ExpandAll();
 
-	FName GetItemName(const TSharedPtr<FItem>& Item) const;
+	UE_API FName GetItemName(const TSharedPtr<FItem>& Item) const;
 
 private:
 
@@ -116,3 +118,5 @@ private:
 
 	bool bExpandAll;
 };
+
+#undef UE_API
diff --git a/Engine/Source/Editor/UMGEditor/Public/IHasClipboardExtensibility.h b/Engine/Source/Editor/UMGEditor/Public/IHasClipboardExtensibility.h
--- a/Engine/Source/Editor/UMGEditor/Public/IHasClipboardExtensibility.h
+++ b/Engine/Source/Editor/UMGEditor/Public/IHasClipboardExtensibility.h
@@ -12,7 +12,7 @@ class UExporter;
 class UWidget;
 class UWidgetBlueprint;
 
-class UMGEDITOR_API IClipboardExtension
+class IClipboardExtension
 {
 public:
 	struct FExportArgs
@@ -40,7 +40,7 @@ public:
 /**
  * Clipboard extensibility manager holds a list of registered clipboard extensions.
  */
-class UMGEDITOR_API FClipboardExtensibilityManager
+class FClipboardExtensibilityManager
 {
 public:
 	void AddExtension(const TSharedRef<IClipboardExtension>& Extension)
diff --git a/Engine/Source/Editor/UMGEditor/Public/IHasDesignerExtensibility.h b/Engine/Source/Editor/UMGEditor/Public/IHasDesignerExtensibility.h
--- a/Engine/Source/Editor/UMGEditor/Public/IHasDesignerExtensibility.h
+++ b/Engine/Source/Editor/UMGEditor/Public/IHasDesignerExtensibility.h
@@ -8,7 +8,7 @@
 /**
  * Factory that creates a unique DesignerExtension when a UMG designer is created.
  */
-class UMGEDITOR_API IDesignerExtensionFactory
+class IDesignerExtensionFactory
 {
 public:
 	virtual ~IDesignerExtensionFactory() { }
@@ -28,7 +28,7 @@ public:
 /**
  * Designer Extensibility Manager keep a series of Designer Extensions. See FDesignerExtension class for more information.
  */
-class UMGEDITOR_API FDesignerExtensibilityManager
+class FDesignerExtensibilityManager
 {
 public:
 	UE_DEPRECATED(4.26, "AddDesignerExtension is deprecated, use the IDesignerExtensibilityFactory instead.")
diff --git a/Engine/Source/Editor/UMGEditor/Public/IHasPropertyBindingExtensibility.h b/Engine/Source/Editor/UMGEditor/Public/IHasPropertyBindingExtensibility.h
--- a/Engine/Source/Editor/UMGEditor/Public/IHasPropertyBindingExtensibility.h
+++ b/Engine/Source/Editor/UMGEditor/Public/IHasPropertyBindingExtensibility.h
@@ -13,7 +13,7 @@ class UWidget;
 class UWidgetBlueprint;
 struct FSlateBrush;
 
-class UMGEDITOR_API IPropertyBindingExtension
+class IPropertyBindingExtension
 {
 public:
 	enum class EDropResult
@@ -49,7 +49,7 @@ public:
 /**
  * Bindings menu extensibility manager holds a list of registered bindings menu extensions.
  */
-class UMGEDITOR_API FPropertyBindingExtensibilityManager
+class FPropertyBindingExtensibilityManager
 {
 public:
 	void AddExtension(const TSharedRef<IPropertyBindingExtension>& Extension)
diff --git a/Engine/Source/Editor/UMGEditor/Public/IHasWidgetContextMenuExtensibility.h b/Engine/Source/Editor/UMGEditor/Public/IHasWidgetContextMenuExtensibility.h
--- a/Engine/Source/Editor/UMGEditor/Public/IHasWidgetContextMenuExtensibility.h
+++ b/Engine/Source/Editor/UMGEditor/Public/IHasWidgetContextMenuExtensibility.h
@@ -7,7 +7,7 @@
 class FMenuBuilder;
 class FWidgetBlueprintEditor;
 
-class UMGEDITOR_API IWidgetContextMenuExtension
+class IWidgetContextMenuExtension
 {
 public:
 	virtual ~IWidgetContextMenuExtension() { }
@@ -18,7 +18,7 @@ public:
 /**
  * Widget Context Menu extensibility manager holds a list of registered widget context menu extensions.
  */
-class UMGEDITOR_API FWidgetContextMenuExtensibilityManager
+class FWidgetContextMenuExtensibilityManager
 {
 public:
 	void AddExtension(const TSharedRef<IWidgetContextMenuExtension>& Extension)
diff --git a/Engine/Source/Editor/UMGEditor/Public/IHasWidgetDragDropExtensibility.h b/Engine/Source/Editor/UMGEditor/Public/IHasWidgetDragDropExtensibility.h
--- a/Engine/Source/Editor/UMGEditor/Public/IHasWidgetDragDropExtensibility.h
+++ b/Engine/Source/Editor/UMGEditor/Public/IHasWidgetDragDropExtensibility.h
@@ -7,7 +7,7 @@
 class FDragDropOperation;
 class UWidget;
 
-class UMGEDITOR_API IWidgetDragDropExtension
+class IWidgetDragDropExtension
 {
 public:
 	virtual ~IWidgetDragDropExtension() { }
@@ -19,7 +19,7 @@ public:
 /**
  * Drag & drop extensibility manager holds a list of registered drag and drop extensions.
  */
-class UMGEDITOR_API FWidgetDragDropExtensibilityManager
+class FWidgetDragDropExtensibilityManager
 {
 public:
 	void AddExtension(const TSharedRef<IWidgetDragDropExtension>& Extension)
diff --git a/Engine/Source/Editor/UMGEditor/Public/IUMGDesigner.h b/Engine/Source/Editor/UMGEditor/Public/IUMGDesigner.h
--- a/Engine/Source/Editor/UMGEditor/Public/IUMGDesigner.h
+++ b/Engine/Source/Editor/UMGEditor/Public/IUMGDesigner.h
@@ -24,7 +24,7 @@ namespace ETransformMode
  * The public interface implemented by the UMG Designer to allow extensions to call methods
  * on the designer.
  */
-class UMGEDITOR_API IUMGDesigner
+class IUMGDesigner
 {
 public:
 	/** @return the effective preview scale after both the DPI and Zoom scale has been applied. */
diff --git a/Engine/Source/Editor/UMGEditor/Public/Preview/PreviewMode.h b/Engine/Source/Editor/UMGEditor/Public/Preview/PreviewMode.h
--- a/Engine/Source/Editor/UMGEditor/Public/Preview/PreviewMode.h
+++ b/Engine/Source/Editor/UMGEditor/Public/Preview/PreviewMode.h
@@ -5,25 +5,27 @@
 #include "CoreMinimal.h"
 #include "UObject/WeakObjectPtr.h"
 
+#define UE_API UMGEDITOR_API
+
 class UObject;
 class UUserWidget;
 
 namespace UE::UMG::Editor
 {
 
-	class UMGEDITOR_API FPreviewMode
+	class FPreviewMode
 	{
 	public:
-		void SetSelectedObject(TArray<TWeakObjectPtr<UObject>> Objects);
-		void SetSelectedObject(TArrayView<UObject*> Objects);
-		TArray<UObject*> GetSelectedObjectList() const;
+		UE_API void SetSelectedObject(TArray<TWeakObjectPtr<UObject>> Objects);
+		UE_API void SetSelectedObject(TArrayView<UObject*> Objects);
+		UE_API TArray<UObject*> GetSelectedObjectList() const;
 
 		FSimpleMulticastDelegate& OnSelectedObjectChanged()
 		{
 			return SelectedObjectChangedDelegate;
 		}
 
-		void SetPreviewWidget(UUserWidget* Widget);
+		UE_API void SetPreviewWidget(UUserWidget* Widget);
 
 		UUserWidget* GetPreviewWidget() const
 		{
@@ -42,4 +44,6 @@ namespace UE::UMG::Editor
 		FSimpleMulticastDelegate PreviewWidgetChangedDelegate;
 	};
 
-} // namespace UE::UMG
\ No newline at end of file
+} // namespace UE::UMG
+
+#undef UE_API
diff --git a/Engine/Source/Editor/UMGEditor/Public/Settings/WidgetDesignerSettings.h b/Engine/Source/Editor/UMGEditor/Public/Settings/WidgetDesignerSettings.h
--- a/Engine/Source/Editor/UMGEditor/Public/Settings/WidgetDesignerSettings.h
+++ b/Engine/Source/Editor/UMGEditor/Public/Settings/WidgetDesignerSettings.h
@@ -15,6 +15,8 @@
 
 #include "WidgetDesignerSettings.generated.h"
 
+#define UE_API UMGEDITOR_API
+
 class UObject;
 class UWidgetPaletteFavorites;
 
@@ -30,17 +32,17 @@ enum EDisplayOnCompile : int
 /**
  * Implements the settings for the Widget Blueprint Designer.
  */
-UCLASS(config=EditorPerProjectUserSettings)
-class UMGEDITOR_API UWidgetDesignerSettings : public UDeveloperSettings
+UCLASS(MinimalAPI, config=EditorPerProjectUserSettings)
+class UWidgetDesignerSettings : public UDeveloperSettings
 {
 	GENERATED_BODY()
 
 public:
-	UWidgetDesignerSettings();
+	UE_API UWidgetDesignerSettings();
 
 #if WITH_EDITOR
-	virtual FText GetSectionText() const override;
-	virtual FText GetSectionDescription() const override;
+	UE_API virtual FText GetSectionText() const override;
+	UE_API virtual FText GetSectionDescription() const override;
 #endif
 
 public:
@@ -100,3 +102,5 @@ public:
 	TObjectPtr<UWidgetPaletteFavorites> Favorites;
 
 };
+
+#undef UE_API
diff --git a/Engine/Source/Editor/UMGEditor/Public/Templates/WidgetTemplateBlueprintClass.h b/Engine/Source/Editor/UMGEditor/Public/Templates/WidgetTemplateBlueprintClass.h
--- a/Engine/Source/Editor/UMGEditor/Public/Templates/WidgetTemplateBlueprintClass.h
+++ b/Engine/Source/Editor/UMGEditor/Public/Templates/WidgetTemplateBlueprintClass.h
@@ -10,13 +10,15 @@
 #include "Widgets/IToolTip.h"
 #include "Templates/WidgetTemplateClass.h"
 
+#define UE_API UMGEDITOR_API
+
 class UWidgetTree;
 
 /**
 * A template for classes generated by UWidgetBlueprint assets
 * Accounts for the possibility that a the blueprint may not be loaded when the palette is created
 */
-class UMGEDITOR_API FWidgetTemplateBlueprintClass : public FWidgetTemplateClass
+class FWidgetTemplateBlueprintClass : public FWidgetTemplateClass
 {
 public:
 	/** 
@@ -25,31 +27,31 @@ public:
 	* @param InUserWidgetClass The user-created widget class that will be created by this template
 	* @param bCacheParentClass True means we will search for parent class and cache, this is a slow operation, but useful for BPGCs.
 	*/
-	FWidgetTemplateBlueprintClass(const FAssetData& InWidgetAssetData, TSubclassOf<UUserWidget> InUserWidgetClass = nullptr);
+	UE_API FWidgetTemplateBlueprintClass(const FAssetData& InWidgetAssetData, TSubclassOf<UUserWidget> InUserWidgetClass = nullptr);
 
 	//UE_DEPRECATED(5.3, "FWidgetTemplateBlueprintClass with bInIsBlueprintGeneratedClass is deprecated.")
 	//FWidgetTemplateBlueprintClass(const FAssetData& InWidgetAssetData, TSubclassOf<UUserWidget> InUserWidgetClass, bool bInIsBlueprintGeneratedClass);
 
 	/** Destructor */
-	virtual ~FWidgetTemplateBlueprintClass();
+	UE_API virtual ~FWidgetTemplateBlueprintClass();
 
 	/** Gets the category for the widget */
-	virtual FText GetCategory() const override;
+	UE_API virtual FText GetCategory() const override;
 
 	/** Creates an instance of the widget for the tree. */
-	virtual UWidget* Create(UWidgetTree* Tree) override;
+	UE_API virtual UWidget* Create(UWidgetTree* Tree) override;
 
 	/** The icon coming from the default object of the class */
-	virtual const FSlateBrush* GetIcon() const override;
+	UE_API virtual const FSlateBrush* GetIcon() const override;
 
 	/** Gets the tooltip widget for this palette item. */
-	virtual TSharedRef<IToolTip> GetToolTip() const override;
+	UE_API virtual TSharedRef<IToolTip> GetToolTip() const override;
 
 	/** Opens the widget blueprint for edit */
-	virtual FReply OnDoubleClicked() override;
+	UE_API virtual FReply OnDoubleClicked() override;
 
 	/** Returns true if the supplied class is supported by this template */
-	static bool Supports(UClass* InClass);
+	static UE_API bool Supports(UClass* InClass);
 
 protected:
 	/** Parent Class of this widget template, may not be valid */
@@ -58,3 +60,5 @@ protected:
 	/** True if this template is for a cooked BP generated class */
 	bool bIsBlueprintGeneratedClass;
 };
+
+#undef UE_API
diff --git a/Engine/Source/Editor/UMGEditor/Public/Templates/WidgetTemplateClass.h b/Engine/Source/Editor/UMGEditor/Public/Templates/WidgetTemplateClass.h
--- a/Engine/Source/Editor/UMGEditor/Public/Templates/WidgetTemplateClass.h
+++ b/Engine/Source/Editor/UMGEditor/Public/Templates/WidgetTemplateClass.h
@@ -9,36 +9,38 @@
 #include "WidgetTemplate.h"
 #include "AssetRegistry/AssetData.h"
 
+#define UE_API UMGEDITOR_API
+
 class UWidgetTree;
 
 /**
  * A template that can spawn any widget derived from the UWidget class.
  */
-class UMGEDITOR_API FWidgetTemplateClass : public FWidgetTemplate
+class FWidgetTemplateClass : public FWidgetTemplate
 {
 public:
 	/** Constructor */
-	explicit FWidgetTemplateClass(TSubclassOf<UWidget> InWidgetClass);
+	UE_API explicit FWidgetTemplateClass(TSubclassOf<UWidget> InWidgetClass);
 
-	explicit FWidgetTemplateClass(const FAssetData& InWidgetAssetData, TSubclassOf<UWidget> InWidgetClass);
+	UE_API explicit FWidgetTemplateClass(const FAssetData& InWidgetAssetData, TSubclassOf<UWidget> InWidgetClass);
 
 	/** Destructor */
-	virtual ~FWidgetTemplateClass();
+	UE_API virtual ~FWidgetTemplateClass();
 
 	/** Gets the category for the widget */
-	virtual FText GetCategory() const override;
+	UE_API virtual FText GetCategory() const override;
 
 	/** Creates an instance of the widget for the tree. */
-	virtual UWidget* Create(UWidgetTree* Tree) override;
+	UE_API virtual UWidget* Create(UWidgetTree* Tree) override;
 
 	/** The icon coming from the default object of the class */
-	virtual const FSlateBrush* GetIcon() const override;
+	UE_API virtual const FSlateBrush* GetIcon() const override;
 
 	/** Gets the tooltip widget for this palette item. */
-	virtual TSharedRef<IToolTip> GetToolTip() const override;
+	UE_API virtual TSharedRef<IToolTip> GetToolTip() const override;
 
 	/** @param OutStrings - Returns an array of strings used for filtering/searching this widget template. */
-	virtual void GetFilterStrings(TArray<FString>& OutStrings) const override;
+	UE_API virtual void GetFilterStrings(TArray<FString>& OutStrings) const override;
 
 	/** Gets the WidgetClass which might be null. */
 	TWeakObjectPtr<UClass> GetWidgetClass() const { return WidgetClass; }
@@ -48,13 +50,13 @@ public:
 
 protected:
 	/** Creates a widget template class without any class reference */
-	FWidgetTemplateClass();
+	UE_API FWidgetTemplateClass();
 
 	/** Called when objects need to be swapped out for new versions, like after a blueprint recompile. */
-	void OnObjectsReplaced(const TMap<UObject*, UObject*>& ReplacementMap);
+	UE_API void OnObjectsReplaced(const TMap<UObject*, UObject*>& ReplacementMap);
 
 	/** Constructs the widget template with an overridden object name. */
-	UWidget* CreateNamed(class UWidgetTree* Tree, FName NameOverride);
+	UE_API UWidget* CreateNamed(class UWidgetTree* Tree, FName NameOverride);
 
 protected:
 	/** The widget class that will be created by this template */
@@ -66,3 +68,5 @@ protected:
 	/** Parent Class of this widget template, may not be valid */
 	TWeakObjectPtr<UClass> CachedParentClass;
 };
+
+#undef UE_API
diff --git a/Engine/Source/Editor/UMGEditor/Public/Templates/WidgetTemplateImageClass.h b/Engine/Source/Editor/UMGEditor/Public/Templates/WidgetTemplateImageClass.h
--- a/Engine/Source/Editor/UMGEditor/Public/Templates/WidgetTemplateImageClass.h
+++ b/Engine/Source/Editor/UMGEditor/Public/Templates/WidgetTemplateImageClass.h
@@ -6,13 +6,15 @@
 #include "AssetRegistry/AssetData.h"
 #include "Templates/WidgetTemplateClass.h"
 
+#define UE_API UMGEDITOR_API
+
 class UWidgetTree;
 
 /**
  * A template for classes generated by UTexture or UMaterial classes, or implements a USlateTextureAtlasInterface
  */
 
-class UMGEDITOR_API FWidgetTemplateImageClass : public FWidgetTemplateClass
+class FWidgetTemplateImageClass : public FWidgetTemplateClass
 {
 
 public:
@@ -20,13 +22,15 @@ public:
 	 * Constructor.
 	 * @param	InAssetData		The asset data used to create the widget
 	 */
-	FWidgetTemplateImageClass(const FAssetData& InAssetData);
+	UE_API FWidgetTemplateImageClass(const FAssetData& InAssetData);
 
-	virtual ~FWidgetTemplateImageClass();
+	UE_API virtual ~FWidgetTemplateImageClass();
 
 	/** Creates an instance of the widget for the tree */
-	virtual UWidget* Create(UWidgetTree* WidgetTree) override;
+	UE_API virtual UWidget* Create(UWidgetTree* WidgetTree) override;
 
 	/** Returns true if the supplied class is supported by this template */
-	static bool Supports(UClass* InClass);
+	static UE_API bool Supports(UClass* InClass);
 };
+
+#undef UE_API
diff --git a/Engine/Source/Editor/UMGEditor/Public/UIComponentUtils.h b/Engine/Source/Editor/UMGEditor/Public/UIComponentUtils.h
--- a/Engine/Source/Editor/UMGEditor/Public/UIComponentUtils.h
+++ b/Engine/Source/Editor/UMGEditor/Public/UIComponentUtils.h
@@ -6,12 +6,14 @@
 #include "ClassViewerFilter.h"
 #include "UObject/ObjectMacros.h"
 
+#define UE_API UMGEDITOR_API
+
 class UUIComponent;
 class UWidgetBlueprint;
 class FWidgetBlueprintEditor;
 /**  */
 
-class UMGEDITOR_API FUIComponentUtils
+class FUIComponentUtils
 {
 public:
 	class FUIComponentClassFilter : public IClassViewerFilter
@@ -29,10 +31,12 @@ public:
 		virtual bool IsUnloadedClassAllowed(const FClassViewerInitializationOptions& InInitOptions, const TSharedRef< const IUnloadedBlueprintData > InUnloadedClassData, TSharedRef< FClassViewerFilterFuncs > InFilterFuncs) override;
 	};
 
-	static FClassViewerInitializationOptions CreateClassViewerInitializationOptions();
+	static UE_API FClassViewerInitializationOptions CreateClassViewerInitializationOptions();
 
-	static void OnWidgetRenamed(const TSharedRef<FWidgetBlueprintEditor>& BlueprintEditor, UWidgetBlueprint* WidgetBlueprint, const FName& OldVarName, const FName& NewVarName);
-	static void AddComponent(const TSharedRef<FWidgetBlueprintEditor>& BlueprintEditor, const UClass* ComponentClass, const FName WidgetName);	
-	static void RemoveComponent(const TSharedRef<FWidgetBlueprintEditor>& BlueprintEditor, const UClass* ComponentClass, const FName WidgetName);
+	static UE_API void OnWidgetRenamed(const TSharedRef<FWidgetBlueprintEditor>& BlueprintEditor, UWidgetBlueprint* WidgetBlueprint, const FName& OldVarName, const FName& NewVarName);
+	static UE_API void AddComponent(const TSharedRef<FWidgetBlueprintEditor>& BlueprintEditor, const UClass* ComponentClass, const FName WidgetName);	
+	static UE_API void RemoveComponent(const TSharedRef<FWidgetBlueprintEditor>& BlueprintEditor, const UClass* ComponentClass, const FName WidgetName);
 	
-};
\ No newline at end of file
+};
+
+#undef UE_API
diff --git a/Engine/Source/Editor/UMGEditor/Public/UIComponentWidgetBlueprintExtension.h b/Engine/Source/Editor/UMGEditor/Public/UIComponentWidgetBlueprintExtension.h
--- a/Engine/Source/Editor/UMGEditor/Public/UIComponentWidgetBlueprintExtension.h
+++ b/Engine/Source/Editor/UMGEditor/Public/UIComponentWidgetBlueprintExtension.h
@@ -11,6 +11,8 @@
 
 #include "UIComponentWidgetBlueprintExtension.generated.h"
 
+#define UE_API UMGEDITOR_API
+
 class UUIComponentUserWidgetExtension;
 class UUIComponentContainer;
 
@@ -18,36 +20,36 @@ UCLASS()
 /***
  * Extension to the Widget Blueprint that will save the Components and will be used in the editor for compilation.
  */
-class UMGEDITOR_API UUIComponentWidgetBlueprintExtension : public UWidgetBlueprintExtension
+class UUIComponentWidgetBlueprintExtension : public UWidgetBlueprintExtension
 {
 	GENERATED_BODY()
 	
-	explicit UUIComponentWidgetBlueprintExtension(const FObjectInitializer& ObjectInitializer = FObjectInitializer());
+	UE_API explicit UUIComponentWidgetBlueprintExtension(const FObjectInitializer& ObjectInitializer = FObjectInitializer());
 
 public:
-	static const FName MD_ComponentVariable;
+	static UE_API const FName MD_ComponentVariable;
 
-	UUIComponent* AddComponent(const UClass* ComponentClass, FName OwnerName);
-	void RemoveComponent(const UClass* ComponentClass, FName OwnerName);
-	TArray<UUIComponent*> GetComponentsFor(const UWidget* Target) const;
-	UUIComponent* GetComponent(const UClass* ComponentClass, FName OwnerName) const;
+	UE_API UUIComponent* AddComponent(const UClass* ComponentClass, FName OwnerName);
+	UE_API void RemoveComponent(const UClass* ComponentClass, FName OwnerName);
+	UE_API TArray<UUIComponent*> GetComponentsFor(const UWidget* Target) const;
+	UE_API UUIComponent* GetComponent(const UClass* ComponentClass, FName OwnerName) const;
 
-	UUIComponentUserWidgetExtension* GetOrCreateExtension(UUserWidget* PreviewWidget);	
-	void RenameWidget(const FName& OldVarName, const FName& NewVarName);
+	UE_API UUIComponentUserWidgetExtension* GetOrCreateExtension(UUserWidget* PreviewWidget);	
+	UE_API void RenameWidget(const FName& OldVarName, const FName& NewVarName);
 
-	bool VerifyContainer(UUserWidget* UserWidget) const;
+	UE_API bool VerifyContainer(UUserWidget* UserWidget) const;
 
 protected:
-	virtual void HandleBeginCompilation(FWidgetBlueprintCompilerContext& InCreationContext) override;	
-	virtual void HandleCleanAndSanitizeClass(UWidgetBlueprintGeneratedClass* ClassToClean, UObject* OldCDO) override;
-	virtual void HandlePopulateGeneratedVariables(const FWidgetBlueprintCompilerContext::FPopulateGeneratedVariablesContext& Context) override;
-	virtual void HandleFinishCompilingClass(UWidgetBlueprintGeneratedClass* Class) override;
-	virtual bool HandleValidateGeneratedClass(UWidgetBlueprintGeneratedClass* Class) override; 
+	UE_API virtual void HandleBeginCompilation(FWidgetBlueprintCompilerContext& InCreationContext) override;	
+	UE_API virtual void HandleCleanAndSanitizeClass(UWidgetBlueprintGeneratedClass* ClassToClean, UObject* OldCDO) override;
+	UE_API virtual void HandlePopulateGeneratedVariables(const FWidgetBlueprintCompilerContext::FPopulateGeneratedVariablesContext& Context) override;
+	UE_API virtual void HandleFinishCompilingClass(UWidgetBlueprintGeneratedClass* Class) override;
+	UE_API virtual bool HandleValidateGeneratedClass(UWidgetBlueprintGeneratedClass* Class) override; 
 
-	virtual void HandleEndCompilation() override;
+	UE_API virtual void HandleEndCompilation() override;
 
 private:	
-	UUIComponentContainer* DuplicateContainer(UObject* Outer) const;
+	UE_API UUIComponentContainer* DuplicateContainer(UObject* Outer) const;
 
 	FWidgetBlueprintCompilerContext* CompilerContext = nullptr;
 
@@ -55,4 +57,6 @@ private:
 	UPROPERTY()
 	TObjectPtr<UUIComponentContainer> ComponentContainer;
 #endif //WITH_EDITORONLY_DATA
-};
\ No newline at end of file
+};
+
+#undef UE_API
diff --git a/Engine/Source/Editor/UMGEditor/Public/UMGEditorProjectSettings.h b/Engine/Source/Editor/UMGEditor/Public/UMGEditorProjectSettings.h
--- a/Engine/Source/Editor/UMGEditor/Public/UMGEditorProjectSettings.h
+++ b/Engine/Source/Editor/UMGEditor/Public/UMGEditorProjectSettings.h
@@ -12,6 +12,8 @@
 #include "Misc/NamePermissionList.h"
 #include "UMGEditorProjectSettings.generated.h"
 
+#define UE_API UMGEDITOR_API
+
 class UWidgetCompilerRule;
 class UUserWidget;
 class UWidgetBlueprint;
@@ -21,23 +23,25 @@ class UPanelWidget;
 /**
  * Implements the settings for the UMG Editor Project Settings
  */
-UCLASS(config=Editor, defaultconfig)
-class UMGEDITOR_API UUMGEditorProjectSettings : public UWidgetEditingProjectSettings
+UCLASS(MinimalAPI, config=Editor, defaultconfig)
+class UUMGEditorProjectSettings : public UWidgetEditingProjectSettings
 {
 	GENERATED_BODY()
 
 public:
-	UUMGEditorProjectSettings();
+	UE_API UUMGEditorProjectSettings();
 
 #if WITH_EDITOR
-	virtual FText GetSectionText() const override;
-	virtual FText GetSectionDescription() const override;
+	UE_API virtual FText GetSectionText() const override;
+	UE_API virtual FText GetSectionDescription() const override;
 #endif
 
 #if WITH_EDITOR
 	// Begin UObject Interface
-	virtual void PostEditChangeChainProperty(struct FPropertyChangedChainEvent& PropertyChangedEvent) override;
+	UE_API virtual void PostEditChangeChainProperty(struct FPropertyChangedChainEvent& PropertyChangedEvent) override;
 	// End UObject Interface
 #endif
 
 };
+
+#undef UE_API
diff --git a/Engine/Source/Editor/UMGEditor/Public/WidgetBlueprint.h b/Engine/Source/Editor/UMGEditor/Public/WidgetBlueprint.h
--- a/Engine/Source/Editor/UMGEditor/Public/WidgetBlueprint.h
+++ b/Engine/Source/Editor/UMGEditor/Public/WidgetBlueprint.h
@@ -14,6 +14,8 @@
 
 #include "WidgetBlueprint.generated.h"
 
+#define UE_API UMGEDITOR_API
+
 class FCompilerResultsLog;
 class UEdGraph;
 class UMovieScene;
@@ -29,7 +31,7 @@ class UWidgetEditingProjectSettings;
 
 
 /** Widget Delegates */
-class UMGEDITOR_API FWidgetBlueprintDelegates
+class FWidgetBlueprintDelegates
 {
 public:
 	// delegate for generating widget asset registry tags.
@@ -37,33 +39,33 @@ public:
 	DECLARE_MULTICAST_DELEGATE_TwoParams(FGetAssetTags, const UWidgetBlueprint*, TArray<UObject::FAssetRegistryTag>&);
 
 	// called by UWdgetBlueprint::GetAssetRegistryTags()
-	static FGetAssetTagsWithContext GetAssetTagsWithContext;
+	static UE_API FGetAssetTagsWithContext GetAssetTagsWithContext;
 	UE_DEPRECATED(5.4, "Subscribe to GetAssetTagsWithContext instead.")
-	static FGetAssetTags GetAssetTags;
+	static UE_API FGetAssetTags GetAssetTags;
 };
 
 
 /** */
 USTRUCT()
-struct UMGEDITOR_API FEditorPropertyPathSegment
+struct FEditorPropertyPathSegment
 {
 	GENERATED_USTRUCT_BODY()
 
 public:
-	FEditorPropertyPathSegment();
-	FEditorPropertyPathSegment(const FProperty* InProperty);
-	FEditorPropertyPathSegment(const UFunction* InFunction);
-	FEditorPropertyPathSegment(const UEdGraph* InFunctionGraph);
+	UE_API FEditorPropertyPathSegment();
+	UE_API FEditorPropertyPathSegment(const FProperty* InProperty);
+	UE_API FEditorPropertyPathSegment(const UFunction* InFunction);
+	UE_API FEditorPropertyPathSegment(const UEdGraph* InFunctionGraph);
 
 	UStruct* GetStruct() const { return Struct; }
-	FFieldVariant GetMember() const;
+	UE_API FFieldVariant GetMember() const;
 
-	void Rebase(UBlueprint* SegmentBase);
-	bool ValidateMember(FDelegateProperty* DelegateProperty, FText& OutError) const;
+	UE_API void Rebase(UBlueprint* SegmentBase);
+	UE_API bool ValidateMember(FDelegateProperty* DelegateProperty, FText& OutError) const;
 
-	FName GetMemberName() const;
-	FText GetMemberDisplayText() const;
-	FGuid GetMemberGuid() const;
+	UE_API FName GetMemberName() const;
+	UE_API FText GetMemberDisplayText() const;
+	UE_API FGuid GetMemberGuid() const;
 
 private:
 
@@ -90,32 +92,32 @@ private:
 
 /**  */
 USTRUCT()
-struct UMGEDITOR_API FEditorPropertyPath
+struct FEditorPropertyPath
 {
 	GENERATED_USTRUCT_BODY()
 
 public:
 
 	/**  */
-	FEditorPropertyPath();
+	UE_API FEditorPropertyPath();
 
 	/**  */
-	FEditorPropertyPath(const TArray<FFieldVariant>& BindingChain);
+	UE_API FEditorPropertyPath(const TArray<FFieldVariant>& BindingChain);
 
 	/**  */
-	bool Rebase(UBlueprint* SegmentBase);
+	UE_API bool Rebase(UBlueprint* SegmentBase);
 
 	/**  */
 	bool IsEmpty() const { return Segments.Num() == 0; }
 
 	/**  */
-	bool Validate(FDelegateProperty* Destination, FText& OutError) const;
+	UE_API bool Validate(FDelegateProperty* Destination, FText& OutError) const;
 
 	/**  */
-	FText GetDisplayText() const;
+	UE_API FText GetDisplayText() const;
 
 	/**  */
-	FDynamicPropertyPath ToPropertyPath() const;
+	UE_API FDynamicPropertyPath ToPropertyPath() const;
 
 public:
 
@@ -126,7 +128,7 @@ public:
 
 /** */
 USTRUCT()
-struct UMGEDITOR_API FDelegateEditorBinding
+struct FDelegateEditorBinding
 {
 	GENERATED_USTRUCT_BODY()
 
@@ -164,13 +166,13 @@ struct UMGEDITOR_API FDelegateEditorBinding
 		return ObjectName == Other.ObjectName && PropertyName == Other.PropertyName;
 	}
 
-	bool IsAttributePropertyBinding(class UWidgetBlueprint* Blueprint) const;
+	UE_API bool IsAttributePropertyBinding(class UWidgetBlueprint* Blueprint) const;
 
-	bool DoesBindingTargetExist(UWidgetBlueprint* Blueprint) const;
+	UE_API bool DoesBindingTargetExist(UWidgetBlueprint* Blueprint) const;
 
-	bool IsBindingValid(UClass* Class, class UWidgetBlueprint* Blueprint, FCompilerResultsLog& MessageLog) const;
+	UE_API bool IsBindingValid(UClass* Class, class UWidgetBlueprint* Blueprint, FCompilerResultsLog& MessageLog) const;
 
-	FDelegateRuntimeBinding ToRuntimeBinding(class UWidgetBlueprint* Blueprint) const;
+	UE_API FDelegateRuntimeBinding ToRuntimeBinding(class UWidgetBlueprint* Blueprint) const;
 };
 
 UENUM()
@@ -214,8 +216,8 @@ enum class EWidgetCompileTimeTickPrediction : uint8
 /**
  * The widget blueprint enables extending UUserWidget the user extensible UWidget.
  */
-UCLASS(BlueprintType)
-class UMGEDITOR_API UWidgetBlueprint : public UBaseWidgetBlueprint
+UCLASS(MinimalAPI, BlueprintType)
+class UWidgetBlueprint : public UBaseWidgetBlueprint
 {
 	GENERATED_UCLASS_BODY()
 
@@ -254,90 +256,90 @@ public:
 	UPROPERTY(EditAnywhere, Category="Widget")
 	bool bCanCallInitializedWithoutPlayerContext;
 
-	void OnVariableAdded(const FName& VariableName);
-	void OnVariableRenamed(const FName& OldName, const FName& NewName);
-	void OnVariableRemoved(const FName& VariableName);
+	UE_API void OnVariableAdded(const FName& VariableName);
+	UE_API void OnVariableRenamed(const FName& OldName, const FName& NewName);
+	UE_API void OnVariableRemoved(const FName& VariableName);
 #endif
 
 public:
 
 	/** UObject interface */
-	virtual void PostLoad() override;
-	virtual void PostDuplicate(bool bDuplicateForPIE) override;
+	UE_API virtual void PostLoad() override;
+	UE_API virtual void PostDuplicate(bool bDuplicateForPIE) override;
 
 #if WITH_EDITORONLY_DATA
-	virtual void PreSave(FObjectPreSaveContext ObjectSaveContext) override;
+	UE_API virtual void PreSave(FObjectPreSaveContext ObjectSaveContext) override;
 #endif // WITH_EDITORONLY_DATA
 
 #if WITH_EDITOR
-	virtual void GetAssetRegistryTags(FAssetRegistryTagsContext Context) const override;
+	UE_API virtual void GetAssetRegistryTags(FAssetRegistryTagsContext Context) const override;
 	UE_DEPRECATED(5.4, "Implement the version that takes FAssetRegistryTagsContext instead.")
-	virtual void GetAssetRegistryTags(TArray<FAssetRegistryTag>& OutTags) const override;
-	virtual void NotifyGraphRenamed(class UEdGraph* Graph, FName OldName, FName NewName) override;
-	virtual EDataValidationResult IsDataValid(class FDataValidationContext& Context) const override;
-	bool DetectSlateWidgetLeaks(class FDataValidationContext& Context) const;
-	virtual bool FindDiffs(const UBlueprint* OtherBlueprint, FDiffResults& Results) const override;
+	UE_API virtual void GetAssetRegistryTags(TArray<FAssetRegistryTag>& OutTags) const override;
+	UE_API virtual void NotifyGraphRenamed(class UEdGraph* Graph, FName OldName, FName NewName) override;
+	UE_API virtual EDataValidationResult IsDataValid(class FDataValidationContext& Context) const override;
+	UE_API bool DetectSlateWidgetLeaks(class FDataValidationContext& Context) const;
+	UE_API virtual bool FindDiffs(const UBlueprint* OtherBlueprint, FDiffResults& Results) const override;
 #endif
 
-	virtual void Serialize(FArchive& Ar) override;
+	UE_API virtual void Serialize(FArchive& Ar) override;
 
-	UPackage* GetWidgetTemplatePackage() const;
+	UE_API UPackage* GetWidgetTemplatePackage() const;
 
-	virtual void ReplaceDeprecatedNodes() override;
+	UE_API virtual void ReplaceDeprecatedNodes() override;
 	
 	//~ Begin UBlueprint Interface
-	virtual UClass* GetBlueprintClass() const override;
+	UE_API virtual UClass* GetBlueprintClass() const override;
 
-	virtual bool AllowsDynamicBinding() const override;
+	UE_API virtual bool AllowsDynamicBinding() const override;
 
-	virtual bool SupportsInputEvents() const override;
+	UE_API virtual bool SupportsInputEvents() const override;
 
 	virtual bool SupportedByDefaultBlueprintFactory() const override
 	{
 		return false;
 	}
 
-	virtual void GetReparentingRules(TSet< const UClass* >& AllowedChildrenOfClasses, TSet< const UClass* >& DisallowedChildrenOfClasses) const override;
+	UE_API virtual void GetReparentingRules(TSet< const UClass* >& AllowedChildrenOfClasses, TSet< const UClass* >& DisallowedChildrenOfClasses) const override;
 
 	/** UWidget blueprints are never data only, should always compile on load (data only blueprints cannot declare new variables) */
 	virtual bool AlwaysCompileOnLoad() const override { return true; }
 	//~ End UBlueprint Interface
 
-	virtual void GatherDependencies(TSet<TWeakObjectPtr<UBlueprint>>& InDependencies) const override;
+	UE_API virtual void GatherDependencies(TSet<TWeakObjectPtr<UBlueprint>>& InDependencies) const override;
 
 	/** Returns true if the supplied user widget will not create a circular reference when added to this blueprint */
-	bool IsWidgetFreeFromCircularReferences(UUserWidget* UserWidget) const;
+	UE_API bool IsWidgetFreeFromCircularReferences(UUserWidget* UserWidget) const;
 	
 	/**  */
-	TValueOrError<void, UWidget*> HasCircularReferences() const;
+	UE_API TValueOrError<void, UWidget*> HasCircularReferences() const;
 
-	TValueOrError<void, TSet<UWidget*>> HasConflictingWidgetNamesFromInheritance() const;
+	UE_API TValueOrError<void, TSet<UWidget*>> HasConflictingWidgetNamesFromInheritance() const;
 
 
-	static bool ValidateGeneratedClass(const UClass* InClass);
+	static UE_API bool ValidateGeneratedClass(const UClass* InClass);
 	
-	static TSharedPtr<FKismetCompilerContext> GetCompilerForWidgetBP(UBlueprint* BP, FCompilerResultsLog& InMessageLog, const FKismetCompilerOptions& InCompileOptions);
+	static UE_API TSharedPtr<FKismetCompilerContext> GetCompilerForWidgetBP(UBlueprint* BP, FCompilerResultsLog& InMessageLog, const FKismetCompilerOptions& InCompileOptions);
 
-	void UpdateTickabilityStats(bool& OutHasLatentActions, bool& OutHasAnimations, bool& OutClassRequiresNativeTick);
+	UE_API void UpdateTickabilityStats(bool& OutHasLatentActions, bool& OutHasAnimations, bool& OutClassRequiresNativeTick);
 
-	bool ArePropertyBindingsAllowed() const;
+	UE_API bool ArePropertyBindingsAllowed() const;
 
 	/** Gets any named slots exposed by the parent generated class that can be slotted into by the subclass. */
-	TArray<FName> GetInheritedAvailableNamedSlots() const;
+	UE_API TArray<FName> GetInheritedAvailableNamedSlots() const;
 
 	/** Gets named slots exposed by the parent generated class that already have content in the parent widget tree. */
-	TSet<FName> GetInheritedNamedSlotsWithContentInSameTree() const;
+	UE_API TSet<FName> GetInheritedNamedSlotsWithContentInSameTree() const;
 
-	virtual UWidgetEditingProjectSettings* GetRelevantSettings();
-	virtual const UWidgetEditingProjectSettings* GetRelevantSettings() const;
+	UE_API virtual UWidgetEditingProjectSettings* GetRelevantSettings();
+	UE_API virtual const UWidgetEditingProjectSettings* GetRelevantSettings() const;
 
 protected:
 
-	virtual void LoadModulesRequiredForCompilation() override;
+	UE_API virtual void LoadModulesRequiredForCompilation() override;
 
 private:
 
-	void FixupLegacySequencePlayerCalls();
+	UE_API void FixupLegacySequencePlayerCalls();
 
 private:
 	/**
@@ -377,3 +379,5 @@ public:
 	TObjectPtr<UTexture2D> ThumbnailImage;
 
 };
+
+#undef UE_API
diff --git a/Engine/Source/Editor/UMGEditor/Public/WidgetBlueprintCompiler.h b/Engine/Source/Editor/UMGEditor/Public/WidgetBlueprintCompiler.h
--- a/Engine/Source/Editor/UMGEditor/Public/WidgetBlueprintCompiler.h
+++ b/Engine/Source/Editor/UMGEditor/Public/WidgetBlueprintCompiler.h
@@ -11,6 +11,8 @@
 #include "ComponentReregisterContext.h"
 #include "Components/WidgetComponent.h"
 
+#define UE_API UMGEDITOR_API
+
 class FProperty;
 class UEdGraph;
 class UWidget;
@@ -20,17 +22,17 @@ class UWidgetGraphSchema;
 //////////////////////////////////////////////////////////////////////////
 // FWidgetBlueprintCompiler 
 
-class UMGEDITOR_API FWidgetBlueprintCompiler : public IBlueprintCompiler
+class FWidgetBlueprintCompiler : public IBlueprintCompiler
 {
 
 public:
-	FWidgetBlueprintCompiler();
+	UE_API FWidgetBlueprintCompiler();
 
-	bool CanCompile(const UBlueprint* Blueprint) override;
-	void PreCompile(UBlueprint* Blueprint, const FKismetCompilerOptions& CompileOptions) override;
-	void Compile(UBlueprint* Blueprint, const FKismetCompilerOptions& CompileOptions, FCompilerResultsLog& Results) override;
-	void PostCompile(UBlueprint* Blueprint, const FKismetCompilerOptions& CompileOptions) override;
-	bool GetBlueprintTypesForClass(UClass* ParentClass, UClass*& OutBlueprintClass, UClass*& OutBlueprintGeneratedClass) const override;
+	UE_API bool CanCompile(const UBlueprint* Blueprint) override;
+	UE_API void PreCompile(UBlueprint* Blueprint, const FKismetCompilerOptions& CompileOptions) override;
+	UE_API void Compile(UBlueprint* Blueprint, const FKismetCompilerOptions& CompileOptions, FCompilerResultsLog& Results) override;
+	UE_API void PostCompile(UBlueprint* Blueprint, const FKismetCompilerOptions& CompileOptions) override;
+	UE_API bool GetBlueprintTypesForClass(UClass* ParentClass, UClass*& OutBlueprintClass, UClass*& OutBlueprintGeneratedClass) const override;
 
 private:
 
@@ -49,100 +51,100 @@ private:
 // FWidgetBlueprintCompilerContext
 
 
-class UMGEDITOR_API FWidgetBlueprintCompilerContext : public FKismetCompilerContext
+class FWidgetBlueprintCompilerContext : public FKismetCompilerContext
 {
 protected:
 	typedef FKismetCompilerContext Super;
 
 public:
-	FWidgetBlueprintCompilerContext(UWidgetBlueprint* SourceSketch, FCompilerResultsLog& InMessageLog, const FKismetCompilerOptions& InCompilerOptions);
-	virtual ~FWidgetBlueprintCompilerContext();
+	UE_API FWidgetBlueprintCompilerContext(UWidgetBlueprint* SourceSketch, FCompilerResultsLog& InMessageLog, const FKismetCompilerOptions& InCompilerOptions);
+	UE_API virtual ~FWidgetBlueprintCompilerContext();
 
 protected:
 	/**
 	 * Checks if the animations' bindings are valid. Shows a compiler warning if any of the animations have a null track, urging the user
 	 * to delete said track to avoid performance hitches when playing said null tracks in large user widgets.
 	*/
-	void ValidateWidgetAnimations();
+	UE_API void ValidateWidgetAnimations();
 
 	/** Validates the Desired Focus name to make sure it's part of the Widget Tree. */
-	void ValidateDesiredFocusWidgetName();
+	UE_API void ValidateDesiredFocusWidgetName();
 
 	//~ Begin FKismetCompilerContext
-	virtual UEdGraphSchema_K2* CreateSchema() override;
-	virtual void CreateFunctionList() override;
-	virtual void SpawnNewClass(const FString& NewClassName) override;
-	virtual void OnNewClassSet(UBlueprintGeneratedClass* ClassToUse) override;
-	virtual void PrecompileFunction(FKismetFunctionContext& Context, EInternalCompilerFlags InternalFlags) override;
-	virtual void CleanAndSanitizeClass(UBlueprintGeneratedClass* ClassToClean, UObject*& InOutOldCDO) override;
-	virtual void SaveSubObjectsFromCleanAndSanitizeClass(FSubobjectCollection& SubObjectsToSave, UBlueprintGeneratedClass* ClassToClean) override;
-	virtual void EnsureProperGeneratedClass(UClass*& TargetClass) override;
-	virtual void PopulateBlueprintGeneratedVariables() override;
-	virtual void CreateClassVariablesFromBlueprint() override;
-	virtual void CopyTermDefaultsToDefaultObject(UObject* DefaultObject);
-	virtual void FinishCompilingClass(UClass* Class) override;
-	virtual bool ValidateGeneratedClass(UBlueprintGeneratedClass* Class) override;
-	virtual void OnPostCDOCompiled(const UObject::FPostCDOCompiledContext& Context) override;
+	UE_API virtual UEdGraphSchema_K2* CreateSchema() override;
+	UE_API virtual void CreateFunctionList() override;
+	UE_API virtual void SpawnNewClass(const FString& NewClassName) override;
+	UE_API virtual void OnNewClassSet(UBlueprintGeneratedClass* ClassToUse) override;
+	UE_API virtual void PrecompileFunction(FKismetFunctionContext& Context, EInternalCompilerFlags InternalFlags) override;
+	UE_API virtual void CleanAndSanitizeClass(UBlueprintGeneratedClass* ClassToClean, UObject*& InOutOldCDO) override;
+	UE_API virtual void SaveSubObjectsFromCleanAndSanitizeClass(FSubobjectCollection& SubObjectsToSave, UBlueprintGeneratedClass* ClassToClean) override;
+	UE_API virtual void EnsureProperGeneratedClass(UClass*& TargetClass) override;
+	UE_API virtual void PopulateBlueprintGeneratedVariables() override;
+	UE_API virtual void CreateClassVariablesFromBlueprint() override;
+	UE_API virtual void CopyTermDefaultsToDefaultObject(UObject* DefaultObject);
+	UE_API virtual void FinishCompilingClass(UClass* Class) override;
+	UE_API virtual bool ValidateGeneratedClass(UBlueprintGeneratedClass* Class) override;
+	UE_API virtual void OnPostCDOCompiled(const UObject::FPostCDOCompiledContext& Context) override;
 	//~ End FKismetCompilerContext
 
-	void SanitizeBindings(UBlueprintGeneratedClass* Class);
-	void ValidateAndFixUpVariableGuids();
+	UE_API void SanitizeBindings(UBlueprintGeneratedClass* Class);
+	UE_API void ValidateAndFixUpVariableGuids();
 
-	void VerifyEventReplysAreNotEmpty(FKismetFunctionContext& Context);
-	void VerifyFieldNotifyFunction(FKismetFunctionContext& Context);
+	UE_API void VerifyEventReplysAreNotEmpty(FKismetFunctionContext& Context);
+	UE_API void VerifyFieldNotifyFunction(FKismetFunctionContext& Context);
 
 public:
 	UWidgetBlueprint* WidgetBlueprint() const { return Cast<UWidgetBlueprint>(Blueprint); }
 
-	void AddExtension(UWidgetBlueprintGeneratedClass* Class, UWidgetBlueprintGeneratedClassExtension* Extension);
-	void AddGeneratedVariable(FBPVariableDescription&& VariableDescription) const;
+	UE_API void AddExtension(UWidgetBlueprintGeneratedClass* Class, UWidgetBlueprintGeneratedClassExtension* Extension);
+	UE_API void AddGeneratedVariable(FBPVariableDescription&& VariableDescription) const;
 	
-	struct UMGEDITOR_API FPopulateGeneratedVariablesContext
+	struct FPopulateGeneratedVariablesContext
 	{
 	public:
-		void AddGeneratedVariable(FBPVariableDescription&& VariableDescription) const;
-		UWidgetBlueprint* GetWidgetBlueprint() const;
+		UE_API void AddGeneratedVariable(FBPVariableDescription&& VariableDescription) const;
+		UE_API UWidgetBlueprint* GetWidgetBlueprint() const;
 
 	private:
 		friend FWidgetBlueprintCompilerContext;
-		FPopulateGeneratedVariablesContext(FWidgetBlueprintCompilerContext& InContext);
+		UE_API FPopulateGeneratedVariablesContext(FWidgetBlueprintCompilerContext& InContext);
 		FWidgetBlueprintCompilerContext& Context;
 	};
 
-	struct UMGEDITOR_API FCreateVariableContext
+	struct FCreateVariableContext
 	{
 	public:
-		FProperty* CreateVariable(const FName Name, const FEdGraphPinType& Type) const;
-		FMulticastDelegateProperty* CreateMulticastDelegateVariable(const FName Name, const FEdGraphPinType& Type) const;
-		FMulticastDelegateProperty* CreateMulticastDelegateVariable(const FName Name) const;
-		void AddGeneratedFunctionGraph(UEdGraph* Graph) const;
-		UWidgetBlueprint* GetWidgetBlueprint() const;
+		UE_API FProperty* CreateVariable(const FName Name, const FEdGraphPinType& Type) const;
+		UE_API FMulticastDelegateProperty* CreateMulticastDelegateVariable(const FName Name, const FEdGraphPinType& Type) const;
+		UE_API FMulticastDelegateProperty* CreateMulticastDelegateVariable(const FName Name) const;
+		UE_API void AddGeneratedFunctionGraph(UEdGraph* Graph) const;
+		UE_API UWidgetBlueprint* GetWidgetBlueprint() const;
 		UE_DEPRECATED(5.4, "GetSkeletonGeneratedClass renamed to GetGeneratedClass")
-		UWidgetBlueprintGeneratedClass* GetSkeletonGeneratedClass() const;
-		UWidgetBlueprintGeneratedClass* GetGeneratedClass() const;
-		EKismetCompileType::Type GetCompileType() const;
+		UE_API UWidgetBlueprintGeneratedClass* GetSkeletonGeneratedClass() const;
+		UE_API UWidgetBlueprintGeneratedClass* GetGeneratedClass() const;
+		UE_API EKismetCompileType::Type GetCompileType() const;
 
 	private:
 		friend FWidgetBlueprintCompilerContext;
-		FCreateVariableContext(FWidgetBlueprintCompilerContext& InContext);
+		UE_API FCreateVariableContext(FWidgetBlueprintCompilerContext& InContext);
 		FWidgetBlueprintCompilerContext& Context;
 	};
 
-	struct UMGEDITOR_API FCreateFunctionContext
+	struct FCreateFunctionContext
 	{
 	public:
-		void AddGeneratedFunctionGraph(UEdGraph* Graph) const;
-		void AddGeneratedUbergraphPage(UEdGraph* Graph) const;
-		UWidgetBlueprintGeneratedClass* GetGeneratedClass() const;
+		UE_API void AddGeneratedFunctionGraph(UEdGraph* Graph) const;
+		UE_API void AddGeneratedUbergraphPage(UEdGraph* Graph) const;
+		UE_API UWidgetBlueprintGeneratedClass* GetGeneratedClass() const;
 
 	private:
 		friend FWidgetBlueprintCompilerContext;
-		FCreateFunctionContext(FWidgetBlueprintCompilerContext& InContext);
+		UE_API FCreateFunctionContext(FWidgetBlueprintCompilerContext& InContext);
 		FWidgetBlueprintCompilerContext& Context;
 	};
 
 protected:
-	void FixAbandonedWidgetTree(UWidgetBlueprint* WidgetBP);
+	UE_API void FixAbandonedWidgetTree(UWidgetBlueprint* WidgetBP);
 
 	UWidgetBlueprintGeneratedClass* NewWidgetBlueprintClass;
 
@@ -164,3 +166,4 @@ protected:
 	///----------------------------------------------------------------
 };
 
+#undef UE_API
