Updated headers using UnrealCodeFixup to make sure dllstorage is on methods/staticvars instead of types. Part 5/n

diff --git a/Engine/Source/Editor/UMGEditor/Public/WidgetBlueprintEditor.h b/Engine/Source/Editor/UMGEditor/Public/WidgetBlueprintEditor.h
--- a/Engine/Source/Editor/UMGEditor/Public/WidgetBlueprintEditor.h
+++ b/Engine/Source/Editor/UMGEditor/Public/WidgetBlueprintEditor.h
@@ -16,6 +16,8 @@
 #include "WidgetReference.h"
 #include "Blueprint/UserWidget.h"
 
+#define UE_API UMGEDITOR_API
+
 class FMenuBuilder;
 class FWidgetBlueprintEditorToolbar;
 class IMessageLogListing;
@@ -36,7 +38,7 @@ struct FNamedSlotSelection
 /**
  * widget blueprint editor (extends Blueprint editor)
  */
-class UMGEDITOR_API FWidgetBlueprintEditor : public FBlueprintEditor
+class FWidgetBlueprintEditor : public FBlueprintEditor
 {
 private:
 	using Super = FBlueprintEditor;
@@ -60,109 +62,109 @@ public:
 	DECLARE_EVENT(FWidgetBlueprintEditor, FOnEnterWidgetDesigner)
 
 public:
-	FWidgetBlueprintEditor();
+	UE_API FWidgetBlueprintEditor();
 
-	virtual ~FWidgetBlueprintEditor();
+	UE_API virtual ~FWidgetBlueprintEditor();
 
-	void InitWidgetBlueprintEditor(const EToolkitMode::Type Mode, const TSharedPtr< IToolkitHost >& InitToolkitHost, const TArray<UBlueprint*>& InBlueprints, bool bShouldOpenInDefaultsMode);
+	UE_API void InitWidgetBlueprintEditor(const EToolkitMode::Type Mode, const TSharedPtr< IToolkitHost >& InitToolkitHost, const TArray<UBlueprint*>& InBlueprints, bool bShouldOpenInDefaultsMode);
 
 	//~ Begin FBlueprintEditor interface
-	virtual void Tick(float DeltaTime) override;
-	virtual void PostUndo(bool bSuccessful) override;
-	virtual void PostRedo(bool bSuccessful) override;
-	virtual void Compile() override;
+	UE_API virtual void Tick(float DeltaTime) override;
+	UE_API virtual void PostUndo(bool bSuccessful) override;
+	UE_API virtual void PostRedo(bool bSuccessful) override;
+	UE_API virtual void Compile() override;
 	//~ End FBlueprintEditor interface
 	
 	//~ Begin FAssetEditorToolkit Interface
-	virtual bool OnRequestClose(EAssetEditorCloseReason InCloseReason) override;
+	UE_API virtual bool OnRequestClose(EAssetEditorCloseReason InCloseReason) override;
 	// End of FAssetEditorToolkit 
 
 	//~ Begin FGCObjectInterface interface
-	virtual void AddReferencedObjects( FReferenceCollector& Collector ) override;
+	UE_API virtual void AddReferencedObjects( FReferenceCollector& Collector ) override;
 	//~ End FGCObjectInterface interface
 
 	//~ Begin IToolkit interface
-	virtual FName GetToolkitContextFName() const override;
-	virtual FName GetToolkitFName() const override;
-	virtual FText GetBaseToolkitName() const override;
-	virtual FString GetWorldCentricTabPrefix() const override;
-	virtual FLinearColor GetWorldCentricTabColorScale() const override;
-	virtual void InitToolMenuContext(FToolMenuContext& MenuContext) override;
-	void OnToolkitHostingStarted(const TSharedRef<IToolkit>& Toolkit) override;
-	void OnToolkitHostingFinished(const TSharedRef<IToolkit>& Toolkit) override;
+	UE_API virtual FName GetToolkitContextFName() const override;
+	UE_API virtual FName GetToolkitFName() const override;
+	UE_API virtual FText GetBaseToolkitName() const override;
+	UE_API virtual FString GetWorldCentricTabPrefix() const override;
+	UE_API virtual FLinearColor GetWorldCentricTabColorScale() const override;
+	UE_API virtual void InitToolMenuContext(FToolMenuContext& MenuContext) override;
+	UE_API void OnToolkitHostingStarted(const TSharedRef<IToolkit>& Toolkit) override;
+	UE_API void OnToolkitHostingFinished(const TSharedRef<IToolkit>& Toolkit) override;
 	//~ End IToolkit interface
 
 	/** @return The widget blueprint currently being edited in this editor */
-	class UWidgetBlueprint* GetWidgetBlueprintObj() const;
+	UE_API class UWidgetBlueprint* GetWidgetBlueprintObj() const;
 
 	/** @return The preview widget. */
-	UUserWidget* GetPreview() const;
+	UE_API UUserWidget* GetPreview() const;
 
 	/** @return The preview scene that owns the preview widget. */
-	FPreviewScene* GetPreviewScene();
+	UE_API FPreviewScene* GetPreviewScene();
 
 	/**  */
-	bool IsSimulating() const;
+	UE_API bool IsSimulating() const;
 
 	/**  */
-	void SetIsSimulating(bool bSimulating);
+	UE_API void SetIsSimulating(bool bSimulating);
 
 	/** Causes the preview to be destroyed and a new one to be created next tick */
-	void InvalidatePreview(bool bViewOnly = false);
+	UE_API void InvalidatePreview(bool bViewOnly = false);
 
 	/** Immediately rebuilds the preview widget. */
-	void RefreshPreview();
+	UE_API void RefreshPreview();
 
 	/** Creates a widget reference using the template. */
-	FWidgetReference GetReferenceFromTemplate(UWidget* TemplateWidget);
+	UE_API FWidgetReference GetReferenceFromTemplate(UWidget* TemplateWidget);
 
 	/** Creates a widget reference using the preview.  Which is used to lookup the stable template pointer. */
-	FWidgetReference GetReferenceFromPreview(UWidget* PreviewWidget);
+	UE_API FWidgetReference GetReferenceFromPreview(UWidget* PreviewWidget);
 
 	/** @return The currently active sequencer used to create widget animations */
-	TSharedPtr<ISequencer>& GetSequencer();
+	UE_API TSharedPtr<ISequencer>& GetSequencer();
 
 	/** @return The sequencer used to create widget animations */
-	TSharedPtr<ISequencer>& GetTabSequencer();
+	UE_API TSharedPtr<ISequencer>& GetTabSequencer();
 
 	/** @return The sequencer used to create widget animations in drawer */
-	TSharedPtr<ISequencer>& GetDrawerSequencer();
+	UE_API TSharedPtr<ISequencer>& GetDrawerSequencer();
 
 	/** Callback to dock anim tab into layout */
-	void DockInLayoutClicked();
+	UE_API void DockInLayoutClicked();
 
 	/** Changes the currently viewed animation in Sequencer to the new one*/
-	void ChangeViewedAnimation( UWidgetAnimation& InAnimationToView );
+	UE_API void ChangeViewedAnimation( UWidgetAnimation& InAnimationToView );
 
 	/** Get the current animation*/
 	UWidgetAnimation* GetCurrentAnimation() { return CurrentAnimation.Get(); }
 
 	/** Updates the current animation if it is invalid */
-	const UWidgetAnimation* RefreshCurrentAnimation();
+	UE_API const UWidgetAnimation* RefreshCurrentAnimation();
 
 	/** Sets the currently selected set of widgets */
-	void SelectWidgets(const TSet<FWidgetReference>& Widgets, bool bAppendOrToggle);
+	UE_API void SelectWidgets(const TSet<FWidgetReference>& Widgets, bool bAppendOrToggle);
 
 	/** Sets the currently selected set of objects */
-	void SelectObjects(const TSet<UObject*>& Objects);
+	UE_API void SelectObjects(const TSet<UObject*>& Objects);
 
 	/** Called to determine whether a binding is selected in the tree view */
-	bool IsBindingSelected(const FMovieSceneBinding& InBinding);
+	UE_API bool IsBindingSelected(const FMovieSceneBinding& InBinding);
 
 	/** Sets the selected named slot */
-	void SetSelectedNamedSlot(TOptional<FNamedSlotSelection> SelectedNamedSlot);
+	UE_API void SetSelectedNamedSlot(TOptional<FNamedSlotSelection> SelectedNamedSlot);
 
 	/** Removes removed widgets from the selection set. */
-	void CleanSelection();
+	UE_API void CleanSelection();
 
 	/** @return The selected set of widgets */
-	const TSet<FWidgetReference>& GetSelectedWidgets() const;
+	UE_API const TSet<FWidgetReference>& GetSelectedWidgets() const;
 
 	/** @return The selected named slot */
-	TOptional<FNamedSlotSelection> GetSelectedNamedSlot() const;
+	UE_API TOptional<FNamedSlotSelection> GetSelectedNamedSlot() const;
 
 	/** @return The selected set of Objects */
-	const TSet< TWeakObjectPtr<UObject> >& GetSelectedObjects() const;
+	UE_API const TSet< TWeakObjectPtr<UObject> >& GetSelectedObjects() const;
 
 	/** @return the selected template widget */
 	const TWeakObjectPtr<UClass> GetSelectedTemplate() const { return SelectedTemplate; }
@@ -179,84 +181,84 @@ public:
 	TSharedPtr<class FWidgetBlueprintEditorToolbar> GetWidgetToolbarBuilder() { return WidgetToolbar; }
 
 	/** Migrate a property change from the preview GUI to the template GUI. */
-	void MigrateFromChain(FEditPropertyChain* PropertyThatChanged, bool bIsModify);
+	UE_API void MigrateFromChain(FEditPropertyChain* PropertyThatChanged, bool bIsModify);
 
 	/** Event called when an undo/redo transaction occurs */
 	DECLARE_EVENT(FWidgetBlueprintEditor, FOnWidgetBlueprintTransaction)
 	FOnWidgetBlueprintTransaction& GetOnWidgetBlueprintTransaction() { return OnWidgetBlueprintTransaction; }
 
 	/** Creates a sequencer widget */
-	TSharedRef<SWidget> CreateSequencerTabWidget();
+	UE_API TSharedRef<SWidget> CreateSequencerTabWidget();
 
 	/** Creates a sequencer widget, for anim drawer */
-	TSharedRef<SWidget> CreateSequencerDrawerWidget();
+	UE_API TSharedRef<SWidget> CreateSequencerDrawerWidget();
 
 	/** Gets sequencer widget, for anim drawer */
-	TSharedRef<SWidget> OnGetWidgetAnimSequencer();
+	UE_API TSharedRef<SWidget> OnGetWidgetAnimSequencer();
 
 	/** Adds external widget whose lifetime should be managed by this editor */
-	void AddExternalEditorWidget(FName ID, TSharedRef<SWidget> InExternalWidget);
+	UE_API void AddExternalEditorWidget(FName ID, TSharedRef<SWidget> InExternalWidget);
 
 	/** Removes external widget whose lifetime should be managed by this editor */
-	int32 RemoveExternalEditorWidget(FName ID);
+	UE_API int32 RemoveExternalEditorWidget(FName ID);
 
 	/** Gets external widget whose lifetime should be managed by this editor, NullWidget if not found */
-	TSharedPtr<SWidget> GetExternalEditorWidget(FName ID);
+	UE_API TSharedPtr<SWidget> GetExternalEditorWidget(FName ID);
 
 	/** Toggles Anim Drawer, only used for keyboard shortcut */
-	void ToggleAnimDrawer();
+	UE_API void ToggleAnimDrawer();
 
 	/** Callback for when widget animation list has been updated */
-	void NotifyWidgetAnimListChanged();
+	UE_API void NotifyWidgetAnimListChanged();
 
 	UE_DEPRECATED(5.2, "OnWidgetAnimSequencerOpened is deprecated, please use OnWidgetAnimDrawerSequencerOpened instead.")
-	void OnWidgetAnimSequencerOpened(FName StatusBarWithDrawerName);
+	UE_API void OnWidgetAnimSequencerOpened(FName StatusBarWithDrawerName);
 
 	UE_DEPRECATED(5.2, "OnWidgetAnimSequencerDismissed is deprecated, please use OnWidgetAnimDrawerSequencerDismissed instead.")
-	void OnWidgetAnimSequencerDismissed(const TSharedPtr<SWidget>& NewlyFocusedWidget);
+	UE_API void OnWidgetAnimSequencerDismissed(const TSharedPtr<SWidget>& NewlyFocusedWidget);
 
 	/** Callback for anim drawer opening */
-	void OnWidgetAnimDrawerSequencerOpened(FName StatusBarWithDrawerName);
+	UE_API void OnWidgetAnimDrawerSequencerOpened(FName StatusBarWithDrawerName);
 
 	/** Callback for anim drawer closing */
-	void OnWidgetAnimDrawerSequencerDismissed(const TSharedPtr<SWidget>& NewlyFocusedWidget);
+	UE_API void OnWidgetAnimDrawerSequencerDismissed(const TSharedPtr<SWidget>& NewlyFocusedWidget);
 
 	/** Callback for anim tab closing */
-	void OnWidgetAnimTabSequencerClosed(TSharedRef<SDockTab> ClosedTab);
+	UE_API void OnWidgetAnimTabSequencerClosed(TSharedRef<SDockTab> ClosedTab);
 
 	/** Callback for anim tab opening */
-	void OnWidgetAnimTabSequencerOpened();
+	UE_API void OnWidgetAnimTabSequencerOpened();
 
 	/**
 	 * The widget we're now hovering over in any particular context, allows multiple views to 
 	 * synchronize feedback on where that widget is in their representation.
 	 */
-	void SetHoveredWidget(FWidgetReference& InHoveredWidget);
+	UE_API void SetHoveredWidget(FWidgetReference& InHoveredWidget);
 
-	void ClearHoveredWidget();
+	UE_API void ClearHoveredWidget();
 
 	/** @return The widget that is currently being hovered over (either in the designer or hierarchy) */
-	const FWidgetReference& GetHoveredWidget() const;
+	UE_API const FWidgetReference& GetHoveredWidget() const;
 
-	void AddPostDesignerLayoutAction(TFunction<void()> Action);
+	UE_API void AddPostDesignerLayoutAction(TFunction<void()> Action);
 
-	void OnEnteringDesigner();
+	UE_API void OnEnteringDesigner();
 
-	TArray< TFunction<void()> >& GetQueuedDesignerActions();
+	UE_API TArray< TFunction<void()> >& GetQueuedDesignerActions();
 
 	/** Get the current designer flags that are in effect for the current user widget we're editing. */
-	EWidgetDesignFlags GetCurrentDesignerFlags() const;
+	UE_API EWidgetDesignFlags GetCurrentDesignerFlags() const;
 
-	bool GetShowDashedOutlines() const;
-	void SetShowDashedOutlines(bool Value);
+	UE_API bool GetShowDashedOutlines() const;
+	UE_API void SetShowDashedOutlines(bool Value);
 
-	bool GetIsRespectingLocks() const;
-	void SetIsRespectingLocks(bool Value);
+	UE_API bool GetIsRespectingLocks() const;
+	UE_API void SetIsRespectingLocks(bool Value);
 
 	TSharedPtr<FPaletteViewModel> GetPaletteViewModel() { return PaletteViewModel; };
 	TSharedPtr<FLibraryViewModel> GetLibraryViewModel() { return LibraryViewModel; };
 
-	void CreateEditorModeManager() override;
+	UE_API void CreateEditorModeManager() override;
 
 	/** Get the relative info for the Debug mode. */
 	TSharedPtr<UE::UMG::Editor::FPreviewMode> GetPreviewMode() const { return PreviewMode; }
@@ -296,135 +298,135 @@ public:
 	FVector2D PasteDropLocation;
 
 protected:
-	virtual void InitalizeExtenders() override;
-	TSharedPtr<FExtender> CreateMenuExtender();
-	void FillFileMenu(FMenuBuilder& MenuBuilder);
-	void FillAssetMenu(FMenuBuilder& MenuBuilder);
-	void BindToolkitCommands();
-
-	void TakeSnapshot();
-	void CaptureThumbnail();
-	void ClearThumbnail();
-	bool IsImageUsedAsThumbnail();
-	bool IsPreviewWidgetInitialized();
-
-	void CustomizeWidgetCompileOptions(FMenuBuilder& InMenuBuilder);
-
-	static void AddCreateCompileTabSubMenu(FMenuBuilder& InMenuBuilder);
-	static void AddDismissCompileTabSubMenu(FMenuBuilder& InMenuBuilder);
-	static void SetDismissOnCompileSetting(EDisplayOnCompile InDismissOnCompile);
-	static void SetCreateOnCompileSetting(EDisplayOnCompile InCreateOnCompile);
-	static bool IsDismissOnCompileSet(EDisplayOnCompile InDismissOnCompile);
-	static bool IsCreateOnCompileSet(EDisplayOnCompile InCreateOnCompile);
-
-	void OpenCreateNativeBaseClassDialog();
-	void OnCreateNativeBaseClassSuccessfully(const FString& InClassName, const FString& InClassPath, const FString& InModuleName);
-	TSharedPtr<ISequencer> CreateSequencerWidgetInternal();
+	UE_API virtual void InitalizeExtenders() override;
+	UE_API TSharedPtr<FExtender> CreateMenuExtender();
+	UE_API void FillFileMenu(FMenuBuilder& MenuBuilder);
+	UE_API void FillAssetMenu(FMenuBuilder& MenuBuilder);
+	UE_API void BindToolkitCommands();
+
+	UE_API void TakeSnapshot();
+	UE_API void CaptureThumbnail();
+	UE_API void ClearThumbnail();
+	UE_API bool IsImageUsedAsThumbnail();
+	UE_API bool IsPreviewWidgetInitialized();
+
+	UE_API void CustomizeWidgetCompileOptions(FMenuBuilder& InMenuBuilder);
+
+	static UE_API void AddCreateCompileTabSubMenu(FMenuBuilder& InMenuBuilder);
+	static UE_API void AddDismissCompileTabSubMenu(FMenuBuilder& InMenuBuilder);
+	static UE_API void SetDismissOnCompileSetting(EDisplayOnCompile InDismissOnCompile);
+	static UE_API void SetCreateOnCompileSetting(EDisplayOnCompile InCreateOnCompile);
+	static UE_API bool IsDismissOnCompileSet(EDisplayOnCompile InDismissOnCompile);
+	static UE_API bool IsCreateOnCompileSet(EDisplayOnCompile InCreateOnCompile);
+
+	UE_API void OpenCreateNativeBaseClassDialog();
+	UE_API void OnCreateNativeBaseClassSuccessfully(const FString& InClassName, const FString& InClassPath, const FString& InModuleName);
+	UE_API TSharedPtr<ISequencer> CreateSequencerWidgetInternal();
 
 	//~ Begin FBlueprintEditor interface
-	virtual void RegisterApplicationModes(const TArray<UBlueprint*>& InBlueprints, bool bShouldOpenInDefaultsMode, bool bNewlyCreated = false) override;
-	virtual FGraphAppearanceInfo GetGraphAppearance(class UEdGraph* InGraph) const override;
-	virtual TSubclassOf<UEdGraphSchema> GetDefaultSchemaClass() const override;
+	UE_API virtual void RegisterApplicationModes(const TArray<UBlueprint*>& InBlueprints, bool bShouldOpenInDefaultsMode, bool bNewlyCreated = false) override;
+	UE_API virtual FGraphAppearanceInfo GetGraphAppearance(class UEdGraph* InGraph) const override;
+	UE_API virtual TSubclassOf<UEdGraphSchema> GetDefaultSchemaClass() const override;
 	//~ End FBlueprintEditor interface
 
 private:
-	bool CanDeleteSelectedWidgets();
-	void DeleteSelectedWidgets();
+	UE_API bool CanDeleteSelectedWidgets();
+	UE_API void DeleteSelectedWidgets();
 
-	bool CanCopySelectedWidgets();
-	void CopySelectedWidgets();
+	UE_API bool CanCopySelectedWidgets();
+	UE_API void CopySelectedWidgets();
 
-	bool CanPasteWidgets();
-	void PasteWidgets();
+	UE_API bool CanPasteWidgets();
+	UE_API void PasteWidgets();
 
-	bool CanCutSelectedWidgets();
-	void CutSelectedWidgets();
+	UE_API bool CanCutSelectedWidgets();
+	UE_API void CutSelectedWidgets();
 
-	bool CanDuplicateSelectedWidgets();
-	void DuplicateSelectedWidgets();
+	UE_API bool CanDuplicateSelectedWidgets();
+	UE_API void DuplicateSelectedWidgets();
 
-	void OnFindWidgetReferences(bool bSearchAllBlueprints, const EGetFindReferenceSearchStringFlags Flags);
-	bool CanFindWidgetReferences() const;
+	UE_API void OnFindWidgetReferences(bool bSearchAllBlueprints, const EGetFindReferenceSearchStringFlags Flags);
+	UE_API bool CanFindWidgetReferences() const;
 
 	/** Is creating a native base class for the current widget blueprint allowed */
-	bool CanCreateNativeBaseClass() const;
+	UE_API bool CanCreateNativeBaseClass() const;
 
 	/** Whether menu to create a native base class is visible */
-	bool IsCreateNativeBaseClassVisible() const;
+	UE_API bool IsCreateNativeBaseClassVisible() const;
 
 private:
 	/** Called whenever the blueprint is structurally changed. */
-	virtual void OnBlueprintChangedImpl(UBlueprint* InBlueprint, bool bIsJustBeingCompiled = false ) override;
+	UE_API virtual void OnBlueprintChangedImpl(UBlueprint* InBlueprint, bool bIsJustBeingCompiled = false ) override;
 
 	/** Called when objects need to be swapped out for new versions, like after a blueprint recompile. */
-	void OnObjectsReplaced(const TMap<UObject*, UObject*>& ReplacementMap);
+	UE_API void OnObjectsReplaced(const TMap<UObject*, UObject*>& ReplacementMap);
 
 	/** Destroy the current preview GUI object */
-	void DestroyPreview();
+	UE_API void DestroyPreview();
 
 	/** Tick the current preview GUI object */
-	void UpdatePreview(UBlueprint* InBlueprint, bool bInForceFullUpdate);
+	UE_API void UpdatePreview(UBlueprint* InBlueprint, bool bInForceFullUpdate);
 
 	/** Populates the sequencer add menu. */
-	void OnGetAnimationAddMenuContent(FMenuBuilder& MenuBuilder, TSharedRef<ISequencer> Sequencer);
+	UE_API void OnGetAnimationAddMenuContent(FMenuBuilder& MenuBuilder, TSharedRef<ISequencer> Sequencer);
 
 	/** Populates sequencer menu when added with objects. Used to handle case where widget is deleted so it's 
 	    Object * is null.*/
-	void OnBuildCustomContextMenuForGuid(FMenuBuilder& MenuBuilder,FGuid ObjectBinding);
+	UE_API void OnBuildCustomContextMenuForGuid(FMenuBuilder& MenuBuilder,FGuid ObjectBinding);
 
 	/** Populates the sequencer add submenu for the big list of widgets. */
-	void OnGetAnimationAddMenuContentAllWidgets(FMenuBuilder& MenuBuilder);
+	UE_API void OnGetAnimationAddMenuContentAllWidgets(FMenuBuilder& MenuBuilder);
 
 	/** Adds the supplied UObject to the current animation. */
-	void AddObjectToAnimation(UObject* ObjectToAnimate);
+	UE_API void AddObjectToAnimation(UObject* ObjectToAnimate);
 
 	/** Gets the extender to use for sequencers context sensitive menus and toolbars. */
-	TSharedRef<FExtender> GetAddTrackSequencerExtender(const TSharedRef<FUICommandList> CommandList, const TArray<UObject*> ContextSensitiveObjects);
+	UE_API TSharedRef<FExtender> GetAddTrackSequencerExtender(const TSharedRef<FUICommandList> CommandList, const TArray<UObject*> ContextSensitiveObjects);
 
 	/** Extends the sequencer add track menu. */
-	void ExtendSequencerAddTrackMenu( FMenuBuilder& AddTrackMenuBuilder, const TArray<UObject*> ContextObjects );
+	UE_API void ExtendSequencerAddTrackMenu( FMenuBuilder& AddTrackMenuBuilder, const TArray<UObject*> ContextObjects );
 
 	/** Binds additional widgets to a track of the same type */
-	void AddWidgetsToTrack(const TArray<FWidgetReference> Widgets, FGuid ObjectId);
+	UE_API void AddWidgetsToTrack(const TArray<FWidgetReference> Widgets, FGuid ObjectId);
 
 	/** Unbind widgets from a track*/
-	void RemoveWidgetsFromTrack(const TArray<FWidgetReference> Widgets, FGuid ObjectId);
+	UE_API void RemoveWidgetsFromTrack(const TArray<FWidgetReference> Widgets, FGuid ObjectId);
 
 	/** Remove all bindings from a track */
-	void RemoveAllWidgetsFromTrack(FGuid ObjectId);
+	UE_API void RemoveAllWidgetsFromTrack(FGuid ObjectId);
 
 	/** Remove any missing bindings from a track */
-	void RemoveMissingWidgetsFromTrack(FGuid ObjectId);
+	UE_API void RemoveMissingWidgetsFromTrack(FGuid ObjectId);
 	
 	/** Replace current widget bindings on a track with new widget bindings */
-	void ReplaceTrackWithWidgets(const TArray<FWidgetReference> Widgets, FGuid ObjectId);
+	UE_API void ReplaceTrackWithWidgets(const TArray<FWidgetReference> Widgets, FGuid ObjectId);
 
 	/** Dynamic binding */
-	void AddDynamicPossessionMenu(FMenuBuilder& MenuBuilder, FGuid ObjectId);
+	UE_API void AddDynamicPossessionMenu(FMenuBuilder& MenuBuilder, FGuid ObjectId);
 
 	/** Add an animation track for the supplied slot to the current animation. */
-	void AddSlotTrack( UPanelSlot* Slot );
+	UE_API void AddSlotTrack( UPanelSlot* Slot );
 
 	/** Add an animation track for the supplied slot to the current animation. */
 	void AddSlotTrack(TObjectPtr<UPanelSlot> Slot) { return AddSlotTrack(Slot.Get()); }
 
 	/** Add an animation track for the supplied material property path to the current animation. */
-	void AddMaterialTrack( UWidget* Widget, TArray<FProperty*> MaterialPropertyPath, FText MaterialPropertyDisplayName );
+	UE_API void AddMaterialTrack( UWidget* Widget, TArray<FProperty*> MaterialPropertyPath, FText MaterialPropertyDisplayName );
 
 	/** Handler which is called whenever sequencer movie scene data changes. */
-	void OnMovieSceneDataChanged(EMovieSceneDataChangeType DataChangeType);
+	UE_API void OnMovieSceneDataChanged(EMovieSceneDataChangeType DataChangeType);
 
 	/** Handler which is called whenever sequencer binding is pasted. */
-	void OnMovieSceneBindingsPasted(const TArray<FMovieSceneBinding>& BindingsPasted);
+	UE_API void OnMovieSceneBindingsPasted(const TArray<FMovieSceneBinding>& BindingsPasted);
 
 	/** Fire off when sequencer selection changed */
-	void SyncSelectedWidgetsWithSequencerSelection(TArray<FGuid> ObjectGuids);
+	UE_API void SyncSelectedWidgetsWithSequencerSelection(TArray<FGuid> ObjectGuids);
 
 	/** Tell sequencer the selected widgets changed */
-	void SyncSequencerSelectionToSelectedWidgets();
+	UE_API void SyncSequencerSelectionToSelectedWidgets();
 
 	/** Tell sequencers movie data changed */
-	void SyncSequencersMovieSceneData();
+	UE_API void SyncSequencersMovieSceneData();
 
 	/** Get the animation playback context */
 	UObject* GetAnimationPlaybackContext() const { return GetPreview(); }
@@ -433,10 +435,10 @@ private:
 	TArray<UObject*> GetAnimationEventContexts() const { TArray<UObject*> EventContexts; EventContexts.Add(GetPreview()); return EventContexts; }
 	
 	/** Update the name of a track to reflect changes in bindings */
-	void UpdateTrackName(FGuid ObjectId);
+	UE_API void UpdateTrackName(FGuid ObjectId);
 
 	/* Called after changing dynamic binding properties*/
-	void OnFinishedChangingDynamicBindingProperties(const FPropertyChangedEvent& ChangeEvent, TSharedPtr<FStructOnScope> ValueStruct, FGuid ObjectId);
+	UE_API void OnFinishedChangingDynamicBindingProperties(const FPropertyChangedEvent& ChangeEvent, TSharedPtr<FStructOnScope> ValueStruct, FGuid ObjectId);
 
 private:
 	/** The preview scene that owns the preview GUI */
@@ -549,4 +551,6 @@ private:
 
 	/** True when the selected object is the preview widget. Used to restore selection after refreshing the preview.*/
 	bool bIsPreviewWidgetSelected;
-};
\ No newline at end of file
+};
+
+#undef UE_API
diff --git a/Engine/Source/Editor/UMGEditor/Public/WidgetBlueprintEditorUtils.h b/Engine/Source/Editor/UMGEditor/Public/WidgetBlueprintEditorUtils.h
--- a/Engine/Source/Editor/UMGEditor/Public/WidgetBlueprintEditorUtils.h
+++ b/Engine/Source/Editor/UMGEditor/Public/WidgetBlueprintEditorUtils.h
@@ -8,6 +8,8 @@
 #include "WidgetReference.h"
 #include "Misc/StringOutputDevice.h"
 
+#define UE_API UMGEDITOR_API
+
 class FDragDropOperation;
 class FHittestGrid;
 class FMenuBuilder;
@@ -22,7 +24,7 @@ class ULocalPlayer;
 //////////////////////////////////////////////////////////////////////////
 // FWidgetBlueprintEditorUtils
 
-class UMGEDITOR_API FWidgetBlueprintEditorUtils
+class FWidgetBlueprintEditorUtils
 {
 
 public:
@@ -38,21 +40,21 @@ public:
 		ULocalPlayer* LocalPlayer = nullptr; // Optionally specify if available
 	};
 
-	static bool VerifyWidgetRename(TSharedRef<class FWidgetBlueprintEditor> BlueprintEditor, FWidgetReference Widget, const FText& NewName, FText& OutErrorMessage);
+	static UE_API bool VerifyWidgetRename(TSharedRef<class FWidgetBlueprintEditor> BlueprintEditor, FWidgetReference Widget, const FText& NewName, FText& OutErrorMessage);
 
-	static bool RenameWidget(TSharedRef<class FWidgetBlueprintEditor> BlueprintEditor, const FName& OldObjectName, const FString& NewDisplayName);
+	static UE_API bool RenameWidget(TSharedRef<class FWidgetBlueprintEditor> BlueprintEditor, const FName& OldObjectName, const FString& NewDisplayName);
 
-	static void ReplaceDesiredFocus(TSharedRef<FWidgetBlueprintEditor> BlueprintEditor, const FName& OldName, const FName& NewName);
-	static void ReplaceDesiredFocus(UWidgetBlueprint* Blueprint, const FName& OldName, const FName& NewName);
+	static UE_API void ReplaceDesiredFocus(TSharedRef<FWidgetBlueprintEditor> BlueprintEditor, const FName& OldName, const FName& NewName);
+	static UE_API void ReplaceDesiredFocus(UWidgetBlueprint* Blueprint, const FName& OldName, const FName& NewName);
 	
-	static void SetDesiredFocus(TSharedRef<FWidgetBlueprintEditor> BlueprintEditor, const FName DesiredFocusWidgetName);
-	static void SetDesiredFocus(UWidgetBlueprint* Blueprint, const FName& DesiredFocusWidgetName);
+	static UE_API void SetDesiredFocus(TSharedRef<FWidgetBlueprintEditor> BlueprintEditor, const FName DesiredFocusWidgetName);
+	static UE_API void SetDesiredFocus(UWidgetBlueprint* Blueprint, const FName& DesiredFocusWidgetName);
 
-	static void CreateWidgetContextMenu(FMenuBuilder& MenuBuilder, TSharedRef<FWidgetBlueprintEditor> BlueprintEditor, FVector2D TargetLocation);
+	static UE_API void CreateWidgetContextMenu(FMenuBuilder& MenuBuilder, TSharedRef<FWidgetBlueprintEditor> BlueprintEditor, FVector2D TargetLocation);
 
-	static void CopyWidgets(UWidgetBlueprint* BP, TSet<FWidgetReference> Widgets);
+	static UE_API void CopyWidgets(UWidgetBlueprint* BP, TSet<FWidgetReference> Widgets);
 
-	static TArray<UWidget*> PasteWidgets(TSharedRef<FWidgetBlueprintEditor> BlueprintEditor, UWidgetBlueprint* BP, FWidgetReference ParentWidget, FName SlotName, FVector2D PasteLocation);
+	static UE_API TArray<UWidget*> PasteWidgets(TSharedRef<FWidgetBlueprintEditor> BlueprintEditor, UWidgetBlueprint* BP, FWidgetReference ParentWidget, FName SlotName, FVector2D PasteLocation);
 
 	enum class EReplaceWidgetNamingMethod
 	{
@@ -66,10 +68,10 @@ public:
 		// Will use the new widget's generated name and not make an effort to maintain references
 		UseNewGeneratedName
 	};
-	static void ReplaceWidgets(UWidgetBlueprint* BP, TSet<UWidget*> Widgets, UClass* WidgetClass, EReplaceWidgetNamingMethod NewWidgetNamingMethod);
+	static UE_API void ReplaceWidgets(UWidgetBlueprint* BP, TSet<UWidget*> Widgets, UClass* WidgetClass, EReplaceWidgetNamingMethod NewWidgetNamingMethod);
 
 	UE_DEPRECATED(5.6, "FWidgetBlueprintEditorUtils::DeleteWidgets no longer takes in the BlueprintEditor, use the version without it instead")
-	static void DeleteWidgets(TSharedRef<FWidgetBlueprintEditor> BlueprintEditor, UWidgetBlueprint* BP, TSet<FWidgetReference> Widgets, bool bSilentDelete = false);
+	static UE_API void DeleteWidgets(TSharedRef<FWidgetBlueprintEditor> BlueprintEditor, UWidgetBlueprint* BP, TSet<FWidgetReference> Widgets, bool bSilentDelete = false);
 
 	enum class EDeleteWidgetWarningType
 	{
@@ -78,30 +80,30 @@ public:
 		// Don't notify the user at all and delete the widget even if it is referenced
 		DeleteSilently,
 	};
-	static void DeleteWidgets(UWidgetBlueprint* BP, TSet<UWidget*> Widgets, EDeleteWidgetWarningType WarningType);
+	static UE_API void DeleteWidgets(UWidgetBlueprint* BP, TSet<UWidget*> Widgets, EDeleteWidgetWarningType WarningType);
 
-	static void CutWidgets(TSharedRef<FWidgetBlueprintEditor> BlueprintEditor, UWidgetBlueprint* BP, TSet<FWidgetReference> Widgets);
+	static UE_API void CutWidgets(TSharedRef<FWidgetBlueprintEditor> BlueprintEditor, UWidgetBlueprint* BP, TSet<FWidgetReference> Widgets);
 
-	static TArray<UWidget*> DuplicateWidgets(TSharedRef<FWidgetBlueprintEditor> BlueprintEditor, UWidgetBlueprint* BP, TSet<FWidgetReference> Widgets);
+	static UE_API TArray<UWidget*> DuplicateWidgets(TSharedRef<FWidgetBlueprintEditor> BlueprintEditor, UWidgetBlueprint* BP, TSet<FWidgetReference> Widgets);
 
-	static UUserWidget* CreateUserWidgetFromBlueprint(UObject* Outer, UWidgetBlueprint* BP, const FCreateWidgetFromBlueprintParams& Params);
+	static UE_API UUserWidget* CreateUserWidgetFromBlueprint(UObject* Outer, UWidgetBlueprint* BP, const FCreateWidgetFromBlueprintParams& Params);
 
 	/** Performs cleanup on the specified UserWidget. */
-	static void DestroyUserWidget(UUserWidget* UserWidget);
+	static UE_API void DestroyUserWidget(UUserWidget* UserWidget);
 
-	static bool IsAnySelectedWidgetLocked(TSet<FWidgetReference> SelectedWidgets);
+	static UE_API bool IsAnySelectedWidgetLocked(TSet<FWidgetReference> SelectedWidgets);
 
-	static bool CanPasteWidgetsExtension(TSet<FWidgetReference> SelectedWidgets);
+	static UE_API bool CanPasteWidgetsExtension(TSet<FWidgetReference> SelectedWidgets);
 
-	static UWidget* GetWidgetTemplateFromDragDrop(UWidgetBlueprint* Blueprint, UWidgetTree* RootWidgetTree, TSharedPtr<FDragDropOperation>& DragDropOp);
+	static UE_API UWidget* GetWidgetTemplateFromDragDrop(UWidgetBlueprint* Blueprint, UWidgetTree* RootWidgetTree, TSharedPtr<FDragDropOperation>& DragDropOp);
 
-	static bool ShouldPreventDropOnTargetExtensions(const UWidget* Target, const TSharedPtr<FDragDropOperation>& DragDropOp, FText& OutFailureText);
+	static UE_API bool ShouldPreventDropOnTargetExtensions(const UWidget* Target, const TSharedPtr<FDragDropOperation>& DragDropOp, FText& OutFailureText);
 
-	static bool IsBindWidgetProperty(const FProperty* InProperty);
-	static bool IsBindWidgetProperty(const FProperty* InProperty, bool& bIsOptional);
+	static UE_API bool IsBindWidgetProperty(const FProperty* InProperty);
+	static UE_API bool IsBindWidgetProperty(const FProperty* InProperty, bool& bIsOptional);
 
-	static bool IsBindWidgetAnimProperty(const FProperty* InProperty);
-	static bool IsBindWidgetAnimProperty(const FProperty* InProperty, bool& bIsOptional);
+	static UE_API bool IsBindWidgetAnimProperty(const FProperty* InProperty);
+	static UE_API bool IsBindWidgetAnimProperty(const FProperty* InProperty, bool& bIsOptional);
 
 	struct FUsableWidgetClassResult
 	{
@@ -110,107 +112,109 @@ public:
 	};
 
 	UE_DEPRECATED(5.3, "This function has been deprecated. Use the version of IsUsableWidgetClass that takes a second argument of TSharedRef<FWidgetBlueprintEditor>.")
-	static bool IsUsableWidgetClass(const UClass* WidgetClass);
+	static UE_API bool IsUsableWidgetClass(const UClass* WidgetClass);
 	UE_DEPRECATED(5.3, "This function has been deprecated. Use the version of IsUsableWidgetClass that takes a second argument of TSharedRef<FWidgetBlueprintEditor>.")
-	static TValueOrError<FUsableWidgetClassResult, void> IsUsableWidgetClass(const FAssetData& WidgetAsset);
+	static UE_API TValueOrError<FUsableWidgetClassResult, void> IsUsableWidgetClass(const FAssetData& WidgetAsset);
 
-	static bool IsUsableWidgetClass(const UClass* WidgetClass, TSharedRef<FWidgetBlueprintEditor> InCurrentActiveBlueprintEditor);
-	static TValueOrError<FUsableWidgetClassResult, void> IsUsableWidgetClass(const FAssetData& WidgetAsset, TSharedRef<FWidgetBlueprintEditor> InCurrentActiveBlueprintEditor);
+	static UE_API bool IsUsableWidgetClass(const UClass* WidgetClass, TSharedRef<FWidgetBlueprintEditor> InCurrentActiveBlueprintEditor);
+	static UE_API TValueOrError<FUsableWidgetClassResult, void> IsUsableWidgetClass(const FAssetData& WidgetAsset, TSharedRef<FWidgetBlueprintEditor> InCurrentActiveBlueprintEditor);
 
-	static void ExportWidgetsToText(TArray<UWidget*> WidgetsToExport, /*out*/ FString& ExportedText);
+	static UE_API void ExportWidgetsToText(TArray<UWidget*> WidgetsToExport, /*out*/ FString& ExportedText);
 
-	static void ImportWidgetsFromText(UWidgetBlueprint* BP, const FString& TextToImport, /*out*/ TSet<UWidget*>& ImportedWidgetSet, /*out*/ TMap<FName, UWidgetSlotPair*>& PastedExtraSlotData);
+	static UE_API void ImportWidgetsFromText(UWidgetBlueprint* BP, const FString& TextToImport, /*out*/ TSet<UWidget*>& ImportedWidgetSet, /*out*/ TMap<FName, UWidgetSlotPair*>& PastedExtraSlotData);
 
 	/** Exports the individual properties of an object to text and stores them in a map. */
-	static void ExportPropertiesToText(UObject* Object, TMap<FName, FString>& ExportedProperties);
+	static UE_API void ExportPropertiesToText(UObject* Object, TMap<FName, FString>& ExportedProperties);
 
 	/** Attempts to import any property in the map and apply it to a property with the same name on the object. */
-	static void ImportPropertiesFromText(UObject* Object, const TMap<FName, FString>& ExportedProperties);
+	static UE_API void ImportPropertiesFromText(UObject* Object, const TMap<FName, FString>& ExportedProperties);
 
-	static bool DoesClipboardTextContainWidget(UWidgetBlueprint* BP);
+	static UE_API bool DoesClipboardTextContainWidget(UWidgetBlueprint* BP);
 
-	static TScriptInterface<INamedSlotInterface> FindNamedSlotHostForContent(UWidget* WidgetTemplate, UWidgetTree* WidgetTree);
+	static UE_API TScriptInterface<INamedSlotInterface> FindNamedSlotHostForContent(UWidget* WidgetTemplate, UWidgetTree* WidgetTree);
 
-	static UWidget* FindNamedSlotHostWidgetForContent(UWidget* WidgetTemplate, UWidgetTree* WidgetTree);
+	static UE_API UWidget* FindNamedSlotHostWidgetForContent(UWidget* WidgetTemplate, UWidgetTree* WidgetTree);
 
-	static void FindAllAncestorNamedSlotHostWidgetsForContent(TArray<FWidgetReference>& OutSlotHostWidgets, UWidget* WidgetTemplate, TSharedRef<FWidgetBlueprintEditor> BlueprintEditor);
+	static UE_API void FindAllAncestorNamedSlotHostWidgetsForContent(TArray<FWidgetReference>& OutSlotHostWidgets, UWidget* WidgetTemplate, TSharedRef<FWidgetBlueprintEditor> BlueprintEditor);
 
-	static bool RemoveNamedSlotHostContent(UWidget* WidgetTemplate, TScriptInterface<INamedSlotInterface> NamedSlotHost);
+	static UE_API bool RemoveNamedSlotHostContent(UWidget* WidgetTemplate, TScriptInterface<INamedSlotInterface> NamedSlotHost);
 
-	static bool ReplaceNamedSlotHostContent(UWidget* WidgetTemplate, TScriptInterface<INamedSlotInterface> NamedSlotHost, UWidget* NewContentWidget);
+	static UE_API bool ReplaceNamedSlotHostContent(UWidget* WidgetTemplate, TScriptInterface<INamedSlotInterface> NamedSlotHost, UWidget* NewContentWidget);
 
-	static UWidgetTree* FindLatestWidgetTree(UWidgetBlueprint* Blueprint, UUserWidget* UserWidget);
+	static UE_API UWidgetTree* FindLatestWidgetTree(UWidgetBlueprint* Blueprint, UUserWidget* UserWidget);
 
-	static int32 UpdateHittestGrid(FHittestGrid& HitTestGrid, TSharedRef<SWindow> Window, float Scale, FVector2D DrawSize, float DeltaTime);
+	static UE_API int32 UpdateHittestGrid(FHittestGrid& HitTestGrid, TSharedRef<SWindow> Window, float Scale, FVector2D DrawSize, float DeltaTime);
 
-	static TTuple<FVector2D, FVector2D> GetWidgetPreviewAreaAndSize(UUserWidget* UserWidget, FVector2D DesiredSize, FVector2D PreviewSize, EDesignPreviewSizeMode ThumbnailSizeMode, TOptional<FVector2D> ThumbnailCustomSize);
+	static UE_API TTuple<FVector2D, FVector2D> GetWidgetPreviewAreaAndSize(UUserWidget* UserWidget, FVector2D DesiredSize, FVector2D PreviewSize, EDesignPreviewSizeMode ThumbnailSizeMode, TOptional<FVector2D> ThumbnailCustomSize);
 
-	static float GetWidgetPreviewDPIScale(UUserWidget* UserWidget, FVector2D PreviewSize);
+	static UE_API float GetWidgetPreviewDPIScale(UUserWidget* UserWidget, FVector2D PreviewSize);
 
-	static EDesignPreviewSizeMode ConvertThumbnailSizeModeToDesignerSizeMode(EThumbnailPreviewSizeMode ThumbnailSizeMode, UUserWidget* WidgetInstance);
+	static UE_API EDesignPreviewSizeMode ConvertThumbnailSizeModeToDesignerSizeMode(EThumbnailPreviewSizeMode ThumbnailSizeMode, UUserWidget* WidgetInstance);
 
-	static FVector2D GetWidgetPreviewUnScaledCustomSize(FVector2D DesiredSize, UUserWidget* UserWidget, TOptional<FVector2D> ThumbnailCustomSize, EThumbnailPreviewSizeMode ThumbnailSizeMode = EThumbnailPreviewSizeMode::MatchDesignerMode);
+	static UE_API FVector2D GetWidgetPreviewUnScaledCustomSize(FVector2D DesiredSize, UUserWidget* UserWidget, TOptional<FVector2D> ThumbnailCustomSize, EThumbnailPreviewSizeMode ThumbnailSizeMode = EThumbnailPreviewSizeMode::MatchDesignerMode);
 
-	static TTuple<float, FVector2D> GetThumbnailImageScaleAndOffset(FVector2D WidgetSize, FVector2D ThumbnailSize);
+	static UE_API TTuple<float, FVector2D> GetThumbnailImageScaleAndOffset(FVector2D WidgetSize, FVector2D ThumbnailSize);
 
-	static void SetTextureAsAssetThumbnail(UWidgetBlueprint* WidgetBlueprint, UTexture2D* ThumbnailTexture);
+	static UE_API void SetTextureAsAssetThumbnail(UWidgetBlueprint* WidgetBlueprint, UTexture2D* ThumbnailTexture);
 
-	static FText GetPaletteCategory(const TSubclassOf<UWidget> Widget);
-	static FText GetPaletteCategory(const FAssetData& WidgetAsset, const TSubclassOf<UWidget> NativeClass);
+	static UE_API FText GetPaletteCategory(const TSubclassOf<UWidget> Widget);
+	static UE_API FText GetPaletteCategory(const FAssetData& WidgetAsset, const TSubclassOf<UWidget> NativeClass);
 
-	static TOptional<FWidgetThumbnailProperties> DrawSWidgetInRenderTargetForThumbnail(UUserWidget* WidgetInstance, FRenderTarget* RenderTarget2D, FVector2D ThumbnailSize, TOptional<FVector2D> ThumbnailCustomSize, EThumbnailPreviewSizeMode ThumbnailSizeMode = EThumbnailPreviewSizeMode::MatchDesignerMode);
+	static UE_API TOptional<FWidgetThumbnailProperties> DrawSWidgetInRenderTargetForThumbnail(UUserWidget* WidgetInstance, FRenderTarget* RenderTarget2D, FVector2D ThumbnailSize, TOptional<FVector2D> ThumbnailCustomSize, EThumbnailPreviewSizeMode ThumbnailSizeMode = EThumbnailPreviewSizeMode::MatchDesignerMode);
 
-	static TOptional<FWidgetThumbnailProperties> DrawSWidgetInRenderTargetForThumbnail(UUserWidget* WidgetInstance, UTextureRenderTarget2D* RenderTarget2D, FVector2D ThumbnailSize, TOptional<FVector2D> ThumbnailCustomSize, EThumbnailPreviewSizeMode ThumbnailSizeMode);
+	static UE_API TOptional<FWidgetThumbnailProperties> DrawSWidgetInRenderTargetForThumbnail(UUserWidget* WidgetInstance, UTextureRenderTarget2D* RenderTarget2D, FVector2D ThumbnailSize, TOptional<FVector2D> ThumbnailCustomSize, EThumbnailPreviewSizeMode ThumbnailSizeMode);
 
-	static TOptional<FWidgetThumbnailProperties> DrawSWidgetInRenderTarget(UUserWidget* WidgetInstance, UTextureRenderTarget2D* RenderTarget2D);
+	static UE_API TOptional<FWidgetThumbnailProperties> DrawSWidgetInRenderTarget(UUserWidget* WidgetInstance, UTextureRenderTarget2D* RenderTarget2D);
 
-	static UWidgetEditingProjectSettings* GetRelevantMutableSettings(TWeakPtr<FWidgetBlueprintEditor> CurrentEditor);
-	static const UWidgetEditingProjectSettings* GetRelevantSettings(TWeakPtr<FWidgetBlueprintEditor> CurrentEditor);
+	static UE_API UWidgetEditingProjectSettings* GetRelevantMutableSettings(TWeakPtr<FWidgetBlueprintEditor> CurrentEditor);
+	static UE_API const UWidgetEditingProjectSettings* GetRelevantSettings(TWeakPtr<FWidgetBlueprintEditor> CurrentEditor);
 
-	static UWidgetBlueprint* GetWidgetBlueprintFromWidget(const UWidget* Widget);
+	static UE_API UWidgetBlueprint* GetWidgetBlueprintFromWidget(const UWidget* Widget);
 	
-	static TSet<UWidget*> ResolveWidgetTemplates(const TSet<FWidgetReference>& Widgets);
+	static UE_API TSet<UWidget*> ResolveWidgetTemplates(const TSet<FWidgetReference>& Widgets);
 
 private:
 
-	static FString CopyWidgetsInternal(UWidgetBlueprint* BP, TSet<FWidgetReference> Widgets);
+	static UE_API FString CopyWidgetsInternal(UWidgetBlueprint* BP, TSet<FWidgetReference> Widgets);
 
-	static TArray<UWidget*> PasteWidgetsInternal(TSharedRef<FWidgetBlueprintEditor> BlueprintEditor, UWidgetBlueprint* BP, const FString& TextToImport, FWidgetReference ParentWidget, FName SlotName, FVector2D PasteLocation, bool bForceSibling, bool& TransactionSuccesful);
+	static UE_API TArray<UWidget*> PasteWidgetsInternal(TSharedRef<FWidgetBlueprintEditor> BlueprintEditor, UWidgetBlueprint* BP, const FString& TextToImport, FWidgetReference ParentWidget, FName SlotName, FVector2D PasteLocation, bool bForceSibling, bool& TransactionSuccesful);
 
-	static bool DisplayPasteWarningAndEarlyExit();
+	static UE_API bool DisplayPasteWarningAndEarlyExit();
 
-	static void ExecuteOpenSelectedWidgetsForEdit( TSet<FWidgetReference> SelectedWidgets );
+	static UE_API void ExecuteOpenSelectedWidgetsForEdit( TSet<FWidgetReference> SelectedWidgets );
 
-	static bool FindAndRemoveNamedSlotContent(UWidget* WidgetTemplate, UWidgetTree* WidgetTree);
+	static UE_API bool FindAndRemoveNamedSlotContent(UWidget* WidgetTemplate, UWidgetTree* WidgetTree);
 
-	static bool CanOpenSelectedWidgetsForEdit( TSet<FWidgetReference> SelectedWidgets );
+	static UE_API bool CanOpenSelectedWidgetsForEdit( TSet<FWidgetReference> SelectedWidgets );
 
-	static void BuildWrapWithMenu(FMenuBuilder& Menu, TSharedRef<FWidgetBlueprintEditor> BlueprintEditor, UWidgetBlueprint* BP, TSet<FWidgetReference> Widgets);
+	static UE_API void BuildWrapWithMenu(FMenuBuilder& Menu, TSharedRef<FWidgetBlueprintEditor> BlueprintEditor, UWidgetBlueprint* BP, TSet<FWidgetReference> Widgets);
 
-	static void WrapWidgets(TSharedRef<FWidgetBlueprintEditor> BlueprintEditor, UWidgetBlueprint* BP, TSet<FWidgetReference> Widgets, UClass* WidgetClass);
+	static UE_API void WrapWidgets(TSharedRef<FWidgetBlueprintEditor> BlueprintEditor, UWidgetBlueprint* BP, TSet<FWidgetReference> Widgets, UClass* WidgetClass);
 
-	static void BuildReplaceWithMenu(FMenuBuilder& Menu, TSharedRef<FWidgetBlueprintEditor> BlueprintEditor, UWidgetBlueprint* BP, TSet<FWidgetReference> Widgets);
+	static UE_API void BuildReplaceWithMenu(FMenuBuilder& Menu, TSharedRef<FWidgetBlueprintEditor> BlueprintEditor, UWidgetBlueprint* BP, TSet<FWidgetReference> Widgets);
 
-	static void ReplaceWidgetWithSelectedTemplate(TSharedRef<FWidgetBlueprintEditor> BlueprintEditor, UWidgetBlueprint* BP, FWidgetReference Widget);
+	static UE_API void ReplaceWidgetWithSelectedTemplate(TSharedRef<FWidgetBlueprintEditor> BlueprintEditor, UWidgetBlueprint* BP, FWidgetReference Widget);
 
-	static bool CanBeReplacedWithTemplate(TSharedRef<FWidgetBlueprintEditor> BlueprintEditor, UWidgetBlueprint* BP, FWidgetReference Widget);
+	static UE_API bool CanBeReplacedWithTemplate(TSharedRef<FWidgetBlueprintEditor> BlueprintEditor, UWidgetBlueprint* BP, FWidgetReference Widget);
 
-	static void ReplaceWidgetWithChildren(TSharedRef<FWidgetBlueprintEditor> BlueprintEditor, UWidgetBlueprint* BP, FWidgetReference Widget);
+	static UE_API void ReplaceWidgetWithChildren(TSharedRef<FWidgetBlueprintEditor> BlueprintEditor, UWidgetBlueprint* BP, FWidgetReference Widget);
 
-	static void ReplaceWidgetWithNamedSlot(TSharedRef<FWidgetBlueprintEditor> BlueprintEditor, UWidgetBlueprint* BP, FWidgetReference Widget, FName NamedSlot);
+	static UE_API void ReplaceWidgetWithNamedSlot(TSharedRef<FWidgetBlueprintEditor> BlueprintEditor, UWidgetBlueprint* BP, FWidgetReference Widget, FName NamedSlot);
 
-	static FString FindNextValidName(UWidgetTree* WidgetTree, const FString& Name);
+	static UE_API FString FindNextValidName(UWidgetTree* WidgetTree, const FString& Name);
 
-	static void FindUsedVariablesForWidgets(const TSet<UWidget*>& Widgets, const UWidgetBlueprint* BP, TArray<UWidget*>& UsedVariables, TArray<FText>& WidgetNames, bool bIncludeVariablesOnChildren);
+	static UE_API void FindUsedVariablesForWidgets(const TSet<UWidget*>& Widgets, const UWidgetBlueprint* BP, TArray<UWidget*>& UsedVariables, TArray<FText>& WidgetNames, bool bIncludeVariablesOnChildren);
 
-	static bool ShouldContinueDeleteOperation(UWidgetBlueprint* BP, const TArray<FText>& WidgetNames);
+	static UE_API bool ShouldContinueDeleteOperation(UWidgetBlueprint* BP, const TArray<FText>& WidgetNames);
 
-	static bool ShouldContinueReplaceOperation(UWidgetBlueprint* BP, const TArray<FText>& WidgetNames);	
+	static UE_API bool ShouldContinueReplaceOperation(UWidgetBlueprint* BP, const TArray<FText>& WidgetNames);	
 
-	static TOptional<FWidgetThumbnailProperties> DrawSWidgetInRenderTargetInternal(UUserWidget* WidgetInstance, FRenderTarget* RenderTarget2D, UTextureRenderTarget2D* TextureRenderTarget,FVector2D ThumbnailSize, bool bIsForThumbnail, TOptional<FVector2D> ThumbnailCustomSize, EThumbnailPreviewSizeMode ThumbnailSizeMode);
+	static UE_API TOptional<FWidgetThumbnailProperties> DrawSWidgetInRenderTargetInternal(UUserWidget* WidgetInstance, FRenderTarget* RenderTarget2D, UTextureRenderTarget2D* TextureRenderTarget,FVector2D ThumbnailSize, bool bIsForThumbnail, TOptional<FVector2D> ThumbnailCustomSize, EThumbnailPreviewSizeMode ThumbnailSizeMode);
 	
-	static bool IsDesiredFocusWidget(TSharedRef<FWidgetBlueprintEditor> BlueprintEditor, UWidget* Widget);
-	static bool IsDesiredFocusWidget(UWidgetBlueprint* Blueprint, UWidget* Widget);
+	static UE_API bool IsDesiredFocusWidget(TSharedRef<FWidgetBlueprintEditor> BlueprintEditor, UWidget* Widget);
+	static UE_API bool IsDesiredFocusWidget(UWidgetBlueprint* Blueprint, UWidget* Widget);
 
-	static FWidgetObjectTextFactory ProcessImportedText(UWidgetBlueprint* BP, const FString& TextToImport, /*out*/ UPackage*& TempPackage);
+	static UE_API FWidgetObjectTextFactory ProcessImportedText(UWidgetBlueprint* BP, const FString& TextToImport, /*out*/ UPackage*& TempPackage);
 };
+
+#undef UE_API
diff --git a/Engine/Source/Editor/UMGEditor/Public/WidgetBlueprintExtension.h b/Engine/Source/Editor/UMGEditor/Public/WidgetBlueprintExtension.h
--- a/Engine/Source/Editor/UMGEditor/Public/WidgetBlueprintExtension.h
+++ b/Engine/Source/Editor/UMGEditor/Public/WidgetBlueprintExtension.h
@@ -13,9 +13,11 @@ class UWidgetBlueprintGeneratedClass;
 
 #include "WidgetBlueprintExtension.generated.h"
 
+#define UE_API UMGEDITOR_API
+
 /** Extension that allows per-system data to be held on the widget blueprint, and per-system logic to be executed during compilation */
-UCLASS(Within=WidgetBlueprint)
-class UMGEDITOR_API UWidgetBlueprintExtension : public UBlueprintExtension
+UCLASS(MinimalAPI, Within=WidgetBlueprint)
+class UWidgetBlueprintExtension : public UBlueprintExtension
 {
 	GENERATED_BODY()
 
@@ -34,7 +36,7 @@ public:
 	 * Request an WidgetBlueprintExtension for an WidgetBlueprint.
 	 * @note It is illegal to perform this operation once compilation has commenced, use GetExtension instead.
 	 */
-	static UWidgetBlueprintExtension* RequestExtension(UWidgetBlueprint* InWidgetBlueprint, TSubclassOf<UWidgetBlueprintExtension> InExtensionType);
+	static UE_API UWidgetBlueprintExtension* RequestExtension(UWidgetBlueprint* InWidgetBlueprint, TSubclassOf<UWidgetBlueprintExtension> InExtensionType);
 
 	/** Get an already-requested extension for an WidgetBlueprint. */
 	template<typename ExtensionType>
@@ -44,10 +46,10 @@ public:
 	}
 
 	/** Get an already-requested extension for an WidgetBlueprint. */
-	static UWidgetBlueprintExtension* GetExtension(const UWidgetBlueprint* InWidgetBlueprint, TSubclassOf<UWidgetBlueprintExtension> InExtensionType);
+	static UE_API UWidgetBlueprintExtension* GetExtension(const UWidgetBlueprint* InWidgetBlueprint, TSubclassOf<UWidgetBlueprintExtension> InExtensionType);
 
 	/** Get all subsystems currently present on an WidgetBlueprint */
-	static TArray<UWidgetBlueprintExtension*> GetExtensions(const UWidgetBlueprint* InWidgetBlueprint);
+	static UE_API TArray<UWidgetBlueprintExtension*> GetExtensions(const UWidgetBlueprint* InWidgetBlueprint);
 
 	/** Iterate over all registered WidgetBlueprintExtensions in an WidgetBlueprint */
 	template<typename Predicate>
@@ -63,7 +65,7 @@ public:
 	}
 
 	/** Get the WidgetBlueprint that hosts this extension */
-	UWidgetBlueprint* GetWidgetBlueprint() const;
+	UE_API UWidgetBlueprint* GetWidgetBlueprint() const;
 
 #if WITH_EDITOR
 	virtual void FindDiffs(const UWidgetBlueprint* OwningBlueprint, const UWidgetBlueprint* OtherBlueprint, FDiffResults& Results) const {}
@@ -145,3 +147,5 @@ private:
 		HandleEndCompilation();
 	}
 };
+
+#undef UE_API
diff --git a/Engine/Source/Editor/UMGEditor/Public/WidgetBlueprintThumbnailRenderer.h b/Engine/Source/Editor/UMGEditor/Public/WidgetBlueprintThumbnailRenderer.h
--- a/Engine/Source/Editor/UMGEditor/Public/WidgetBlueprintThumbnailRenderer.h
+++ b/Engine/Source/Editor/UMGEditor/Public/WidgetBlueprintThumbnailRenderer.h
@@ -6,27 +6,29 @@
 #include "ThumbnailRendering/DefaultSizedThumbnailRenderer.h"
 #include "WidgetBlueprintThumbnailRenderer.generated.h"
 
+#define UE_API UMGEDITOR_API
+
 class FCanvas;
 class FRenderTarget;
 class UTextureRenderTarget2D;
 class FWidgetBlueprintThumbnailPool;
 
-UCLASS()
-class UMGEDITOR_API UWidgetBlueprintThumbnailRenderer : public UDefaultSizedThumbnailRenderer
+UCLASS(MinimalAPI)
+class UWidgetBlueprintThumbnailRenderer : public UDefaultSizedThumbnailRenderer
 {
 	GENERATED_BODY()
 
-	UWidgetBlueprintThumbnailRenderer();
-	virtual ~UWidgetBlueprintThumbnailRenderer();
+	UE_API UWidgetBlueprintThumbnailRenderer();
+	UE_API virtual ~UWidgetBlueprintThumbnailRenderer();
 
 	//~ Begin UThumbnailRenderer Object
-	virtual bool CanVisualizeAsset(UObject* Object) override;
-	virtual void Draw(UObject* Object, int32 X, int32 Y, uint32 Width, uint32 Height, FRenderTarget* RenderTarget, FCanvas* Canvas, bool bAdditionalViewFamily) override;
-	virtual EThumbnailRenderFrequency GetThumbnailRenderFrequency(UObject* Object) const override;
+	UE_API virtual bool CanVisualizeAsset(UObject* Object) override;
+	UE_API virtual void Draw(UObject* Object, int32 X, int32 Y, uint32 Width, uint32 Height, FRenderTarget* RenderTarget, FCanvas* Canvas, bool bAdditionalViewFamily) override;
+	UE_API virtual EThumbnailRenderFrequency GetThumbnailRenderFrequency(UObject* Object) const override;
 	//~ End UThumbnailRenderer Object
 
 private:
-	void OnBlueprintUnloaded(UBlueprint* Blueprint);
+	UE_API void OnBlueprintUnloaded(UBlueprint* Blueprint);
 
 private:
 	struct FWidgetBlueprintThumbnailPoolDeleter
@@ -36,3 +38,5 @@ private:
 
 	TUniquePtr<FWidgetBlueprintThumbnailPool, FWidgetBlueprintThumbnailPoolDeleter> ThumbnailPool;
 };
+
+#undef UE_API
diff --git a/Engine/Source/Editor/UMGEditor/Public/WidgetBlueprintToolMenuContext.h b/Engine/Source/Editor/UMGEditor/Public/WidgetBlueprintToolMenuContext.h
--- a/Engine/Source/Editor/UMGEditor/Public/WidgetBlueprintToolMenuContext.h
+++ b/Engine/Source/Editor/UMGEditor/Public/WidgetBlueprintToolMenuContext.h
@@ -8,8 +8,8 @@
 
 class FWidgetBlueprintEditor;
 
-UCLASS()
-class UMGEDITOR_API UWidgetBlueprintToolMenuContext : public UObject
+UCLASS(MinimalAPI)
+class UWidgetBlueprintToolMenuContext : public UObject
 {
 	GENERATED_BODY()
 
diff --git a/Engine/Source/Editor/UMGEditor/Public/WidgetCompilerRule.h b/Engine/Source/Editor/UMGEditor/Public/WidgetCompilerRule.h
--- a/Engine/Source/Editor/UMGEditor/Public/WidgetCompilerRule.h
+++ b/Engine/Source/Editor/UMGEditor/Public/WidgetCompilerRule.h
@@ -10,18 +10,22 @@
 
 #include "WidgetCompilerRule.generated.h"
 
+#define UE_API UMGEDITOR_API
+
 class FCompilerResultsLog;
 class UWidgetBlueprint;
 
 /**
  * 
  */
-UCLASS(Abstract)
-class UMGEDITOR_API UWidgetCompilerRule : public UObject
+UCLASS(MinimalAPI, Abstract)
+class UWidgetCompilerRule : public UObject
 {
 	GENERATED_BODY()
 public:
-	UWidgetCompilerRule();
+	UE_API UWidgetCompilerRule();
 
-	virtual void ExecuteRule(UWidgetBlueprint* WidgetBlueprint, FCompilerResultsLog& MessageLog);
+	UE_API virtual void ExecuteRule(UWidgetBlueprint* WidgetBlueprint, FCompilerResultsLog& MessageLog);
 };
+
+#undef UE_API
diff --git a/Engine/Source/Editor/UMGEditor/Public/WidgetEditingProjectSettings.h b/Engine/Source/Editor/UMGEditor/Public/WidgetEditingProjectSettings.h
--- a/Engine/Source/Editor/UMGEditor/Public/WidgetEditingProjectSettings.h
+++ b/Engine/Source/Editor/UMGEditor/Public/WidgetEditingProjectSettings.h
@@ -12,6 +12,8 @@
 #include "Misc/NamePermissionList.h"
 #include "WidgetEditingProjectSettings.generated.h"
 
+#define UE_API UMGEDITOR_API
+
 class UWidgetCompilerRule;
 class UUserWidget;
 class UWidgetBlueprint;
@@ -119,15 +121,15 @@ public:
 /**
  * Implements the settings for the UMG Editor Project Settings
  */
-UCLASS(Abstract, config=Editor, defaultconfig)
-class UMGEDITOR_API UWidgetEditingProjectSettings : public UDeveloperSettings
+UCLASS(MinimalAPI, Abstract, config=Editor, defaultconfig)
+class UWidgetEditingProjectSettings : public UDeveloperSettings
 {
 	GENERATED_BODY()
 
 public:
-	UWidgetEditingProjectSettings();
+	UE_API UWidgetEditingProjectSettings();
 
-	virtual void PostInitProperties() override;
+	UE_API virtual void PostInitProperties() override;
 
 public:
 	UPROPERTY(EditAnywhere, config, Category = "Compiler")
@@ -219,18 +221,18 @@ public:
 	UPROPERTY(EditAnywhere, config, Category=Designer)
 	TArray<FDebugResolution> DebugResolutions;
 
-	bool CompilerOption_AllowBlueprintTick(const class UWidgetBlueprint* WidgetBlueprint) const;
-	bool CompilerOption_AllowBlueprintPaint(const class UWidgetBlueprint* WidgetBlueprint) const;
-	EPropertyBindingPermissionLevel CompilerOption_PropertyBindingRule(const class UWidgetBlueprint* WidgetBlueprint) const;
-	TArray<UWidgetCompilerRule*> CompilerOption_Rules(const class UWidgetBlueprint* WidgetBlueprint) const;
+	UE_API bool CompilerOption_AllowBlueprintTick(const class UWidgetBlueprint* WidgetBlueprint) const;
+	UE_API bool CompilerOption_AllowBlueprintPaint(const class UWidgetBlueprint* WidgetBlueprint) const;
+	UE_API EPropertyBindingPermissionLevel CompilerOption_PropertyBindingRule(const class UWidgetBlueprint* WidgetBlueprint) const;
+	UE_API TArray<UWidgetCompilerRule*> CompilerOption_Rules(const class UWidgetBlueprint* WidgetBlueprint) const;
 
 	/** Get the permission list that controls which categories are exposed in config palette filtering */
-	FNamePermissionList& GetAllowedPaletteCategories();
-	const FNamePermissionList& GetAllowedPaletteCategories() const;
+	UE_API FNamePermissionList& GetAllowedPaletteCategories();
+	UE_API const FNamePermissionList& GetAllowedPaletteCategories() const;
 
 	/** Get the permission list that controls which widgets are exposed in config palette filtering */
-	FPathPermissionList& GetAllowedPaletteWidgets();
-	const FPathPermissionList& GetAllowedPaletteWidgets() const;
+	UE_API FPathPermissionList& GetAllowedPaletteWidgets();
+	UE_API const FPathPermissionList& GetAllowedPaletteWidgets() const;
 
 private:
 	template<typename ReturnType, typename T>
@@ -244,10 +246,10 @@ private:
 		return bValue;
 	}
 
-	void GetCompilerOptionsForWidget(const class UWidgetBlueprint* WidgetBlueprint, TFunctionRef<bool(const FWidgetCompilerOptions&)> Operator) const;
+	UE_API void GetCompilerOptionsForWidget(const class UWidgetBlueprint* WidgetBlueprint, TFunctionRef<bool(const FWidgetCompilerOptions&)> Operator) const;
 
 protected:
-	virtual void PerformUpgradeStepForVersion(int32 ForVersion);
+	UE_API virtual void PerformUpgradeStepForVersion(int32 ForVersion);
 
 	UPROPERTY(config)
 	int32 Version;
@@ -261,3 +263,5 @@ protected:
 	/** Individual widgets to always allow when using permission list palette filtering, regardless of category */
 	FPathPermissionList AllowedPaletteWidgets;
 };
+
+#undef UE_API
diff --git a/Engine/Source/Editor/UMGEditor/Public/WidgetReference.h b/Engine/Source/Editor/UMGEditor/Public/WidgetReference.h
--- a/Engine/Source/Editor/UMGEditor/Public/WidgetReference.h
+++ b/Engine/Source/Editor/UMGEditor/Public/WidgetReference.h
@@ -5,16 +5,18 @@
 #include "CoreMinimal.h"
 #include "Components/Widget.h"
 
+#define UE_API UMGEDITOR_API
+
 class FWidgetBlueprintEditor;
 struct FWidgetReference;
 
-struct UMGEDITOR_API FWidgetHandle
+struct FWidgetHandle
 {
 	friend FWidgetReference;
 	friend FWidgetBlueprintEditor;
 
 private:
-	FWidgetHandle(UWidget* Widget);
+	UE_API FWidgetHandle(UWidget* Widget);
 
 	TWeakObjectPtr<UWidget> Widget;
 };
@@ -23,21 +25,21 @@ private:
  * The Widget reference is a useful way to hold onto the selection in a way that allows for up to date access to the current preview object.
  * Because the designer could end up rebuilding the preview, it's best to hold onto an FWidgetReference.
  */
-struct UMGEDITOR_API FWidgetReference
+struct FWidgetReference
 {
 	friend FWidgetBlueprintEditor;
 
 public:
-	FWidgetReference();
+	UE_API FWidgetReference();
 
 	/** @returns true if both the template and the preview pointers are valid. */
-	bool IsValid() const;
+	UE_API bool IsValid() const;
 
 	/** @returns The template widget.  This is the widget that is serialized to disk */
-	UWidget* GetTemplate() const;
+	UE_API UWidget* GetTemplate() const;
 
 	/** @returns the preview widget.  This is the transient representation of the template.  Constantly being destroyed and recreated.  Do not cache this pointer. */
-	UWidget* GetPreview() const;
+	UE_API UWidget* GetPreview() const;
 
 	/** @returns the preview slate widget.  This is the transient representation of the template.  Constantly being destroyed and recreated.  Do not cache this pointer. */
 	TSharedPtr<SWidget> GetPreviewSlate() const
@@ -70,9 +72,9 @@ public:
 	}
 
 private:
-	FWidgetReference(TSharedPtr<FWidgetBlueprintEditor> WidgetEditor, TSharedPtr< FWidgetHandle > TemplateHandle);
+	UE_API FWidgetReference(TSharedPtr<FWidgetBlueprintEditor> WidgetEditor, TSharedPtr< FWidgetHandle > TemplateHandle);
 
-	void OnObjectsReplaced(const TMap<UObject*, UObject*>& ReplacementMap);
+	UE_API void OnObjectsReplaced(const TMap<UObject*, UObject*>& ReplacementMap);
 
 private:
 	TWeakPtr<FWidgetBlueprintEditor> WidgetEditor;
@@ -84,3 +86,5 @@ inline uint32 GetTypeHash(const struct FWidgetReference& WidgetRef)
 {
 	return ::GetTypeHash(reinterpret_cast<void*>( WidgetRef.GetTemplate() ));
 }
+
+#undef UE_API
diff --git a/Engine/Source/Editor/UMGEditor/Public/WidgetTemplate.h b/Engine/Source/Editor/UMGEditor/Public/WidgetTemplate.h
--- a/Engine/Source/Editor/UMGEditor/Public/WidgetTemplate.h
+++ b/Engine/Source/Editor/UMGEditor/Public/WidgetTemplate.h
@@ -12,6 +12,8 @@
 #include "Templates/SharedPointer.h"
 #include "Widgets/IToolTip.h"
 
+#define UE_API UMGEDITOR_API
+
 class IToolTip;
 class UWidget;
 struct FSlateBrush;
@@ -20,11 +22,11 @@ struct FSlateBrush;
  * The widget template represents a widget or a set of widgets to create and spawn into the widget tree.
  * More complicated defaults could be created by deriving from this class and registering new templates with the module.
  */
-class UMGEDITOR_API FWidgetTemplate : public TSharedFromThis<FWidgetTemplate>
+class FWidgetTemplate : public TSharedFromThis<FWidgetTemplate>
 {
 public:
 	/** Constructor */
-	FWidgetTemplate();
+	UE_API FWidgetTemplate();
 	virtual ~FWidgetTemplate() { }
 
 	/** The category this template fits into. */
@@ -53,3 +55,5 @@ public:
 	/** The name of the widget template. */
 	FText Name;
 };
+
+#undef UE_API
diff --git a/Engine/Source/Editor/UndoHistoryEditor/Public/IUndoHistoryEditorModule.h b/Engine/Source/Editor/UndoHistoryEditor/Public/IUndoHistoryEditorModule.h
--- a/Engine/Source/Editor/UndoHistoryEditor/Public/IUndoHistoryEditorModule.h
+++ b/Engine/Source/Editor/UndoHistoryEditor/Public/IUndoHistoryEditorModule.h
@@ -8,7 +8,7 @@
 
 UNDOHISTORYEDITOR_API extern const FName UndoHistoryTabName;
 
-class UNDOHISTORYEDITOR_API IUndoHistoryEditorModule : public IModuleInterface
+class IUndoHistoryEditorModule : public IModuleInterface
 {
 public:
 
diff --git a/Engine/Source/Editor/UniversalObjectLocatorEditor/Public/UniversalObjectLocatorEditor.h b/Engine/Source/Editor/UniversalObjectLocatorEditor/Public/UniversalObjectLocatorEditor.h
--- a/Engine/Source/Editor/UniversalObjectLocatorEditor/Public/UniversalObjectLocatorEditor.h
+++ b/Engine/Source/Editor/UniversalObjectLocatorEditor/Public/UniversalObjectLocatorEditor.h
@@ -6,6 +6,8 @@
 #include "Internationalization/Text.h"
 #include "Textures/SlateIcon.h"
 
+#define UE_API UNIVERSALOBJECTLOCATOREDITOR_API
+
 class UObject;
 class SWidget;
 class FDragDropOperation;
@@ -30,7 +32,7 @@ enum class ELocatorFragmentEditorType : uint8
 	Absolute,
 };
 
-class UNIVERSALOBJECTLOCATOREDITOR_API ILocatorFragmentEditor : public TSharedFromThis<ILocatorFragmentEditor>
+class ILocatorFragmentEditor : public TSharedFromThis<ILocatorFragmentEditor>
 {
 public:
 	virtual ~ILocatorFragmentEditor() = default;
@@ -69,7 +71,7 @@ public:
 	// Make the editor UI for this fragment (displayed in a context menu)
 	// @param InParameters     Parameters used to create the UI
 	// @return the editing widget
-	virtual TSharedPtr<SWidget> MakeEditUI(const FEditUIParameters& InParameters);
+	UE_API virtual TSharedPtr<SWidget> MakeEditUI(const FEditUIParameters& InParameters);
 
 	UE_DEPRECATED(5.5, "Please use MakeEditUI that takes a FEditUIParameters struct")
 	virtual TSharedPtr<SWidget> MakeEditUI(TSharedPtr<IUniversalObjectLocatorCustomization> Customization) { return nullptr; }
@@ -99,17 +101,19 @@ public:
 	// @param InFragment       The fragment to resolve the class of
 	// @param InContext        The context to resolve the class against
 	// @return the class that was resolved from the fragment
-	virtual UClass* ResolveClass(const FUniversalObjectLocatorFragment& InFragment, UObject* InContext) const;
+	UE_API virtual UClass* ResolveClass(const FUniversalObjectLocatorFragment& InFragment, UObject* InContext) const;
 
 	// Make a default fragment for the fragment type that this editor UI manipulates
 	// @return a default fragment for this locator editor to use
 	virtual FUniversalObjectLocatorFragment MakeDefaultLocatorFragment() const = 0;
 
 	UE_DEPRECATED(5.5, "No longer required. Please override MakeDefaultLocatorFragment")
-	virtual FUniversalObjectLocator MakeDefaultLocator() const;
+	UE_API virtual FUniversalObjectLocator MakeDefaultLocator() const;
 };
 
 UE_DEPRECATED(5.5, "Please use ILocatorFragmentEditor")
 typedef ILocatorFragmentEditor ILocatorEditor;
 
-} // namespace UE::UniversalObjectLocator
\ No newline at end of file
+} // namespace UE::UniversalObjectLocator
+
+#undef UE_API
diff --git a/Engine/Source/Editor/UniversalObjectLocatorEditor/Public/UniversalObjectLocatorEditorContext.h b/Engine/Source/Editor/UniversalObjectLocatorEditor/Public/UniversalObjectLocatorEditorContext.h
--- a/Engine/Source/Editor/UniversalObjectLocatorEditor/Public/UniversalObjectLocatorEditorContext.h
+++ b/Engine/Source/Editor/UniversalObjectLocatorEditor/Public/UniversalObjectLocatorEditorContext.h
@@ -1,4 +1,4 @@
-﻿// Copyright Epic Games, Inc. All Rights Reserved.
+// Copyright Epic Games, Inc. All Rights Reserved.
 
 #pragma once
 
@@ -10,7 +10,7 @@ class IPropertyHandle;
 namespace UE::UniversalObjectLocator
 {
 
-class UNIVERSALOBJECTLOCATOREDITOR_API ILocatorFragmentEditorContext : public TSharedFromThis<ILocatorFragmentEditorContext>
+class ILocatorFragmentEditorContext : public TSharedFromThis<ILocatorFragmentEditorContext>
 {
 public:
 	virtual ~ILocatorFragmentEditorContext() = default;
@@ -22,4 +22,4 @@ public:
 	virtual bool IsFragmentAllowed(FName InFragmentName) const { return true; }
 };
 
-} // namespace UE::UniversalObjectLocator
\ No newline at end of file
+} // namespace UE::UniversalObjectLocator
diff --git a/Engine/Source/Editor/UnrealEd/Classes/Animation/DebugSkelMeshComponent.h b/Engine/Source/Editor/UnrealEd/Classes/Animation/DebugSkelMeshComponent.h
--- a/Engine/Source/Editor/UnrealEd/Classes/Animation/DebugSkelMeshComponent.h
+++ b/Engine/Source/Editor/UnrealEd/Classes/Animation/DebugSkelMeshComponent.h
@@ -499,9 +499,9 @@ class UDebugSkelMeshComponent : public USkeletalMeshComponent
 	UNREALED_API void SetProcessRootMotionMode(EProcessRootMotionMode Mode);
 
 	/** Sets how we visualize root motion in the viewport. See EVisualizeRootMotionMode for details. */
-	UNREALED_API void SetVisualizeRootMotionMode(EVisualizeRootMotionMode Mode) { VisualizeRootMotionMode = Mode; };
-	UNREALED_API bool IsVisualizeRootMotionMode(EVisualizeRootMotionMode Mode) const { return VisualizeRootMotionMode == Mode; };
-	UNREALED_API EVisualizeRootMotionMode GetVisualizeRootMotionMode() const { return VisualizeRootMotionMode; };
+	void SetVisualizeRootMotionMode(EVisualizeRootMotionMode Mode) { VisualizeRootMotionMode = Mode; };
+	bool IsVisualizeRootMotionMode(EVisualizeRootMotionMode Mode) const { return VisualizeRootMotionMode == Mode; };
+	EVisualizeRootMotionMode GetVisualizeRootMotionMode() const { return VisualizeRootMotionMode; };
 
 	/** Whether the supplied root motion mode can be used for the current asset */
 	UNREALED_API bool CanUseProcessRootMotionMode(EProcessRootMotionMode Mode) const;
@@ -510,16 +510,16 @@ class UDebugSkelMeshComponent : public USkeletalMeshComponent
 	UNREALED_API bool DoesCurrentAssetHaveRootMotion() const;
 
 	/** Sets flags whether we notification visualizations should be drawn in the viewport. */
-	UNREALED_API void SetShowNotificationVisualizations(const bool bShow) { bShowNotificationVisualizations = bShow; }
-	UNREALED_API bool IsNotificationVisualizationsEnabled() const { return bShowNotificationVisualizations; }
+	void SetShowNotificationVisualizations(const bool bShow) { bShowNotificationVisualizations = bShow; }
+	bool IsNotificationVisualizationsEnabled() const { return bShowNotificationVisualizations; }
 
 	UE_DEPRECATED(5.5, "Use VisualizeRootMotionMode functions instead.")
-	UNREALED_API void SetShowRootMotionVisualizations(const bool bShow) { VisualizeRootMotionMode = bShow ? EVisualizeRootMotionMode::Trajectory : EVisualizeRootMotionMode::None; }
-	UNREALED_API bool IsRootMotionVisualizationsEnabled() const { return VisualizeRootMotionMode != EVisualizeRootMotionMode::None; }
+	void SetShowRootMotionVisualizations(const bool bShow) { VisualizeRootMotionMode = bShow ? EVisualizeRootMotionMode::Trajectory : EVisualizeRootMotionMode::None; }
+	bool IsRootMotionVisualizationsEnabled() const { return VisualizeRootMotionMode != EVisualizeRootMotionMode::None; }
 
 	/** Sets flags whether we AssetUserData visualizations should be drawn in the viewport. */
-	UNREALED_API void SetShowAssetUserDataVisualizations(const bool bShow) { bShowAssetUserDataVisualizations = bShow; }
-	UNREALED_API bool IsAssetUserDataVisualizationsEnabled() const { return bShowAssetUserDataVisualizations; }
+	void SetShowAssetUserDataVisualizations(const bool bShow) { bShowAssetUserDataVisualizations = bShow; }
+	bool IsAssetUserDataVisualizationsEnabled() const { return bShowAssetUserDataVisualizations; }
 
 	/** Whether the current LOD of the debug mesh is being synced with the attached (preview) mesh instance. */
 	UNREALED_API bool IsTrackingAttachedLOD() const;
diff --git a/Engine/Source/Editor/UnrealEd/Classes/Editor/AssetReferenceFilter.h b/Engine/Source/Editor/UnrealEd/Classes/Editor/AssetReferenceFilter.h
--- a/Engine/Source/Editor/UnrealEd/Classes/Editor/AssetReferenceFilter.h
+++ b/Engine/Source/Editor/UnrealEd/Classes/Editor/AssetReferenceFilter.h
@@ -25,11 +25,11 @@ public:
 	DECLARE_DELEGATE_RetVal_TwoParams(bool, FOnIsCrossPluginReferenceAllowed, const FAssetData& /*ReferencingAssetData*/, const FAssetData& /*ReferencedAssetData*/);
 
 	/** Used to register a custom delegate implementation that returns if a plugin asset can reference another plugin asset. */
-	UNREALED_API static FOnIsCrossPluginReferenceAllowed& OnIsCrossPluginReferenceAllowed() { return OnIsCrossPluginReferenceAllowedDelegate; }
+	static FOnIsCrossPluginReferenceAllowed& OnIsCrossPluginReferenceAllowed() { return OnIsCrossPluginReferenceAllowedDelegate; }
 	
 protected:
 	UNREALED_API static bool IsCrossPluginReferenceAllowed(const FAssetData& ReferencingAssetData, const FAssetData& ReferencedAssetData);
 
 private:
 	UNREALED_API static FOnIsCrossPluginReferenceAllowed OnIsCrossPluginReferenceAllowedDelegate;
-};
\ No newline at end of file
+};
diff --git a/Engine/Source/Editor/UnrealEd/Classes/Editor/EditorEngine.h b/Engine/Source/Editor/UnrealEd/Classes/Editor/EditorEngine.h
--- a/Engine/Source/Editor/UnrealEd/Classes/Editor/EditorEngine.h
+++ b/Engine/Source/Editor/UnrealEd/Classes/Editor/EditorEngine.h
@@ -32,6 +32,8 @@
 #include "RHIShaderPlatform.h"
 #include "EditorEngine.generated.h"
 
+#define UE_API UNREALED_API
+
 class APlayerStart;
 class Error;
 class FEditorViewportClient;
@@ -283,7 +285,7 @@ struct FPreviewPlatformInfo
 	}
 
 	/** Return the Constrained Aspect Ratio of the preview platform */
-	UNREALED_API float GetConstrainedAspectRatio()
+	float GetConstrainedAspectRatio()
 	{
 		return ConstrainedAspectRatio;
 	}
@@ -298,7 +300,7 @@ struct FPreviewPlatformInfo
 	UNREALED_API ERHIFeatureLevel::Type GetEffectivePreviewFeatureLevel() const;
 	
 	/** Returns the Friendly Name of the Preview*/
-	UNREALED_API FText GetFriendlyName() const { return PreviewShaderPlatformFriendlyName; }
+	FText GetFriendlyName() const { return PreviewShaderPlatformFriendlyName; }
 	
 	/** Returns the Shader Platform of the Preview*/
 	UNREALED_API EShaderPlatform GetShaderPlatform() const;
@@ -3333,7 +3335,7 @@ public:
 //////////////////////////////////////////////////////////////////////////
 // FActorLabelUtilities
 
-struct UNREALED_API FActorLabelUtilities
+struct FActorLabelUtilities
 {
 public:
 	/**
@@ -3344,7 +3346,7 @@ public:
 	 *
 	 * @return	true if the label ends with a number.
 	 */
-	static bool SplitActorLabel(FString& InOutLabel, int32& OutIdx);
+	static UE_API bool SplitActorLabel(FString& InOutLabel, int32& OutIdx);
 
 	/**
 	 * Assigns a new label to an actor. If the name exists it will be appended with a number to make it unique. Actor labels are only available in development builds.
@@ -3353,7 +3355,7 @@ public:
 	 * @param	NewActorLabel			The new label string to assign to the actor.  If empty, the actor will have a default label.
 	 * @param	InExistingActorLabels	(optional) Pointer to a set of actor labels that are currently in use
 	 */
-	static void SetActorLabelUnique(AActor* Actor, const FString& NewActorLabel, const FCachedActorLabels* InExistingActorLabels = nullptr);
+	static UE_API void SetActorLabelUnique(AActor* Actor, const FString& NewActorLabel, const FCachedActorLabels* InExistingActorLabels = nullptr);
 
 	/** 
 	 * Does an explicit actor rename. In addition to changing the label this will also fix any soft references pointing to it 
@@ -3362,8 +3364,10 @@ public:
 	 * @param	NewActorLabel			The new label string to assign to the actor.  If empty, the actor will have a default label.
 	 * @param	bMakeUnique				If true, it will call SetActorLabelUnique, if false it will use the exact label specified
 	 */
-	static void RenameExistingActor(AActor* Actor, const FString& NewActorLabel, bool bMakeUnique = false);
+	static UE_API void RenameExistingActor(AActor* Actor, const FString& NewActorLabel, bool bMakeUnique = false);
 
 private:
 	FActorLabelUtilities() {}
 };
+
+#undef UE_API
diff --git a/Engine/Source/Editor/UnrealEd/Classes/Editor/Transactor.h b/Engine/Source/Editor/UnrealEd/Classes/Editor/Transactor.h
--- a/Engine/Source/Editor/UnrealEd/Classes/Editor/Transactor.h
+++ b/Engine/Source/Editor/UnrealEd/Classes/Editor/Transactor.h
@@ -525,7 +525,7 @@ class UTransactor : public UObject
 	 * @return	Number of active actions when Begin() was called;  values greater than
 	 *			0 indicate that there was already an existing undo transaction in progress.
 	 */
-	UNREALED_API virtual int32 Begin( const TCHAR* SessionContext, const FText& Description ) PURE_VIRTUAL(UTransactor::Begin,return 0;);
+	virtual int32 Begin( const TCHAR* SessionContext, const FText& Description ) PURE_VIRTUAL(UTransactor::Begin,return 0;);
 
 	/**
 	 * Attempts to close an undo transaction.  Only successful if the transaction's action
@@ -534,25 +534,25 @@ class UTransactor : public UObject
 	 * @return	Number of active actions when End() was called; a value of 1 indicates that the
 	 *			transaction was successfully closed
 	 */
-	UNREALED_API virtual int32 End() PURE_VIRTUAL(UTransactor::End,return 0;);
+	virtual int32 End() PURE_VIRTUAL(UTransactor::End,return 0;);
 
 	/**
 	 * Cancels the current transaction, no longer capture actions to be placed in the undo buffer.
 	 *
 	 * @param	StartIndex	the value of ActiveIndex when the transaction to be canceled was began. 
 	 */
-	UNREALED_API virtual void Cancel( int32 StartIndex = 0 ) PURE_VIRTUAL(UTransactor::Cancel,);
+	virtual void Cancel( int32 StartIndex = 0 ) PURE_VIRTUAL(UTransactor::Cancel,);
 
 	/**
 	 * Resets the entire undo buffer;  deletes all undo transactions.
 	 */
-	UNREALED_API virtual void Reset( const FText& Reason ) PURE_VIRTUAL(UTransactor::Reset,);
+	virtual void Reset( const FText& Reason ) PURE_VIRTUAL(UTransactor::Reset,);
 
 	/**
 	 * Returns whether there are any active actions; i.e. whether actions are currently
 	 * being captured into the undo buffer.
 	 */
-	UNREALED_API virtual bool IsActive() PURE_VIRTUAL(UTransactor::IsActive,return false;);
+	virtual bool IsActive() PURE_VIRTUAL(UTransactor::IsActive,return false;);
 
 	/**
 	 * Determines whether the undo option should be selectable.
@@ -561,7 +561,7 @@ class UTransactor : public UObject
 	 *
 	 * @return	true if the "Undo" option should be selectable.
 	 */
-	UNREALED_API virtual bool CanUndo( FText* Text=nullptr ) PURE_VIRTUAL(UTransactor::CanUndo,return false;);
+	virtual bool CanUndo( FText* Text=nullptr ) PURE_VIRTUAL(UTransactor::CanUndo,return false;);
 
 	/**
 	 * Determines whether the redo option should be selectable.
@@ -570,14 +570,14 @@ class UTransactor : public UObject
 	 *
 	 * @return	true if the "Redo" option should be selectable.
 	 */
-	UNREALED_API virtual bool CanRedo( FText* Text=nullptr ) PURE_VIRTUAL(UTransactor::CanRedo,return false;);
+	virtual bool CanRedo( FText* Text=nullptr ) PURE_VIRTUAL(UTransactor::CanRedo,return false;);
 
 	/**
 	 * Gets the current length of the transaction queue.
 	 *
 	 * @return Queue length.
 	 */
-	UNREALED_API virtual int32 GetQueueLength( ) const PURE_VIRTUAL(UTransactor::GetQueueLength,return 0;);
+	virtual int32 GetQueueLength( ) const PURE_VIRTUAL(UTransactor::GetQueueLength,return 0;);
 
 	/**
 	 * Gets the transaction queue index from its TransactionId.
@@ -586,7 +586,7 @@ class UTransactor : public UObject
 	 *
 	 * @return The queue index or INDEX_NONE if not found.
 	 */
-	UNREALED_API virtual int32 FindTransactionIndex(const FGuid& TransactionId ) const PURE_VIRTUAL(UTransactor::FindTransactionIndex, return INDEX_NONE;);
+	virtual int32 FindTransactionIndex(const FGuid& TransactionId ) const PURE_VIRTUAL(UTransactor::FindTransactionIndex, return INDEX_NONE;);
 
 	/**
 	 * Gets the transaction at the specified queue index.
@@ -595,7 +595,7 @@ class UTransactor : public UObject
 	 *
 	 * @return A read-only pointer to the transaction, or NULL if it does not exist.
 	 */
-	UNREALED_API virtual const FTransaction* GetTransaction( int32 QueueIndex ) const PURE_VIRTUAL(UTransactor::GetTransaction,return nullptr;);
+	virtual const FTransaction* GetTransaction( int32 QueueIndex ) const PURE_VIRTUAL(UTransactor::GetTransaction,return nullptr;);
 
 	/**
 	 * Returns the description of the undo action that will be performed next.
@@ -605,21 +605,21 @@ class UTransactor : public UObject
 	 *
 	 * @return	text describing the next undo transaction
 	 */
-	UNREALED_API virtual FTransactionContext GetUndoContext( bool bCheckWhetherUndoPossible = true ) PURE_VIRTUAL(UTransactor::GetUndoContext,return FTransactionContext(););
+	virtual FTransactionContext GetUndoContext( bool bCheckWhetherUndoPossible = true ) PURE_VIRTUAL(UTransactor::GetUndoContext,return FTransactionContext(););
 
 	/**
 	 * Determines the amount of data currently stored by the transaction buffer.
 	 *
 	 * @return	number of bytes stored in the undo buffer
 	 */
-	UNREALED_API virtual SIZE_T GetUndoSize() const PURE_VIRTUAL(UTransactor::GetUndoSize,return 0;);
+	virtual SIZE_T GetUndoSize() const PURE_VIRTUAL(UTransactor::GetUndoSize,return 0;);
 
 	/**
 	 * Gets the number of transactions that were undone and can be redone.
 	 *
 	 * @return Undo count.
 	 */
-	UNREALED_API virtual int32 GetUndoCount( ) const PURE_VIRTUAL(UTransactor::GetUndoCount,return 0;);
+	virtual int32 GetUndoCount( ) const PURE_VIRTUAL(UTransactor::GetUndoCount,return 0;);
 
 	/**
 	 * Returns the description of the redo action that will be performed next.
@@ -627,28 +627,28 @@ class UTransactor : public UObject
 	 * 
 	 * @return	text describing the next redo transaction
 	 */
-	UNREALED_API virtual FTransactionContext GetRedoContext() PURE_VIRTUAL(UTransactor::GetRedoContext,return FTransactionContext(););
+	virtual FTransactionContext GetRedoContext() PURE_VIRTUAL(UTransactor::GetRedoContext,return FTransactionContext(););
 
 	/**
 	 * Sets an undo barrier at the current point in the transaction buffer.
 	 * Undoing beyond this point will not be allowed until the barrier is removed.
 	 */
-	UNREALED_API virtual void SetUndoBarrier() PURE_VIRTUAL(UTransactor::SetUndoBarrier, );
+	virtual void SetUndoBarrier() PURE_VIRTUAL(UTransactor::SetUndoBarrier, );
 
 	/**
 	 * Removes the last set undo barrier from the transaction buffer.
 	 */
-	UNREALED_API virtual void RemoveUndoBarrier() PURE_VIRTUAL(UTransactor::RemoveUndoBarrier, );
+	virtual void RemoveUndoBarrier() PURE_VIRTUAL(UTransactor::RemoveUndoBarrier, );
 
 	/**
 	 * Clears all undo barriers.
 	 */
-	UNREALED_API virtual void ClearUndoBarriers() PURE_VIRTUAL(UTransactor::ClearUndoBarriers, );
+	virtual void ClearUndoBarriers() PURE_VIRTUAL(UTransactor::ClearUndoBarriers, );
 
 	/**
 	 * Gets the current undo barrier's position in transaction queue
 	 */
-	 UNREALED_API virtual int32 GetCurrentUndoBarrier() const PURE_VIRTUAL(UTransactor::GetCurrentUndoBarrier(), { return INDEX_NONE; });
+	 virtual int32 GetCurrentUndoBarrier() const PURE_VIRTUAL(UTransactor::GetCurrentUndoBarrier(), { return INDEX_NONE; });
 
 	/**
 	 * Executes an undo transaction, undoing all actions contained by that transaction.
diff --git a/Engine/Source/Editor/UnrealEd/Classes/MaterialEditor/MaterialEditorInstanceConstant.h b/Engine/Source/Editor/UnrealEd/Classes/MaterialEditor/MaterialEditorInstanceConstant.h
--- a/Engine/Source/Editor/UnrealEd/Classes/MaterialEditor/MaterialEditorInstanceConstant.h
+++ b/Engine/Source/Editor/UnrealEd/Classes/MaterialEditor/MaterialEditorInstanceConstant.h
@@ -285,17 +285,17 @@ public:
 	/** 
 	 * Regenerates the parameter arrays. 
 	 */
-	virtual	UNREALED_API void RegenerateArrays() {}
+	virtual	void RegenerateArrays() {}
 
 #if WITH_EDITOR
 	/** Sets back to zero the overrides for any parameters copied out of the layer stack */
-	UNREALED_API virtual void CleanParameterStack(int32 Index, EMaterialParameterAssociation MaterialType) { }
+	virtual void CleanParameterStack(int32 Index, EMaterialParameterAssociation MaterialType) { }
 
 	/** Copies the overrides for any parameters copied out of the layer stack from the layer or blend */
-	UNREALED_API virtual void ResetOverrides(int32 Index, EMaterialParameterAssociation MaterialType) {}
+	virtual void ResetOverrides(int32 Index, EMaterialParameterAssociation MaterialType) {}
 	
 	/** Copies the parameter array values back to the source instance. */
-	UNREALED_API virtual void CopyToSourceInstance(const bool bForceStaticPermutationUpdate = false) {}
+	virtual void CopyToSourceInstance(const bool bForceStaticPermutationUpdate = false) {}
 	
 #endif
 };
diff --git a/Engine/Source/Editor/UnrealEd/Classes/MaterialGraph/MaterialGraphNode_Base.h b/Engine/Source/Editor/UnrealEd/Classes/MaterialGraph/MaterialGraphNode_Base.h
--- a/Engine/Source/Editor/UnrealEd/Classes/MaterialGraph/MaterialGraphNode_Base.h
+++ b/Engine/Source/Editor/UnrealEd/Classes/MaterialGraph/MaterialGraphNode_Base.h
@@ -43,7 +43,7 @@ class UMaterialGraphNode_Base : public UEdGraphNode
 	/** Get the Material value type of an output pin */
 	UNREALED_API EMaterialValueType GetOutputValueType(const UEdGraphPin* OutputPin) const;
 	
-	UNREALED_API uint32 GetOutputType(const UEdGraphPin* OutputPin)
+	uint32 GetOutputType(const UEdGraphPin* OutputPin)
 	{
 		PRAGMA_DISABLE_DEPRECATION_WARNINGS
 		return static_cast<uint32>(GetOutputValueType(OutputPin));
diff --git a/Engine/Source/Editor/UnrealEd/Classes/UserDefinedStructure/UserDefinedStructEditorData.h b/Engine/Source/Editor/UnrealEd/Classes/UserDefinedStructure/UserDefinedStructEditorData.h
--- a/Engine/Source/Editor/UnrealEd/Classes/UserDefinedStructure/UserDefinedStructEditorData.h
+++ b/Engine/Source/Editor/UnrealEd/Classes/UserDefinedStructure/UserDefinedStructEditorData.h
@@ -14,18 +14,20 @@
 #include "StructUtils/UserDefinedStructEditorUtils.h"
 #include "UserDefinedStructEditorData.generated.h"
 
+#define UE_API UNREALED_API
+
 class ITransactionObjectAnnotation;
 
-struct UNREALED_API FStructVariableMetaData
+struct FStructVariableMetaData
 {
-	static const FName ClampMin;
-	static const FName ClampMax;
-	static const FName UIMin;
-	static const FName UIMax;
+	static UE_API const FName ClampMin;
+	static UE_API const FName ClampMax;
+	static UE_API const FName UIMin;
+	static UE_API const FName UIMax;
 
 private:
 	// This struct should never be instantiated
-	FStructVariableMetaData();
+	UE_API FStructVariableMetaData();
 };
 
 USTRUCT()
@@ -182,3 +184,5 @@ private:
 	UNREALED_API void ConsolidatedPostEditUndo(FStructureEditorUtils::EStructureEditorChangeInfo ActiveChange);
 
 };
+
+#undef UE_API
diff --git a/Engine/Source/Editor/UnrealEd/Private/FbxExporter.h b/Engine/Source/Editor/UnrealEd/Private/FbxExporter.h
--- a/Engine/Source/Editor/UnrealEd/Private/FbxExporter.h
+++ b/Engine/Source/Editor/UnrealEd/Private/FbxExporter.h
@@ -494,7 +494,7 @@ private:
 	UNREALED_API void ExportLevelSequenceTrackChannels( FbxNode* FbxActor, UMovieSceneTrack& Track, const TRange<FFrameNumber>& InPlaybackRange, const FMovieSceneSequenceTransform& RootToLocalTransform, bool bBakeBezierCurves);
 
 	UE_DEPRECATED(5.4, "Use ExportLevelSequenceTrackChannels that indicates whether bezier channels should be baked")
-	UNREALED_API void ExportLevelSequenceTrackChannels(FbxNode* FbxActor, UMovieSceneTrack& Track, const TRange<FFrameNumber>& InPlaybackRange, const FMovieSceneSequenceTransform& RootToLocalTransform) { ExportLevelSequenceTrackChannels(FbxActor, Track, InPlaybackRange, RootToLocalTransform, false); }
+	void ExportLevelSequenceTrackChannels(FbxNode* FbxActor, UMovieSceneTrack& Track, const TRange<FFrameNumber>& InPlaybackRange, const FMovieSceneSequenceTransform& RootToLocalTransform) { ExportLevelSequenceTrackChannels(FbxActor, Track, InPlaybackRange, RootToLocalTransform, false); }
 
 	/** Defines value export modes for the EportRichCurveToFbxCurve method. */
 	enum class ERichCurveValueMode
diff --git a/Engine/Source/Editor/UnrealEd/Private/PropertyColorSettings.h b/Engine/Source/Editor/UnrealEd/Private/PropertyColorSettings.h
--- a/Engine/Source/Editor/UnrealEd/Private/PropertyColorSettings.h
+++ b/Engine/Source/Editor/UnrealEd/Private/PropertyColorSettings.h
@@ -37,8 +37,8 @@ struct FPropertyColorCustomProperty
 /**
 * Implements the settings for Property Color.
 */
-UCLASS(config = EditorSettings)
-class UNREALED_API UPropertyColorSettings : public UObject
+UCLASS(MinimalAPI, config = EditorSettings)
+class UPropertyColorSettings : public UObject
 {
 	GENERATED_BODY()
 
@@ -48,4 +48,4 @@ public:
 
 	UPROPERTY(Config)
 	TArray<FPropertyColorCustomProperty> CustomProperties;
-};
\ No newline at end of file
+};
diff --git a/Engine/Source/Editor/UnrealEd/Public/DragTool_BoxSelect.h b/Engine/Source/Editor/UnrealEd/Public/DragTool_BoxSelect.h
--- a/Engine/Source/Editor/UnrealEd/Public/DragTool_BoxSelect.h
+++ b/Engine/Source/Editor/UnrealEd/Public/DragTool_BoxSelect.h
@@ -17,13 +17,13 @@ class UModel;
 class FDragTool_ActorBoxSelect : public FDragTool
 {
 public:
-	UNREALED_API explicit FDragTool_ActorBoxSelect( FLevelEditorViewportClient* InLevelViewportClient )
+	explicit FDragTool_ActorBoxSelect( FLevelEditorViewportClient* InLevelViewportClient )
 		: FDragTool(InLevelViewportClient->GetModeTools())
 		, LevelViewportClient( InLevelViewportClient )
 		, EditorViewportClient( InLevelViewportClient)
 	{}
 
-	UNREALED_API explicit FDragTool_ActorBoxSelect( FEditorViewportClient* InEditorViewportClient )
+	explicit FDragTool_ActorBoxSelect( FEditorViewportClient* InEditorViewportClient )
 		: FDragTool(InEditorViewportClient->GetModeTools())
 		, LevelViewportClient( nullptr)
 		, EditorViewportClient( InEditorViewportClient )
diff --git a/Engine/Source/Editor/UnrealEd/Public/DragTool_FrustumSelect.h b/Engine/Source/Editor/UnrealEd/Public/DragTool_FrustumSelect.h
--- a/Engine/Source/Editor/UnrealEd/Public/DragTool_FrustumSelect.h
+++ b/Engine/Source/Editor/UnrealEd/Public/DragTool_FrustumSelect.h
@@ -17,13 +17,13 @@ class UModel;
 class FDragTool_ActorFrustumSelect : public FDragTool
 {
 public:
-	UNREALED_API explicit FDragTool_ActorFrustumSelect(FLevelEditorViewportClient* InLevelViewportClient)
+	explicit FDragTool_ActorFrustumSelect(FLevelEditorViewportClient* InLevelViewportClient)
 		: FDragTool(InLevelViewportClient->GetModeTools())
 		, LevelViewportClient( InLevelViewportClient )
 		, EditorViewportClient( InLevelViewportClient)
 	{}
 
-	UNREALED_API explicit FDragTool_ActorFrustumSelect(FEditorViewportClient* InEditorViewportClient)
+	explicit FDragTool_ActorFrustumSelect(FEditorViewportClient* InEditorViewportClient)
 	: FDragTool(InEditorViewportClient->GetModeTools())
 	, LevelViewportClient( nullptr )
 	, EditorViewportClient( InEditorViewportClient)
diff --git a/Engine/Source/Editor/UnrealEd/Public/EdGraphNode_Comment.h b/Engine/Source/Editor/UnrealEd/Public/EdGraphNode_Comment.h
--- a/Engine/Source/Editor/UnrealEd/Public/EdGraphNode_Comment.h
+++ b/Engine/Source/Editor/UnrealEd/Public/EdGraphNode_Comment.h
@@ -87,11 +87,11 @@ public:
 	//~ End UObject Interface
 
 	//~ Begin UEdGraphNode Interface
-	UNREALED_API virtual void AllocateDefaultPins() override {}
+	virtual void AllocateDefaultPins() override {}
 	UNREALED_API virtual FText GetTooltipText() const override;
 	UNREALED_API virtual FLinearColor GetNodeCommentColor() const override;
 	UNREALED_API virtual FText GetNodeTitle(ENodeTitleType::Type TitleType) const override;
-	UNREALED_API virtual bool ShouldOverridePinNames() const override { return true; }
+	virtual bool ShouldOverridePinNames() const override { return true; }
 	UNREALED_API virtual FText GetPinNameOverride(const UEdGraphPin& Pin) const override;
 	UNREALED_API virtual void ResizeNode(const FVector2f& NewSize) override;
 	UNREALED_API virtual void PostPlacedNewNode() override;
@@ -115,7 +115,7 @@ public:
 	UNREALED_API const FCommentNodeSet& GetNodesUnderComment() const;
 
 	/** Return the font size of the comment */
-	UNREALED_API virtual int32 GetFontSize() const { return FontSize; }
+	virtual int32 GetFontSize() const { return FontSize; }
 
 	/** Override the default selection state of this graph node */
 	enum class ESelectionState : uint8 { Inherited, Selected, Deselected };
diff --git a/Engine/Source/Editor/UnrealEd/Public/EditorActorFolders.h b/Engine/Source/Editor/UnrealEd/Public/EditorActorFolders.h
--- a/Engine/Source/Editor/UnrealEd/Public/EditorActorFolders.h
+++ b/Engine/Source/Editor/UnrealEd/Public/EditorActorFolders.h
@@ -103,10 +103,10 @@ struct FActorFolders : public FGCObject, public IActorEditorContextClient
 	UNREALED_API bool RenameFolderInWorld(UWorld& World, FName OldPath, FName NewPath);
 
 	UE_DEPRECATED(5.4, "Please use ForEachActorDescInstanceInFolders")
-	static UNREALED_API void ForEachActorDescInFolders(UWorld& InWorld, const TSet<FName>& InPaths, TFunctionRef<bool(const FWorldPartitionActorDesc*)> Operation, const FFolder::FRootObject& InFolderRootObject = FFolder::GetInvalidRootObject()) {}
+	static void ForEachActorDescInFolders(UWorld& InWorld, const TSet<FName>& InPaths, TFunctionRef<bool(const FWorldPartitionActorDesc*)> Operation, const FFolder::FRootObject& InFolderRootObject = FFolder::GetInvalidRootObject()) {}
 	
 	UE_DEPRECATED(5.4, "Please use GetActorDescInstanceFolder")
-	static UNREALED_API FFolder GetActorDescFolder(UWorld& InWorld, const FWorldPartitionActorDesc* InActorDesc) { return FFolder::GetInvalidFolder(); }
+	static FFolder GetActorDescFolder(UWorld& InWorld, const FWorldPartitionActorDesc* InActorDesc) { return FFolder::GetInvalidFolder(); }
 	//~ End Deprecated
 	
 	static UNREALED_API FFolder GetActorDescInstanceFolder(UWorld& InWorld, const FWorldPartitionActorDescInstance* InActorDescInstance);
diff --git a/Engine/Source/Editor/UnrealEd/Public/EditorAxisDisplayInfo.h b/Engine/Source/Editor/UnrealEd/Public/EditorAxisDisplayInfo.h
--- a/Engine/Source/Editor/UnrealEd/Public/EditorAxisDisplayInfo.h
+++ b/Engine/Source/Editor/UnrealEd/Public/EditorAxisDisplayInfo.h
@@ -4,38 +4,42 @@
 
 #include "Misc/AxisDisplayInfo.h"
 
-class UNREALED_API FEditorAxisDisplayInfo : public IAxisDisplayInfo
+#define UE_API UNREALED_API
+
+class FEditorAxisDisplayInfo : public IAxisDisplayInfo
 {
 public:
 
-	FEditorAxisDisplayInfo();
+	UE_API FEditorAxisDisplayInfo();
 	virtual ~FEditorAxisDisplayInfo() = default;
 
-	virtual EAxisList::Type GetAxisDisplayCoordinateSystem() const override;
+	UE_API virtual EAxisList::Type GetAxisDisplayCoordinateSystem() const override;
 	
-	virtual FText GetAxisToolTip(EAxisList::Type Axis) const override;
-	virtual FText GetAxisDisplayName(EAxisList::Type Axis) override;
-	virtual FText GetAxisDisplayNameShort(EAxisList::Type Axis) override;
+	UE_API virtual FText GetAxisToolTip(EAxisList::Type Axis) const override;
+	UE_API virtual FText GetAxisDisplayName(EAxisList::Type Axis) override;
+	UE_API virtual FText GetAxisDisplayNameShort(EAxisList::Type Axis) override;
 	
 
-	virtual FLinearColor GetAxisColor(EAxisList::Type Axis) override;
+	UE_API virtual FLinearColor GetAxisColor(EAxisList::Type Axis) override;
 
-	virtual bool UseForwardRightUpDisplayNames() override;
+	UE_API virtual bool UseForwardRightUpDisplayNames() override;
 
-	virtual FText GetRotationAxisToolTip(EAxisList::Type Axis) const override;
-	virtual FText GetRotationAxisName(EAxisList::Type Axis) override;
-	virtual FText GetRotationAxisNameShort(EAxisList::Type Axis) override;
+	UE_API virtual FText GetRotationAxisToolTip(EAxisList::Type Axis) const override;
+	UE_API virtual FText GetRotationAxisName(EAxisList::Type Axis) override;
+	UE_API virtual FText GetRotationAxisNameShort(EAxisList::Type Axis) override;
 	
-	virtual FIntVector4 DefaultAxisComponentDisplaySwizzle() const override;
+	UE_API virtual FIntVector4 DefaultAxisComponentDisplaySwizzle() const override;
 
 private:
 	// Maps the given axis from FLU -> XYZ if the AxisDisplayCoordinateSystem is XYZ
-	EAxisList::Type MapAxis(EAxisList::Type Axis) const;
+	UE_API EAxisList::Type MapAxis(EAxisList::Type Axis) const;
 
 	// Inits info stored in settings, like the axis colors
-	void InitSettingsInfo(double EditorStartupTime);
+	UE_API void InitSettingsInfo(double EditorStartupTime);
 
 	TOptional<bool> bUseForwardRightUpDisplayNames;
 	mutable TOptional<EAxisList::Type> AxisDisplayCoordinateSystem;
 
 };
+
+#undef UE_API
diff --git a/Engine/Source/Editor/UnrealEd/Public/EditorState/EditorState.h b/Engine/Source/Editor/UnrealEd/Public/EditorState/EditorState.h
--- a/Engine/Source/Editor/UnrealEd/Public/EditorState/EditorState.h
+++ b/Engine/Source/Editor/UnrealEd/Public/EditorState/EditorState.h
@@ -60,7 +60,7 @@ public:
 
 public:
 	/** Get the category under which this state's properties should be displayed. */
-	UNREALED_API virtual FText GetCategoryText() const PURE_VIRTUAL(UEditorState::GetCategoryText, return FText::GetEmpty(); );
+	virtual FText GetCategoryText() const PURE_VIRTUAL(UEditorState::GetCategoryText, return FText::GetEmpty(); );
 	
 	/** Get a list of state types this editor state depends on. Dependent types will be restored after their dependencies. */
 	virtual TArray<TSubclassOf<UEditorState>> GetDependencies() const { return {}; }
diff --git a/Engine/Source/Editor/UnrealEd/Public/EditorViewportClient.h b/Engine/Source/Editor/UnrealEd/Public/EditorViewportClient.h
--- a/Engine/Source/Editor/UnrealEd/Public/EditorViewportClient.h
+++ b/Engine/Source/Editor/UnrealEd/Public/EditorViewportClient.h
@@ -617,8 +617,8 @@ public:
 	UNREALED_API virtual void SetEnabledStats(const TArray<FString>& InEnabledStats) override;
 	UNREALED_API virtual bool IsStatEnabled(const FString& InName) const override;
 
-	UNREALED_API virtual bool BeginTransform(const FGizmoState& InState) { return false; }
-	UNREALED_API virtual bool EndTransform(const FGizmoState& InState) { return false; }
+	virtual bool BeginTransform(const FGizmoState& InState) { return false; }
+	virtual bool EndTransform(const FGizmoState& InState) { return false; }
 	
 protected:
 
@@ -919,7 +919,7 @@ public:
 		;
 
 	UE_DEPRECATED(5.4, "Use the overload that takes FDropObjectOptions instead.")
-	UNREALED_API virtual bool DropObjectsAtCoordinates(int32 MouseX, int32 MouseY, const TArray<UObject*>& DroppedObjects, TArray<AActor*>& OutNewActors,
+	virtual bool DropObjectsAtCoordinates(int32 MouseX, int32 MouseY, const TArray<UObject*>& DroppedObjects, TArray<AActor*>& OutNewActors,
 		bool bOnlyDropOnTarget = false, bool bCreateDropPreview = false, bool bSelectActors = true,
 		UActorFactory* FactoryToUse = nullptr) { return false; }
 
@@ -1309,7 +1309,7 @@ public:
 
 	/** Show or hide the widget. */
 	UNREALED_API void ShowWidget(const bool bShow);
-	UNREALED_API bool GetShowWidget() const { return bShowWidget; }
+	bool GetShowWidget() const { return bShowWidget; }
 
 	/**
 	 * Returns whether or not the flight camera is active
