Double-clicking the Canvas Anchor Medallion in a User Widget Blueprint causes an Ensure condition to fail.

#tests Tested single-clicking, double-clicking, and dragging (for both single and double clicking) the medallion, and all worked as expected without any ensures failing.
[RN] fix ui/umg

diff --git a/Engine/Source/Editor/UMGEditor/Private/Extensions/CanvasSlotExtension.cpp b/Engine/Source/Editor/UMGEditor/Private/Extensions/CanvasSlotExtension.cpp
--- a/Engine/Source/Editor/UMGEditor/Private/Extensions/CanvasSlotExtension.cpp
+++ b/Engine/Source/Editor/UMGEditor/Private/Extensions/CanvasSlotExtension.cpp
@@ -332,28 +332,38 @@ void FCanvasSlotExtension::Paint(const TSet< FWidgetReference >& Selection, cons
 
 FReply FCanvasSlotExtension::HandleAnchorBeginDrag(const FGeometry& Geometry, const FPointerEvent& Event, EAnchorWidget AnchorType)
 {
-	BeginTransaction(LOCTEXT("MoveAnchor", "Move Anchor"));
+	if (Event.GetEffectingButton() == EKeys::LeftMouseButton && AnchorWidgets[(int32)AnchorType].IsValid())
+	{ 
+		BeginTransaction(LOCTEXT("MoveAnchor", "Move Anchor"));
 
-	bMovingAnchor = true;
-	MouseDownPosition = Event.GetScreenSpacePosition();
+		bMovingAnchor = true;
+		MouseDownPosition = Event.GetScreenSpacePosition();
 
-	UCanvasPanelSlot* PreviewCanvasSlot = Cast<UCanvasPanelSlot>(SelectionCache[0].GetPreview()->Slot);
-	BeginAnchors = PreviewCanvasSlot->GetLayout().Anchors;
+		UCanvasPanelSlot* PreviewCanvasSlot = Cast<UCanvasPanelSlot>(SelectionCache[0].GetPreview()->Slot);
+		BeginAnchors = PreviewCanvasSlot->GetLayout().Anchors;
 
-	Designer->PushDesignerMessage(LOCTEXT("CenterAnchorControls", "Hold [Ctrl] to update widget position"));
+		Designer->PushDesignerMessage(LOCTEXT("CenterAnchorControls", "Hold [Ctrl] to update widget position"));
 
-	return FReply::Handled().CaptureMouse(AnchorWidgets[(int32)AnchorType].ToSharedRef());
+		return FReply::Handled().CaptureMouse(AnchorWidgets[(int32)AnchorType].ToSharedRef());
+	}
+	
+	return FReply::Unhandled();
 }
 
 FReply FCanvasSlotExtension::HandleAnchorEndDrag(const FGeometry& Geometry, const FPointerEvent& Event, EAnchorWidget AnchorType)
 {
-	EndTransaction();
+	if (AnchorWidgets[(int32)AnchorType].IsValid() && AnchorWidgets[(int32)AnchorType]->HasMouseCapture() && Event.GetEffectingButton() == EKeys::LeftMouseButton)
+	{
+		EndTransaction();
 
-	bMovingAnchor = false;
+		bMovingAnchor = false;
 
-	Designer->PopDesignerMessage();
+		Designer->PopDesignerMessage();
 
-	return FReply::Handled().ReleaseMouseCapture();
+		return FReply::Handled().ReleaseMouseCapture();
+	}
+
+	return FReply::Unhandled();
 }
 
 void FCanvasSlotExtension::ProximitySnapValue(float SnapFrequency, float SnapProximity, FVector2D::FReal& Value)
diff --git a/Engine/Source/Editor/UMGEditor/Public/DesignerExtension.h b/Engine/Source/Editor/UMGEditor/Public/DesignerExtension.h
--- a/Engine/Source/Editor/UMGEditor/Public/DesignerExtension.h
+++ b/Engine/Source/Editor/UMGEditor/Public/DesignerExtension.h
@@ -69,7 +69,7 @@ public:
 		return Widget;
 	}
 
-	/** Gets the location where the widget will be appear */
+	/** Gets the location where the widget will appear */
 	EExtensionLayoutLocation GetLocation() const
 	{
 		return Location;
