UMG Designer: (Experimental) Add Component List View to allow reordering.

This allows:
- Slate wrapper Component to control the wrapping order when multiple components are wrapping the same widget.
- Control the Construction and Initialization order of Components on a specific widget.

diff --git a/Engine/Source/Editor/UMGEditor/Private/Designer/SDesignerView.cpp b/Engine/Source/Editor/UMGEditor/Private/Designer/SDesignerView.cpp
--- a/Engine/Source/Editor/UMGEditor/Private/Designer/SDesignerView.cpp
+++ b/Engine/Source/Editor/UMGEditor/Private/Designer/SDesignerView.cpp
@@ -2346,7 +2346,8 @@ void SDesignerView::UpdatePreviewWidget(bool bForceUpdate)
 
 void SDesignerView::BroadcastDesignerChanged()
 {
-	UUserWidget* LatestPreviewWidget = BlueprintEditor.Pin()->GetPreview();
+	TSharedPtr<FWidgetBlueprintEditor> WidgetBlueprintEditor = BlueprintEditor.Pin();
+	UUserWidget* LatestPreviewWidget = WidgetBlueprintEditor->GetPreview();
 	ULevelEditorPlaySettings* PlayInSettings = GetMutableDefault<ULevelEditorPlaySettings>();
 	
 	if ( LatestPreviewWidget )
@@ -2370,6 +2371,7 @@ void SDesignerView::BroadcastDesignerChanged()
 		EventArgs.DpiScale = GetPreviewDPIScale();
 
 		LatestPreviewWidget->OnDesignerChanged(EventArgs);
+		WidgetBlueprintEditor->OnWidgetPreviewReady.Broadcast();
 	}
 }
 
diff --git a/Engine/Source/Editor/UMGEditor/Private/Details/SUIComponentView.cpp b/Engine/Source/Editor/UMGEditor/Private/Details/SUIComponentView.cpp
new file mode 100644
--- /dev/null
+++ b/Engine/Source/Editor/UMGEditor/Private/Details/SUIComponentView.cpp
@@ -0,0 +1,389 @@
+// Copyright Epic Games, Inc. All Rights Reserved.
+
+#include "SUIComponentView.h"
+
+#include "DragAndDrop/DecoratedDragDropOp.h"
+#include "ScopedTransaction.h"
+#include "Extensions/UIComponent.h"
+#include "Extensions/UIComponentUserWidgetExtension.h"
+#include "Framework/Commands/GenericCommands.h"
+#include "Framework/MultiBox/MultiBoxBuilder.h"
+#include "Input/DragAndDrop.h"
+#include "Kismet2/SClassPickerDialog.h"
+#include "SPositiveActionButton.h"
+#include "UIComponentUtils.h"
+#include "WidgetBlueprintEditor.h"
+#include "Widgets/Docking/SDockTab.h"
+#include "Widgets/Images/SImage.h"
+#include "Widgets/Input/SButton.h"
+#include "Widgets/Layout/SBorder.h"
+#include "Widgets/Layout/SScrollBorder.h"
+#include "Widgets/Text/STextBlock.h"
+#include "Widgets/Views/SListView.h"
+#include "Widgets/Views/STableRow.h"
+#include "Widgets/Views/STableViewBase.h"
+#include "Widgets/Text/SInlineEditableTextBlock.h"
+
+#if WITH_EDITOR
+#include "Styling/AppStyle.h"
+#endif // WITH_EDITOR
+
+#define LOCTEXT_NAMESPACE "UMG"
+
+struct FUIComponentListItem
+{
+	explicit FUIComponentListItem(UUIComponent* InComponent, UWidget* InWidget)
+		: Component(InComponent),
+		  Widget(InWidget)
+	{}
+
+	TWeakObjectPtr<UUIComponent> Component;
+	TWeakObjectPtr<UWidget> Widget;
+};
+
+class FUIComponentDragDropOp final : public FDecoratedDragDropOp
+{
+public:
+	DRAG_DROP_OPERATOR_TYPE(FUIComponentDragDropOp, FDecoratedDragDropOp)
+
+	/** The template to create an instance */
+	TSharedPtr<FUIComponentListItem> ListItem;
+
+	/** Constructs the drag drop operation */
+	static TSharedRef<FUIComponentDragDropOp> New(const TSharedPtr<FUIComponentListItem>& InListItem, FText InDragText)
+	{
+		TSharedRef<FUIComponentDragDropOp> Operation = MakeShared<FUIComponentDragDropOp>();
+		Operation->ListItem = InListItem;
+		Operation->DefaultHoverText = InDragText;
+		Operation->CurrentHoverText = InDragText;
+		Operation->Construct();
+
+		return Operation;
+	}
+};
+
+class SUIComponentListItem : public STableRow<TSharedPtr<FUIComponentListItem> >
+{
+public:
+	SLATE_BEGIN_ARGS(SUIComponentListView) {}
+	SLATE_END_ARGS()
+
+	void Construct(const FArguments& InArgs, const TSharedRef<STableViewBase>& InOwnerTableView, TSharedPtr<FWidgetBlueprintEditor> InBlueprintEditor, TSharedPtr<FUIComponentListItem> InListItem)
+	{
+		ListItem = InListItem;
+		BlueprintEditor = InBlueprintEditor;
+
+		STableRow<TSharedPtr<FUIComponentListItem>>::Construct(
+			STableRow<TSharedPtr<FUIComponentListItem>>::FArguments()
+			.Style(&FAppStyle::Get().GetWidgetStyle<FTableRowStyle>("TableView.AlternatingRow"))
+			.Padding(FMargin(3.0f, 2.0f))
+			.OnDragDetected(this, &SUIComponentListItem::OnDragDetected)
+			.OnCanAcceptDrop(this, &SUIComponentListItem::OnCanAcceptDrop)
+			.OnAcceptDrop(this, &SUIComponentListItem::OnAcceptDrop)
+			.Content()
+			[
+				SAssignNew(InlineTextBlock, SInlineEditableTextBlock)
+				.Font(FCoreStyle::Get().GetFontStyle("NormalFont"))
+				.Text(this, &SUIComponentListItem::GetComponentName)
+				.IsReadOnly(true)
+				.OverflowPolicy(ETextOverflowPolicy::Ellipsis)
+				.IsSelected(this, &SUIComponentListItem::IsSelectedExclusively)
+			],
+			InOwnerTableView);
+	}
+
+	void BeginRename()
+	{
+		InlineTextBlock->EnterEditingMode();
+	}
+
+private:
+
+	FText GetComponentName() const
+	{
+		if( TSharedPtr<FUIComponentListItem> Item = ListItem.Pin() )
+		{
+			if (Item->Component != nullptr)
+			{
+				return Item->Component->GetClass()->GetDisplayNameText();
+			}
+		}
+
+		return FText::GetEmpty();
+	}
+
+	/** Called whenever a drag is detected by the list view. */
+	FReply OnDragDetected(const FGeometry& InGeometry, const FPointerEvent& InPointerEvent)
+	{
+		TSharedPtr<FUIComponentListItem> ListItemPinned = ListItem.Pin();
+		if (ListItemPinned.IsValid())
+		{
+			FText DefaultText = LOCTEXT("DefaultDragDropFormat", "Move 1 item(s)");
+			return FReply::Handled().BeginDragDrop(FUIComponentDragDropOp::New(ListItemPinned, DefaultText));
+		}
+		return FReply::Unhandled();
+	}
+
+	/** Called to determine whether a current drag operation is valid for this row. */
+	TOptional<EItemDropZone> OnCanAcceptDrop(const FDragDropEvent& DragDropEvent, EItemDropZone InItemDropZone, TSharedPtr<FUIComponentListItem> InListItem)
+	{
+		TSharedPtr<FUIComponentDragDropOp> DragDropOp = DragDropEvent.GetOperationAs<FUIComponentDragDropOp>();
+		if (DragDropOp.IsValid())
+		{
+			if (InItemDropZone == EItemDropZone::OntoItem)
+			{
+				DragDropOp->CurrentIconBrush = FAppStyle::GetBrush(TEXT("Graph.ConnectorFeedback.Error"));
+			}
+			else
+			{
+				DragDropOp->CurrentIconBrush = FAppStyle::GetBrush(TEXT("Graph.ConnectorFeedback.Ok"));
+			}
+			return InItemDropZone;
+		}
+		return TOptional<EItemDropZone>();
+	}
+
+	/** Called to complete a drag and drop onto this drop. */
+	FReply OnAcceptDrop(const FDragDropEvent& DragDropEvent, EItemDropZone InItemDropZone, TSharedPtr<FUIComponentListItem> InListItem)
+	{
+		TSharedPtr<FUIComponentDragDropOp> DragDropOp = DragDropEvent.GetOperationAs<FUIComponentDragDropOp>();
+		if (DragDropOp.IsValid()
+			&& DragDropOp->ListItem.IsValid() 
+			&& DragDropOp->ListItem->Component != nullptr
+			&& InListItem.IsValid() 
+			&& InListItem->Component != nullptr
+			&& DragDropOp->ListItem->Component != InListItem->Component
+			&& InListItem->Widget != nullptr
+			)
+		{
+			
+			TSharedPtr<FWidgetBlueprintEditor> WidgetBlueprintEditor = BlueprintEditor.Pin();
+			if (UWidgetBlueprint* Blueprint = WidgetBlueprintEditor->GetWidgetBlueprintObj())
+			{
+				// When we drop on another item, we expect to replace it, so we move it before that Item.
+				// So we move it after only if it's Below
+				const bool bMoveAfter = InItemDropZone == EItemDropZone::BelowItem;
+				
+				FUIComponentUtils::MoveComponent(WidgetBlueprintEditor.ToSharedRef(),
+					DragDropOp->ListItem->Component.Get()->GetClass(),
+					InListItem->Component.Get()->GetClass(),
+					InListItem->Widget.Get()->GetFName(),
+					bMoveAfter
+					);
+								
+				return FReply::Handled();
+			}
+		}
+		return FReply::Unhandled();
+	}
+
+private:
+	TWeakPtr<FUIComponentListItem> ListItem;
+	TWeakPtr<FWidgetBlueprintEditor> BlueprintEditor;
+	TSharedPtr<SInlineEditableTextBlock> InlineTextBlock;
+};
+
+SUIComponentView::~SUIComponentView()
+{
+	if (TSharedPtr<FWidgetBlueprintEditor> WidgetBlueprintEditorPin = BlueprintEditor.Pin())
+	{
+		WidgetBlueprintEditorPin->OnSelectedWidgetsChanged.RemoveAll(this);		
+		WidgetBlueprintEditorPin->OnWidgetPreviewUpdated.RemoveAll(this);
+	}
+}
+
+void SUIComponentView::Construct(const FArguments& InArgs, TSharedPtr<FWidgetBlueprintEditor> InBlueprintEditor)
+{
+	BlueprintEditor = InBlueprintEditor;
+
+	InBlueprintEditor->GetOnWidgetBlueprintTransaction().AddSP(this, &SUIComponentView::OnWidgetBlueprintTransaction);
+	InBlueprintEditor->OnSelectedWidgetsChanged.AddSP(this, &SUIComponentView::OnEditorSelectionChanged);
+	InBlueprintEditor->OnWidgetPreviewReady.AddSP(this, &SUIComponentView::OnWidgetPreviewReady);
+		
+	ChildSlot
+	[
+		SNew(SVerticalBox)
+		+ SVerticalBox::Slot()
+		.AutoHeight()
+		.HAlign(EHorizontalAlignment::HAlign_Left)
+		.Padding(4.0f, 0.0f)
+		[
+			SNew(SPositiveActionButton)
+			.Icon(FAppStyle::Get().GetBrush("Icons.Plus"))
+			.Text(LOCTEXT("AddComponentLabel", "Add Component"))
+			.ToolTipText(LOCTEXT("AddComponentToolTip", "Add a Component to this widget."))		
+			.OnClicked(this, &SUIComponentView::OnAddComponentButtonClicked)
+		]
+		+ SVerticalBox::Slot()
+		.AutoHeight()		
+		.HAlign(EHorizontalAlignment::HAlign_Fill)
+		[			
+			SNew(SBorder)
+			.BorderImage(FAppStyle::GetBrush("ToolPanel.DarkGroupBorder"))
+			.Padding(4.f)		
+			[	SNew(SBox)
+				.MinDesiredHeight(60)
+				[
+					SNew(SScrollBorder, ComponentListView.ToSharedRef())
+					[
+						SAssignNew(ComponentListView, SUIComponentListView)
+						.SelectionMode(ESelectionMode::Single)
+						.OnGenerateRow(this, &SUIComponentView::OnGenerateWidgetForComponent)
+						.OnContextMenuOpening(this, &SUIComponentView::OnContextMenuOpening)
+						.ListItemsSource(&Components)
+					]
+				]
+			]			
+		]		
+	];
+
+	UpdateComponentList();
+	CreateCommandList();
+}
+
+void SUIComponentView::CreateCommandList()
+{
+	CommandList = MakeShareable(new FUICommandList);
+
+	CommandList->MapAction(
+		FGenericCommands::Get().Delete,
+		FExecuteAction::CreateSP(this, &SUIComponentView::OnDeleteComponent),
+		FCanExecuteAction::CreateSP(this, &SUIComponentView::CanExecuteContextMenuAction)
+	);
+}
+
+FReply SUIComponentView::OnKeyDown(const FGeometry& MyGeometry, const FKeyEvent& InKeyEvent)
+{
+	FReply Reply = FReply::Unhandled();
+	if (CommandList->ProcessCommandBindings(InKeyEvent))
+	{
+		Reply = FReply::Handled();
+	}
+
+	return Reply;
+}
+
+void SUIComponentView::OnEditorSelectionChanged()
+{
+	// When the selection changes in the Widget Designer, we update the component List
+	UpdateComponentList();
+}
+void SUIComponentView::OnWidgetPreviewReady()
+{
+	// When the the preview is recreated and is ready we update the component List
+	UpdateComponentList();
+}
+
+void SUIComponentView::OnWidgetBlueprintTransaction()
+{	
+	TSharedPtr<FWidgetBlueprintEditor> WidgetBlueprintEditorPin = BlueprintEditor.Pin();
+	UpdateComponentList();
+}
+
+FReply SUIComponentView::OnAddComponentButtonClicked()
+{
+	if (const TSharedPtr<FWidgetBlueprintEditor> WidgetBlueprintEditor = BlueprintEditor.Pin())
+	{
+		UClass* ChosenClass = nullptr;
+		const FText TitleText = LOCTEXT("AddComponentClassPickerTitle", "Pick a component to add to the widget.");
+		const bool bPressedOk = SClassPickerDialog::PickClass(TitleText, FUIComponentUtils::CreateClassViewerInitializationOptions(), ChosenClass, UUIComponent::StaticClass());
+		if (ChosenClass)
+		{
+			const TSet<FWidgetReference>& SelectedWidgets = WidgetBlueprintEditor->GetSelectedWidgets();
+			for (const FWidgetReference& SelectedObject : SelectedWidgets)
+			{
+				if (const UWidget* SelectedWidget = Cast<UWidget>(SelectedObject.GetPreview()))
+				{
+					FUIComponentUtils::AddComponent(WidgetBlueprintEditor.ToSharedRef(), ChosenClass, SelectedWidget->GetFName());
+				}
+			}
+		}
+	}
+	return FReply::Handled();
+}
+
+TSharedPtr<SWidget> SUIComponentView::OnContextMenuOpening()
+{
+	FMenuBuilder MenuBuilder(true, CommandList.ToSharedRef());
+
+	MenuBuilder.BeginSection("Edit", LOCTEXT("Edit", "Edit"));
+	{
+		MenuBuilder.AddMenuEntry(FGenericCommands::Get().Delete);
+	}
+	MenuBuilder.EndSection();
+
+	return MenuBuilder.MakeWidget();
+}
+
+bool SUIComponentView::CanExecuteContextMenuAction() const
+{
+	return true;
+}
+
+void SUIComponentView::OnDeleteComponent()
+{
+	if (!ComponentListView)
+	{
+		return;
+	}
+	
+	if (const TSharedPtr<FWidgetBlueprintEditor> WidgetBlueprintEditor = BlueprintEditor.Pin())
+	{
+		if (UWidget* SelectedWidget = GetSelectedWidget())
+		{
+			for (TSharedPtr<FUIComponentListItem>& Item : ComponentListView->GetSelectedItems())
+			{
+				if (Item && Item->Component != nullptr)
+				{
+					FUIComponentUtils::RemoveComponent(WidgetBlueprintEditor.ToSharedRef(), Item->Component->GetClass(), SelectedWidget->GetFName());
+				}
+			}
+		}		 
+	}
+}
+
+TSharedRef<ITableRow> SUIComponentView::OnGenerateWidgetForComponent(TSharedPtr<FUIComponentListItem> InListItem, const TSharedRef< STableViewBase >& InOwnerTableView)
+{
+	return SNew(SUIComponentListItem, InOwnerTableView, BlueprintEditor.Pin(), InListItem);
+}
+
+UWidget* SUIComponentView::GetSelectedWidget() const
+{
+	UWidget* SelectedPreviewWidget = nullptr;
+	if (const TSharedPtr<FWidgetBlueprintEditor> WidgetBlueprintEditor = BlueprintEditor.Pin())
+	{
+		const TSet<FWidgetReference>& SelectedWidgets = WidgetBlueprintEditor->GetSelectedWidgets();
+		if (SelectedWidgets.Num() == 1)
+		{
+			// Get the first element.
+			const TSet<FWidgetReference>::TConstIterator SetIt(SelectedWidgets);
+			SelectedPreviewWidget = SetIt->GetPreview();
+		}
+	}
+	return SelectedPreviewWidget;
+}
+
+void SUIComponentView::UpdateComponentList()
+{	
+	Components.Reset();
+	
+	if (const TSharedPtr<FWidgetBlueprintEditor> WidgetBlueprintEditor = BlueprintEditor.Pin())
+	{
+		UWidget* SelectedPreviewWidget = GetSelectedWidget();
+		const UUserWidget* PreviewUserWidget = WidgetBlueprintEditor->GetPreview();
+		const UUIComponentUserWidgetExtension* UserWidgetExtension = PreviewUserWidget? PreviewUserWidget->GetExtension<UUIComponentUserWidgetExtension>() : nullptr;
+					
+		if (SelectedPreviewWidget && UserWidgetExtension)
+		{
+			TArray<UUIComponent*> ListOfComponents =  UserWidgetExtension->GetComponentsFor(SelectedPreviewWidget);
+			Components.Reserve(ListOfComponents.Num());
+			for (UUIComponent* Component : ListOfComponents)
+			{
+				Components.Emplace(MakeShared<FUIComponentListItem>(Component, SelectedPreviewWidget));
+			}
+		}		
+	}
+	ComponentListView->RequestListRefresh();
+}
+
+#undef LOCTEXT_NAMESPACE 
\ No newline at end of file
diff --git a/Engine/Source/Editor/UMGEditor/Private/Details/SUIComponentView.h b/Engine/Source/Editor/UMGEditor/Private/Details/SUIComponentView.h
new file mode 100644
--- /dev/null
+++ b/Engine/Source/Editor/UMGEditor/Private/Details/SUIComponentView.h
@@ -0,0 +1,51 @@
+// Copyright Epic Games, Inc. All Rights Reserved.
+
+#pragma once
+
+#include "CoreMinimal.h"
+#include "Widgets/SCompoundWidget.h"
+#include "Widgets/SWidget.h"
+#include "Widgets/Views/SListView.h"
+
+class FWidgetBlueprintEditor;
+class FUICommandList;
+struct FUIComponentListItem;
+class UWidget;
+
+typedef SListView<TSharedPtr<FUIComponentListItem> > SUIComponentListView;
+
+class SUIComponentView final : public SCompoundWidget
+{
+public:
+	SLATE_BEGIN_ARGS(SUIComponentView) {}
+	SLATE_END_ARGS()
+
+	~SUIComponentView();
+
+	void Construct(const FArguments& InArgs, TSharedPtr<FWidgetBlueprintEditor> InBlueprintEditor);
+	void CreateCommandList();
+	
+	virtual FReply OnKeyDown(const FGeometry& MyGeometry, const FKeyEvent& InKeyEvent) override;
+
+private:	
+	void OnEditorSelectionChanged();
+	void OnWidgetPreviewReady();	
+	void OnWidgetBlueprintTransaction();
+	
+	FReply OnAddComponentButtonClicked();
+
+	TSharedPtr<SWidget> OnContextMenuOpening();	
+	bool CanExecuteContextMenuAction() const;
+
+	void OnDeleteComponent();
+	
+	TSharedRef<ITableRow> OnGenerateWidgetForComponent(TSharedPtr<FUIComponentListItem> InListItem, const TSharedRef< STableViewBase >& InOwnerTableView);
+	
+	UWidget* GetSelectedWidget() const;
+	void UpdateComponentList();
+	
+	TSharedPtr<FUICommandList> CommandList;
+	TWeakPtr<FWidgetBlueprintEditor> BlueprintEditor;
+	TSharedPtr<SUIComponentListView> ComponentListView;
+	TArray<TSharedPtr<FUIComponentListItem> > Components;
+};
diff --git a/Engine/Source/Editor/UMGEditor/Private/Details/SWidgetDetailsView.cpp b/Engine/Source/Editor/UMGEditor/Private/Details/SWidgetDetailsView.cpp
--- a/Engine/Source/Editor/UMGEditor/Private/Details/SWidgetDetailsView.cpp
+++ b/Engine/Source/Editor/UMGEditor/Private/Details/SWidgetDetailsView.cpp
@@ -34,12 +34,7 @@
 #include "UMGEditorModule.h"
 #include "WidgetEditingProjectSettings.h"
 
-#include "Extensions/UIComponent.h"
-#include "Extensions/UIComponentContainer.h"
-#include "Kismet2/SClassPickerDialog.h"
-#include "SPositiveActionButton.h"
-#include "UIComponentUtils.h"
-#include "UIComponentWidgetBlueprintExtension.h"
+#include "SUIComponentView.h"
 
 #define LOCTEXT_NAMESPACE "UMG"
 
@@ -191,21 +186,17 @@ void SWidgetDetailsView::Construct(const FArguments& InArgs, TSharedPtr<FWidgetB
 				]
 				// Only if not the user widget
 				+ SVerticalBox::Slot()
-					.AutoHeight()
-					.HAlign(EHorizontalAlignment::HAlign_Left)
-					.Padding(4.0f, 0.0f)
-					[
-						CVarUMGExposeUIComponent->GetBool() ?
-						SNew(SPositiveActionButton)
-						.Icon(FAppStyle::Get().GetBrush("Icons.Plus"))
-						.Text(LOCTEXT("AddNewRetargetOpLabel", "Add New Component"))
-						.ToolTipText(LOCTEXT("AddNewToolTip", "Add a new Component to this widget."))
-						.IsEnabled(this, &SWidgetDetailsView::IsWidgetNameFieldEnabled)
-						.OnClicked(this, &SWidgetDetailsView::OnAddComponentButtonClicked, InBlueprintEditor)
-						.Visibility(this, &SWidgetDetailsView::GetAddComponentAreaVisibility)
-						:
-						SNullWidget::NullWidget
-					]
+				.AutoHeight()
+				.HAlign(EHorizontalAlignment::HAlign_Fill)
+				.Padding(CVarUMGExposeUIComponent->GetBool() ? 4.0f : 0.0f, 0.0f)
+				[
+					CVarUMGExposeUIComponent->GetBool() ?
+					SNew(SUIComponentView, InBlueprintEditor)
+					.IsEnabled(this, &SWidgetDetailsView::IsWidgetNameFieldEnabled)
+					.Visibility(this, &SWidgetDetailsView::GetComponentAreaVisibility)
+					:
+					SNullWidget::NullWidget
+				]
 			]
 		]
 
@@ -438,7 +429,7 @@ EVisibility SWidgetDetailsView::GetCategoryAreaVisibility() const
 	return IsWidgetCDOSelected() ? EVisibility::Visible : EVisibility::Collapsed;
 }
 
-EVisibility SWidgetDetailsView::GetAddComponentAreaVisibility() const
+EVisibility SWidgetDetailsView::GetComponentAreaVisibility() const
 {
 	if (IsWidgetCDOSelected() || SelectedObjects.Num() > 1)
 	{
@@ -657,23 +648,7 @@ void SWidgetDetailsView::HandleIsVariableChanged(ECheckBoxState CheckState)
 	}
 }
 
-FReply SWidgetDetailsView::OnAddComponentButtonClicked(TSharedPtr<FWidgetBlueprintEditor> InBlueprintEditor)
-{
-	UClass* ChosenClass = nullptr;
-	const FText TitleText = LOCTEXT("AddExtension", "Pick a component to add to the widget.");
-	const bool bPressedOk = SClassPickerDialog::PickClass(TitleText, FUIComponentUtils::CreateClassViewerInitializationOptions(), ChosenClass, UUIComponent::StaticClass());
-	if (ChosenClass)
-	{
-		for(TWeakObjectPtr<UObject> SelectedObject : SelectedObjects)
-		{
-			if (const UWidget* SelectedWidget = Cast<UWidget>(SelectedObject.Get()))
-			{
-				FUIComponentUtils::AddComponent(InBlueprintEditor.ToSharedRef(), ChosenClass, SelectedWidget->GetFName());
-			}
-		}
-	}
-	return FReply::Handled();
-}
+
 
 void SWidgetDetailsView::NotifyPreChange(FEditPropertyChain* PropertyAboutToChange)
 {
diff --git a/Engine/Source/Editor/UMGEditor/Private/Details/SWidgetDetailsView.h b/Engine/Source/Editor/UMGEditor/Private/Details/SWidgetDetailsView.h
--- a/Engine/Source/Editor/UMGEditor/Private/Details/SWidgetDetailsView.h
+++ b/Engine/Source/Editor/UMGEditor/Private/Details/SWidgetDetailsView.h
@@ -58,7 +58,7 @@ private:
 
 	EVisibility GetNameAreaVisibility() const;
 	
-	EVisibility GetAddComponentAreaVisibility() const;
+	EVisibility GetComponentAreaVisibility() const;
 
 	const FSlateBrush* GetNameIcon() const;
 
@@ -79,8 +79,6 @@ private:
 	FText GetCategoryText() const;
 	void HandleCategoryTextCommitted(const FText& Text, ETextCommit::Type CommitType);
 
-	FReply OnAddComponentButtonClicked(TSharedPtr<FWidgetBlueprintEditor> InBlueprintEditor);
-
 private:
 	/** The editor that owns this details view */
 	TWeakPtr<class FWidgetBlueprintEditor> BlueprintEditor;
diff --git a/Engine/Source/Editor/UMGEditor/Private/UIComponentUtils.cpp b/Engine/Source/Editor/UMGEditor/Private/UIComponentUtils.cpp
--- a/Engine/Source/Editor/UMGEditor/Private/UIComponentUtils.cpp
+++ b/Engine/Source/Editor/UMGEditor/Private/UIComponentUtils.cpp
@@ -49,7 +49,7 @@ void FUIComponentUtils::AddComponent(const TSharedRef<FWidgetBlueprintEditor>& B
 {	
 	UWidgetBlueprint* WidgetBlueprint = BlueprintEditor->GetWidgetBlueprintObj();
 
-	const FScopedTransaction Transaction(LOCTEXT("AddComponent", "Add Component"));
+	const FScopedTransaction Transaction( FText::Format( LOCTEXT("AddComponentTransaction", "Add Component {0} to {1}"), FText::FromString(ComponentClass->GetName()), FText::FromName(WidgetName)));
 	if (UUIComponentWidgetBlueprintExtension* WidgetBlueprintExtension = UUIComponentWidgetBlueprintExtension::RequestExtension<UUIComponentWidgetBlueprintExtension>(WidgetBlueprint))
 	{
 		UUIComponent* ComponentArchetype = WidgetBlueprintExtension->AddComponent(ComponentClass, WidgetName);		
@@ -71,7 +71,7 @@ void FUIComponentUtils::AddComponent(const TSharedRef<FWidgetBlueprintEditor>& B
 void FUIComponentUtils::RemoveComponent(const TSharedRef<FWidgetBlueprintEditor>& BlueprintEditor, const UClass* ComponentClass, const FName WidgetName)
 {
 	UWidgetBlueprint* WidgetBlueprint = BlueprintEditor->GetWidgetBlueprintObj();
-	const FScopedTransaction Transaction(LOCTEXT("RemoveComponent", "Remove Component"));
+	const FScopedTransaction Transaction(FText::Format(LOCTEXT("RemoveComponentTransaction", "Remove Component {0} from {1}"), FText::FromString(ComponentClass->GetName()), FText::FromName(WidgetName)));
 	if (UUIComponentWidgetBlueprintExtension* WidgetBlueprintExtension = UWidgetBlueprintExtension::GetExtension<UUIComponentWidgetBlueprintExtension>(WidgetBlueprint))
 	{
 		WidgetBlueprintExtension->RemoveComponent(ComponentClass, WidgetName);
@@ -88,6 +88,26 @@ void FUIComponentUtils::RemoveComponent(const TSharedRef<FWidgetBlueprintEditor>
 	}	
 }
 
+void FUIComponentUtils::MoveComponent(const TSharedRef<FWidgetBlueprintEditor>& BlueprintEditor, const UClass* ComponentClassToMove, const UClass* RelativeToComponentClass, const FName WidgetName, bool bMoveAfter)
+{
+	UWidgetBlueprint* WidgetBlueprint = BlueprintEditor->GetWidgetBlueprintObj();
+	const FScopedTransaction Transaction(FText::Format(LOCTEXT("MoveComponentTransaction", "Move Component {0} in {1}"), FText::FromString(ComponentClassToMove->GetName()), FText::FromName(WidgetName)));
+	if (UUIComponentWidgetBlueprintExtension* WidgetBlueprintExtension = UWidgetBlueprintExtension::GetExtension<UUIComponentWidgetBlueprintExtension>(WidgetBlueprint))
+	{
+		WidgetBlueprintExtension->MoveComponent(WidgetName, ComponentClassToMove, RelativeToComponentClass, bMoveAfter);
+
+		// Also Move it in the Preview
+		if (const UUserWidget* PreviewWidget = BlueprintEditor->GetPreview())
+		{
+			if (UUIComponentUserWidgetExtension* UserWidgetExtension = PreviewWidget->GetExtension<UUIComponentUserWidgetExtension>())
+			{
+				UserWidgetExtension->MoveComponent(WidgetName, ComponentClassToMove, RelativeToComponentClass, bMoveAfter);
+			}
+		}
+		FBlueprintEditorUtils::MarkBlueprintAsStructurallyModified(WidgetBlueprint);
+	}	
+}
+
 bool FUIComponentUtils::FUIComponentClassFilter::IsClassAllowed(const FClassViewerInitializationOptions& InInitOptions, const UClass* InClass, TSharedRef< FClassViewerFilterFuncs > InFilterFuncs)
 {
 	return !InClass->HasAnyClassFlags(DisallowedClassFlags)
diff --git a/Engine/Source/Editor/UMGEditor/Private/UIComponentWidgetBlueprintExtension.cpp b/Engine/Source/Editor/UMGEditor/Private/UIComponentWidgetBlueprintExtension.cpp
--- a/Engine/Source/Editor/UMGEditor/Private/UIComponentWidgetBlueprintExtension.cpp
+++ b/Engine/Source/Editor/UMGEditor/Private/UIComponentWidgetBlueprintExtension.cpp
@@ -35,19 +35,28 @@ UUIComponentContainer* UUIComponentWidgetBlueprintExtension::DuplicateContainer(
 void UUIComponentWidgetBlueprintExtension::RemoveComponent(const UClass* ComponentClass, FName OwnerName)
 {
 	// Remove the Component
-	if(ensure(ComponentContainer))
+	if (ensure(ComponentContainer))
 	{
 		ComponentContainer->RemoveAllComponentsOfType(ComponentClass, OwnerName);
 	}
 }
 
+void UUIComponentWidgetBlueprintExtension::MoveComponent(FName OwnerName, const UClass* ComponentClass, const UClass* RelativeToComponentClass, bool bMoveAfter)
+{
+	// Remove the Component
+	if (ensure(ComponentContainer))
+	{
+		ComponentContainer->MoveComponent(OwnerName, ComponentClass, RelativeToComponentClass, bMoveAfter);
+	}
+}
+
 TArray<UUIComponent*> UUIComponentWidgetBlueprintExtension::GetComponentsFor(const UWidget* Target) const
 {
 	TArray<UUIComponent*> Components;
-	if(ensure(Target))
+	if (ensure(Target))
 	{
 		ComponentContainer->ForEachComponentTarget([Target, &Components](const FUIComponentTarget& ComponentTarget){
-			if(ComponentTarget.GetTargetName() == Target->GetFName())			
+			if (ComponentTarget.GetTargetName() == Target->GetFName())			
 			{
 				Components.Push(ComponentTarget.GetComponent());
 			}
@@ -76,7 +85,7 @@ void UUIComponentWidgetBlueprintExtension::RenameWidget(const FName& OldVarName,
 bool UUIComponentWidgetBlueprintExtension::VerifyContainer(UUserWidget* UserWidget) const
 {
 	const UUIComponentUserWidgetExtension* UserWidgetExtension = UserWidget->GetExtension<UUIComponentUserWidgetExtension>();
-	if(UserWidgetExtension)
+	if (UserWidgetExtension)
 	{
 		TArray<FUIComponentTarget> Components;
 		ComponentContainer->ForEachComponentTarget([&Components](const FUIComponentTarget& ComponentTarget){
@@ -86,7 +95,7 @@ bool UUIComponentWidgetBlueprintExtension::VerifyContainer(UUserWidget* UserWidg
 		while(Components.Num() > 0)
 		{
 			FUIComponentTarget& Target = Components[Components.Num()-1];
-			if(!UserWidgetExtension->GetComponent(Target.GetComponent()->GetClass(), Target.GetTargetName()))
+			if (!UserWidgetExtension->GetComponent(Target.GetComponent()->GetClass(), Target.GetTargetName()))
 			{
 				return false;
 			}
@@ -181,9 +190,9 @@ void UUIComponentWidgetBlueprintExtension::HandleFinishCompilingClass(UWidgetBlu
 	Super::HandleFinishCompilingClass(Class);
 
 	// If we do not have Components, do not add the extension to the Generated Class
-	if(ensure(CompilerContext) && !ComponentContainer->IsEmpty())
+	if (ensure(CompilerContext) && !ComponentContainer->IsEmpty())
 	{
-		if(UWidgetBlueprintGeneratedClass* NewWidgetGeneratedClass = Cast<UWidgetBlueprintGeneratedClass>(CompilerContext->NewClass))
+		if (UWidgetBlueprintGeneratedClass* NewWidgetGeneratedClass = Cast<UWidgetBlueprintGeneratedClass>(CompilerContext->NewClass))
 		{
 			UUIComponentWidgetBlueprintGeneratedClassExtension* NewExtension = NewObject<UUIComponentWidgetBlueprintGeneratedClassExtension>(NewWidgetGeneratedClass);
 			CompilerContext->AddExtension(NewWidgetGeneratedClass, NewExtension);
@@ -206,7 +215,7 @@ bool UUIComponentWidgetBlueprintExtension::HandleValidateGeneratedClass(UWidgetB
 	}
 
 	// Validate with WidgetTree if all widget exists.
-	if(UWidgetBlueprintGeneratedClass* NewWidgetGeneratedClass = Cast<UWidgetBlueprintGeneratedClass>(CompilerContext->NewClass))
+	if (UWidgetBlueprintGeneratedClass* NewWidgetGeneratedClass = Cast<UWidgetBlueprintGeneratedClass>(CompilerContext->NewClass))
 	{
 		if (const UUIComponentWidgetBlueprintGeneratedClassExtension* Extension = NewWidgetGeneratedClass->GetExtension<UUIComponentWidgetBlueprintGeneratedClassExtension>())
 		{
diff --git a/Engine/Source/Editor/UMGEditor/Public/UIComponentUtils.h b/Engine/Source/Editor/UMGEditor/Public/UIComponentUtils.h
--- a/Engine/Source/Editor/UMGEditor/Public/UIComponentUtils.h
+++ b/Engine/Source/Editor/UMGEditor/Public/UIComponentUtils.h
@@ -36,6 +36,7 @@ public:
 	static UE_API void OnWidgetRenamed(const TSharedRef<FWidgetBlueprintEditor>& BlueprintEditor, UWidgetBlueprint* WidgetBlueprint, const FName& OldVarName, const FName& NewVarName);
 	static UE_API void AddComponent(const TSharedRef<FWidgetBlueprintEditor>& BlueprintEditor, const UClass* ComponentClass, const FName WidgetName);	
 	static UE_API void RemoveComponent(const TSharedRef<FWidgetBlueprintEditor>& BlueprintEditor, const UClass* ComponentClass, const FName WidgetName);
+	static UE_API void MoveComponent(const TSharedRef<FWidgetBlueprintEditor>& BlueprintEditor, const UClass* ComponentClassToMove, const UClass* RelativeToComponentClass, const FName WidgetName, bool bMoveAfter);
 	
 };
 
diff --git a/Engine/Source/Editor/UMGEditor/Public/UIComponentWidgetBlueprintExtension.h b/Engine/Source/Editor/UMGEditor/Public/UIComponentWidgetBlueprintExtension.h
--- a/Engine/Source/Editor/UMGEditor/Public/UIComponentWidgetBlueprintExtension.h
+++ b/Engine/Source/Editor/UMGEditor/Public/UIComponentWidgetBlueprintExtension.h
@@ -31,6 +31,7 @@ public:
 
 	UE_API UUIComponent* AddComponent(const UClass* ComponentClass, FName OwnerName);
 	UE_API void RemoveComponent(const UClass* ComponentClass, FName OwnerName);
+	UE_API void MoveComponent(FName OwnerName, const UClass* ComponentClass, const UClass* RelativeToComponentClass, bool bMoveAfter);
 	UE_API TArray<UUIComponent*> GetComponentsFor(const UWidget* Target) const;
 	UE_API UUIComponent* GetComponent(const UClass* ComponentClass, FName OwnerName) const;
 
diff --git a/Engine/Source/Editor/UMGEditor/Public/WidgetBlueprintEditor.h b/Engine/Source/Editor/UMGEditor/Public/WidgetBlueprintEditor.h
--- a/Engine/Source/Editor/UMGEditor/Public/WidgetBlueprintEditor.h
+++ b/Engine/Source/Editor/UMGEditor/Public/WidgetBlueprintEditor.h
@@ -60,6 +60,9 @@ public:
 	DECLARE_MULTICAST_DELEGATE(FOnSelectedAnimationChanged)
 
 	DECLARE_EVENT(FWidgetBlueprintEditor, FOnEnterWidgetDesigner)
+	
+	/** Called after the widget preview has fully recreated */
+	DECLARE_MULTICAST_DELEGATE(FOnWidgetPreviewReady)
 
 public:
 	UE_API FWidgetBlueprintEditor();
@@ -282,9 +285,12 @@ public:
 	/** Notification for when the preview widget has been updated */
 	FOnWidgetPreviewUpdated OnWidgetPreviewUpdated;
 
-	/** Notification for when the preview widget has been updated */
+	/** Notification for when animations have been updated */
 	FOnWidgetAnimationsUpdated OnWidgetAnimationsUpdated;
 
+	/** Notification for when the preview widget has been fully recreated */
+	FOnWidgetPreviewReady OnWidgetPreviewReady;
+	
 	/** Fires after the mode change to Designer*/
 	FOnEnterWidgetDesigner OnEnterWidgetDesigner;
 
diff --git a/Engine/Source/Runtime/UMG/Private/Extensions/UIComponentContainer.cpp b/Engine/Source/Runtime/UMG/Private/Extensions/UIComponentContainer.cpp
--- a/Engine/Source/Runtime/UMG/Private/Extensions/UIComponentContainer.cpp
+++ b/Engine/Source/Runtime/UMG/Private/Extensions/UIComponentContainer.cpp
@@ -65,7 +65,7 @@ void UUIComponentContainer::RemoveComponent(FName TargetName, UUIComponent* Comp
 			{
 				ensure(Target.GetTargetName() == TargetName);
 				UpdateModify();
-				Components.RemoveAtSwap(i);
+				Components.RemoveAt(i);
 				return;
 			}
 		}
@@ -115,9 +115,53 @@ void UUIComponentContainer::RemoveAllComponentsFor(FName TargetName)
 		if (Target.GetTargetName() == TargetName)
 		{
 			UpdateModify();
-			Components.RemoveAtSwap(Index);
+			Components.RemoveAt(Index);
+		}
+	}
+}
+
+void UUIComponentContainer::MoveComponent(FName TargetName, const UClass* ComponentClassToMove, const UClass* RelativeToComponentClass, bool bMoveAfter)
+{
+	if (!ensure(!TargetName.IsNone() && ComponentClassToMove != nullptr && ComponentClassToMove != RelativeToComponentClass))
+	{
+		return;
+	}
+
+	FUIComponentTarget TargetToMove;
+	for (int i = 0; i < Components.Num(); i++)
+	{
+		FUIComponentTarget& Target = Components[i];
+		if (Target.GetTargetName() == TargetName &&
+			Target.GetComponent()->GetClass() == ComponentClassToMove)
+		{
+			TargetToMove = Target;
+											
+			Modify();
+			Components.RemoveAt(i, EAllowShrinking::No);
+			break;			
+		}
+	}
+	// If we didn't find the Component to move we are done
+	if (!TargetToMove.IsValid())
+	{
+		return;
+	}
+
+	int IndexToAddTo = 0;
+	if (RelativeToComponentClass)
+	{
+		for (int i = 0; i < Components.Num(); i++)
+		{
+			FUIComponentTarget& Target = Components[i];
+			if (Target.GetTargetName() == TargetName &&
+				Target.GetComponent()->GetClass() == RelativeToComponentClass)
+			{
+				IndexToAddTo = bMoveAfter ? (i+1) : i;				
+				break;
+			}
 		}
 	}
+	Components.Insert(TargetToMove, IndexToAddTo);	
 }
 
 UUIComponent* UUIComponentContainer::GetComponent(const UClass* ComponentClass, FName TargetName) const
@@ -220,7 +264,7 @@ void UUIComponentContainer::CleanupUIComponents(const UWidgetTree* WidgetTree)
 			if (Components[Index].GetComponent() == nullptr)
 			{
 				UpdateModify();
-				Components.RemoveAtSwap(Index);
+				Components.RemoveAt(Index);
 			}
 			else if (!Components[Index].GetTargetName().IsNone())
 			{
diff --git a/Engine/Source/Runtime/UMG/Private/Extensions/UIComponentUserWidgetExtension.cpp b/Engine/Source/Runtime/UMG/Private/Extensions/UIComponentUserWidgetExtension.cpp
--- a/Engine/Source/Runtime/UMG/Private/Extensions/UIComponentUserWidgetExtension.cpp
+++ b/Engine/Source/Runtime/UMG/Private/Extensions/UIComponentUserWidgetExtension.cpp
@@ -110,11 +110,21 @@ void UUIComponentUserWidgetExtension::CreateAndAddComponent(UUIComponent* Archet
 
 void UUIComponentUserWidgetExtension::RemoveComponent(const UClass* ComponentClass, FName OwnerName)
 {
-	if(ComponentContainer)
+	if (ComponentContainer)
 	{
 		ComponentContainer->RemoveAllComponentsOfType(ComponentClass, OwnerName);
 	}
 }
+
+
+void UUIComponentUserWidgetExtension::MoveComponent(FName OwnerName, const UClass* ComponentClass, const UClass* RelativeToComponentClass, bool bMoveAfter)
+{
+	if (ComponentContainer)
+	{
+		ComponentContainer->MoveComponent(OwnerName, ComponentClass, RelativeToComponentClass, bMoveAfter);
+	}
+}
+
 #endif // WITH_EDITOR
 
 TArray<UUIComponent*> UUIComponentUserWidgetExtension::GetComponentsFor(UWidget* Target) const
diff --git a/Engine/Source/Runtime/UMG/Public/Extensions/UIComponentContainer.h b/Engine/Source/Runtime/UMG/Public/Extensions/UIComponentContainer.h
--- a/Engine/Source/Runtime/UMG/Public/Extensions/UIComponentContainer.h
+++ b/Engine/Source/Runtime/UMG/Public/Extensions/UIComponentContainer.h
@@ -30,6 +30,7 @@ public:
 	FName GetTargetName() const { return TargetName; }
 	void SetTargetName(FName NewName);
 	UUIComponent* GetComponent() const { return Component; }
+	bool IsValid() const { return !TargetName.IsNone() && Component != nullptr; }
 
 private:
 
@@ -57,6 +58,8 @@ public:
 	UE_API void RemoveAllComponentsOfType(const UClass* ComponentClass, FName TargetName);
 	UE_API void RemoveAllComponentsFor(FName TargetName);
 
+	UE_API void MoveComponent(FName TargetName, const UClass* ComponentClass, const UClass* RelativeToComponentClass, bool bMoveAfter);
+
 	UE_API UUIComponent* GetComponent(const UClass* ComponentClass, FName TargetName) const; 
 	
 	UE_API void ForEachComponent(TFunctionRef<void(UUIComponent*)> Predicate) const;
diff --git a/Engine/Source/Runtime/UMG/Public/Extensions/UIComponentUserWidgetExtension.h b/Engine/Source/Runtime/UMG/Public/Extensions/UIComponentUserWidgetExtension.h
--- a/Engine/Source/Runtime/UMG/Public/Extensions/UIComponentUserWidgetExtension.h
+++ b/Engine/Source/Runtime/UMG/Public/Extensions/UIComponentUserWidgetExtension.h
@@ -41,6 +41,8 @@ public:
 	UE_API void CreateAndAddComponent(UUIComponent* ArchetypeComponent, FName OwnerName);
 
 	UE_API void RemoveComponent(const UClass* ComponentClass, FName OwnerName);
+	
+	UE_API void MoveComponent(FName OwnerName, const UClass* ComponentClass, const UClass* RelativeToComponentClass, bool bMoveAfter);
 
 #endif // WITH_EDITOR
 
