UMG Designer: (Experimental) Add Component List View to allow reordering.
- Fix to potential nullptr access
- Fix to remove delegate

diff --git a/Engine/Source/Editor/UMGEditor/Private/Details/SUIComponentView.cpp b/Engine/Source/Editor/UMGEditor/Private/Details/SUIComponentView.cpp
--- a/Engine/Source/Editor/UMGEditor/Private/Details/SUIComponentView.cpp
+++ b/Engine/Source/Editor/UMGEditor/Private/Details/SUIComponentView.cpp
@@ -187,8 +187,9 @@ SUIComponentView::~SUIComponentView()
 {
 	if (TSharedPtr<FWidgetBlueprintEditor> WidgetBlueprintEditorPin = BlueprintEditor.Pin())
 	{
-		WidgetBlueprintEditorPin->OnSelectedWidgetsChanged.RemoveAll(this);		
-		WidgetBlueprintEditorPin->OnWidgetPreviewUpdated.RemoveAll(this);
+		WidgetBlueprintEditorPin->OnWidgetPreviewReady.RemoveAll(this);		
+		WidgetBlueprintEditorPin->OnSelectedWidgetsChanged.RemoveAll(this);
+		WidgetBlueprintEditorPin->GetOnWidgetBlueprintTransaction().RemoveAll(this);
 	}
 }
 
@@ -383,7 +384,10 @@ void SUIComponentView::UpdateComponentList()
 			}
 		}		
 	}
-	ComponentListView->RequestListRefresh();
+	if (ComponentListView.IsValid())
+	{
+		ComponentListView->RequestListRefresh();
+	}
 }
 
 #undef LOCTEXT_NAMESPACE 
\ No newline at end of file
