UMG - Hierarchy - Widget Locking tests added

#rnx

#virtualized

diff --git a/Engine/Source/Editor/UMGEditor/Private/Hierarchy/SHierarchyViewItem.cpp b/Engine/Source/Editor/UMGEditor/Private/Hierarchy/SHierarchyViewItem.cpp
--- a/Engine/Source/Editor/UMGEditor/Private/Hierarchy/SHierarchyViewItem.cpp
+++ b/Engine/Source/Editor/UMGEditor/Private/Hierarchy/SHierarchyViewItem.cpp
@@ -1595,6 +1595,7 @@ void SHierarchyViewItem::Construct(const FArguments& InArgs, const TSharedRef< S
 				.ToolTipText(LOCTEXT("WidgetLockedButtonToolTip", "Locks or Unlocks this widget and all children.  Locking a widget prevents it from being selected in the designer view by clicking on them.\n\nHolding [Shift] will only affect this widget and no children."))
 				.HAlign(HAlign_Center)
 				.VAlign(VAlign_Center)
+				.AddMetaData<FTagMetaData>(FTagMetaData(TEXT("ToggleLockedInDesigner")))
 				[
 					SNew(SBox)
 					.MinDesiredWidth(12.0f)
@@ -1769,12 +1770,20 @@ FText SHierarchyViewItem::GetVisibilityBrushForWidget() const
 	return Model->IsVisible() ? FEditorFontGlyphs::Eye : FEditorFontGlyphs::Eye_Slash;
 }
 
+void SHierarchyViewItem::ToggleLockedInDesigner(const bool bRecursive)
+{
+	if (Model.IsValid())
+	{
+		Model->SetIsLockedInDesigner(!Model->IsLockedInDesigner(), bRecursive);
+	}
+}
+
 FReply SHierarchyViewItem::OnToggleLockedInDesigner()
 {
 	if ( Model.IsValid() )
 	{
 		const bool bRecursive = FSlateApplication::Get().GetModifierKeys().IsShiftDown() ? false : true;
-		Model->SetIsLockedInDesigner(!Model->IsLockedInDesigner(), bRecursive);
+		ToggleLockedInDesigner(bRecursive);
 	}
 
 	return FReply::Handled();
diff --git a/Engine/Source/Editor/UMGEditor/Private/Hierarchy/SHierarchyViewItem.h b/Engine/Source/Editor/UMGEditor/Private/Hierarchy/SHierarchyViewItem.h
--- a/Engine/Source/Editor/UMGEditor/Private/Hierarchy/SHierarchyViewItem.h
+++ b/Engine/Source/Editor/UMGEditor/Private/Hierarchy/SHierarchyViewItem.h
@@ -9,11 +9,11 @@
 #include "Input/Reply.h"
 #include "Widgets/DeclarativeSyntaxSupport.h"
 #include "Widgets/Views/STableViewBase.h"
-#include "Widgets/Views/STableRow.h"
 #include "WidgetReference.h"
 #include "WidgetBlueprintEditor.h"
+#include "Hierarchy/IHierarchyViewItem.h"
 
-class FHierarchyModel : public TSharedFromThis < FHierarchyModel >
+class FHierarchyModel : public IHierarchyModel, public TSharedFromThis < FHierarchyModel >
 {
 public:
 	FHierarchyModel(TSharedPtr<FWidgetBlueprintEditor> InBlueprintEditor);
@@ -60,13 +60,13 @@ public:
 	bool ContainsSelection();
 	bool IsSelected() const;
 
-	virtual bool IsHovered() const { return false; }
-	virtual bool IsVisible() const { return true; }
-	virtual bool CanControlVisibility() const { return false; }
+	virtual bool IsHovered() const override { return false; }
+	virtual bool IsVisible() const override { return true; }
+	virtual bool CanControlVisibility() const override { return false; }
 	virtual void SetIsVisible(bool IsVisible) { }
 
-	virtual bool CanControlLockedInDesigner() const { return false; }
-	virtual bool IsLockedInDesigner() const { return false; }
+	virtual bool CanControlLockedInDesigner() const override { return false; }
+	virtual bool IsLockedInDesigner() const override { return false; }
 	virtual void SetIsLockedInDesigner(bool NewIsLocked, bool bRecursive)
 	{
 		if (bRecursive)
@@ -371,7 +371,7 @@ private:
 /**
  * An widget item in the hierarchy tree view.
  */
-class SHierarchyViewItem : public STableRow< TSharedPtr<FHierarchyModel> >
+class SHierarchyViewItem : public IHierarchyViewItem< FHierarchyModel >
 {
 public:
 	SLATE_BEGIN_ARGS( SHierarchyViewItem ){}
@@ -390,7 +390,7 @@ public:
 	// End SWidget
 
 public:
-	FReply HandleDragDetected(const FGeometry& MyGeometry, const FPointerEvent& MouseEvent);
+	virtual FReply HandleDragDetected(const FGeometry& MyGeometry, const FPointerEvent& MouseEvent) override;
 	void HandleDragEnter(const FDragDropEvent& DragDropEvent);
 	void HandleDragLeave(const FDragDropEvent& DragDropEvent);
 
@@ -410,7 +410,7 @@ private:
 
 	bool IsReadOnly() const;
 	void OnRequestBeginRename();
-	bool ShouldAppearHovered() const;
+	virtual bool ShouldAppearHovered() const override;
 
 	/** Gets the font to use for the text item, bold for customized named items */
 	FSlateFontInfo GetItemFont() const;
@@ -421,6 +421,9 @@ private:
 	/** Handles clicking the visibility toggle */
 	FReply OnToggleVisibility();
 
+	/** Toggles lock status */
+	virtual void ToggleLockedInDesigner(const bool bRecursive) override;
+
 	/** Handles clicking the locked toggle */
 	FReply OnToggleLockedInDesigner();
 
@@ -442,7 +445,7 @@ private:
 	FText InitialText;
 
 	/** Keep an internal IsHovered flag*/
-	bool bHovered;
+	bool bHovered = false;
 	
 	/** If the user is dragging over a collapsed item, wait before expanding it */
 	TSharedPtr<FActiveTimerHandle> DragHoverExpandTimer;
diff --git a/Engine/Source/Editor/UMGEditor/Public/Hierarchy/IHierarchyViewItem.h b/Engine/Source/Editor/UMGEditor/Public/Hierarchy/IHierarchyViewItem.h
new file mode 100644
--- /dev/null
+++ b/Engine/Source/Editor/UMGEditor/Public/Hierarchy/IHierarchyViewItem.h
@@ -0,0 +1,26 @@
+// Copyright Epic Games, Inc. All Rights Reserved.
+
+#pragma once
+
+#include "Widgets/Views/STableRow.h"
+
+class IHierarchyModel
+{
+public:
+	virtual ~IHierarchyModel() { }
+	virtual bool IsHovered() const = 0;
+	virtual bool IsVisible() const = 0;
+	virtual bool CanControlVisibility() const = 0;
+	virtual bool CanControlLockedInDesigner() const = 0;
+	virtual bool IsLockedInDesigner() const = 0;
+};
+
+template <typename HierarchyModel>
+class IHierarchyViewItem : public STableRow< TSharedPtr<HierarchyModel> >
+{
+public:
+	virtual ~IHierarchyViewItem() { }
+	virtual FReply HandleDragDetected(const FGeometry& MyGeometry, const FPointerEvent& MouseEvent) = 0;
+	virtual void ToggleLockedInDesigner(const bool bRecursive) = 0;
+	virtual bool ShouldAppearHovered() const = 0;
+};
\ No newline at end of file
