UMGEditor: Ensure WidgetAnimationItem pointers are valid after list regeneration post transaction undo/redo.

#rnx

diff --git a/Engine/Source/Editor/UMGEditor/Private/TabFactory/AnimationTabSummoner.cpp b/Engine/Source/Editor/UMGEditor/Private/TabFactory/AnimationTabSummoner.cpp
--- a/Engine/Source/Editor/UMGEditor/Private/TabFactory/AnimationTabSummoner.cpp
+++ b/Engine/Source/Editor/UMGEditor/Private/TabFactory/AnimationTabSummoner.cpp
@@ -500,6 +500,18 @@ public:
 
 private:
 
+	void RefreshCurrentAnimationSelection()
+	{		
+		if(const UWidgetAnimation* CurrentAnim = BlueprintEditor.Pin()->RefreshCurrentAnimation())
+		{
+			const TSharedPtr<FWidgetAnimationListItem>* FoundListItemPtr = Animations.FindByPredicate([&](const TSharedPtr<FWidgetAnimationListItem>& ListItem) { return ListItem->Animation == CurrentAnim; });
+			if (FoundListItemPtr != nullptr)
+			{
+				AnimationListView->SetSelection(*FoundListItemPtr);
+			}
+		}
+	}
+
 	void OnUpdatedAnimationList()
 	{
 		Animations.Empty();
@@ -511,6 +523,8 @@ private:
 			Animations.Add(MakeShareable(new FWidgetAnimationListItem(Animation)));
 		}
 
+		RefreshCurrentAnimationSelection();
+
 		AnimationListView->RequestListRefresh();
 	}
 
@@ -558,17 +572,8 @@ private:
 
 
 		TSharedPtr<FWidgetBlueprintEditor> WidgetBlueprintEditorPin = BlueprintEditor.Pin();
-		const UWidgetAnimation* ViewedAnim = WidgetBlueprintEditorPin->RefreshCurrentAnimation();
-
-		if (ViewedAnim)
-		{
-			const TSharedPtr<FWidgetAnimationListItem>* FoundListItemPtr = Animations.FindByPredicate([&](const TSharedPtr<FWidgetAnimationListItem>& ListItem) { return ListItem->Animation == ViewedAnim; });
 
-			if (FoundListItemPtr != nullptr)
-			{
-				AnimationListView->SetSelection(*FoundListItemPtr);
-			}
-		}
+		RefreshCurrentAnimationSelection();
 
 		UWidgetAnimation* CurrentAnim = WidgetBlueprintEditorPin->GetCurrentAnimation();
 		WidgetBlueprintEditorPin->ChangeViewedAnimation(*CurrentAnim);
