UMG Designer. UI Components delete operation in the context menu of the Compoent Section of the detail panel is now greyed out when no Compoents is selected.

diff --git a/Engine/Source/Editor/UMGEditor/Private/Details/SUIComponentView.cpp b/Engine/Source/Editor/UMGEditor/Private/Details/SUIComponentView.cpp
--- a/Engine/Source/Editor/UMGEditor/Private/Details/SUIComponentView.cpp
+++ b/Engine/Source/Editor/UMGEditor/Private/Details/SUIComponentView.cpp
@@ -248,8 +248,8 @@ void SUIComponentView::CreateCommandList()
 
 	CommandList->MapAction(
 		FGenericCommands::Get().Delete,
-		FExecuteAction::CreateSP(this, &SUIComponentView::OnDeleteComponent),
-		FCanExecuteAction::CreateSP(this, &SUIComponentView::CanExecuteContextMenuAction)
+		FExecuteAction::CreateSP(this, &SUIComponentView::OnDeleteSelectedComponent),
+		FCanExecuteAction::CreateSP(this, &SUIComponentView::CanExecuteDeleteSelectedComponent)
 	);
 }
 
@@ -316,12 +316,18 @@ TSharedPtr<SWidget> SUIComponentView::OnContextMenuOpening()
 	return MenuBuilder.MakeWidget();
 }
 
-bool SUIComponentView::CanExecuteContextMenuAction() const
+bool SUIComponentView::CanExecuteDeleteSelectedComponent() const
 {
-	return true;
+	if (ComponentListView)
+	{
+
+		return ComponentListView->GetNumItemsSelected() > 0;
+	}	
+
+	return false;
 }
 
-void SUIComponentView::OnDeleteComponent()
+void SUIComponentView::OnDeleteSelectedComponent()
 {
 	if (!ComponentListView)
 	{
diff --git a/Engine/Source/Editor/UMGEditor/Private/Details/SUIComponentView.h b/Engine/Source/Editor/UMGEditor/Private/Details/SUIComponentView.h
--- a/Engine/Source/Editor/UMGEditor/Private/Details/SUIComponentView.h
+++ b/Engine/Source/Editor/UMGEditor/Private/Details/SUIComponentView.h
@@ -34,10 +34,10 @@ private:
 	
 	FReply OnAddComponentButtonClicked();
 
-	TSharedPtr<SWidget> OnContextMenuOpening();	
-	bool CanExecuteContextMenuAction() const;
+	TSharedPtr<SWidget> OnContextMenuOpening();
 
-	void OnDeleteComponent();
+	bool CanExecuteDeleteSelectedComponent() const;
+	void OnDeleteSelectedComponent();
 	
 	TSharedRef<ITableRow> OnGenerateWidgetForComponent(TSharedPtr<FUIComponentListItem> InListItem, const TSharedRef< STableViewBase >& InOwnerTableView);
 	
