Optimize the widget palette - which refreshes itself twice when a widget blueprint is compiled causing a large hitch

[RN] The widget editor's palette has been optimized to reduce refresh cost in large projects

diff --git a/Engine/Source/Editor/UMGEditor/Private/Palette/SPaletteViewModel.cpp b/Engine/Source/Editor/UMGEditor/Private/Palette/SPaletteViewModel.cpp
--- a/Engine/Source/Editor/UMGEditor/Private/Palette/SPaletteViewModel.cpp
+++ b/Engine/Source/Editor/UMGEditor/Private/Palette/SPaletteViewModel.cpp
@@ -22,6 +22,7 @@
 #include "Templates/WidgetTemplateBlueprintClass.h"
 
 #include "AssetRegistry/AssetRegistryModule.h"
+#include "AssetRegistry/AssetRegistryHelpers.h"
 #include "WidgetBlueprintEditorUtils.h"
 #include "Misc/NamePermissionList.h"
 
@@ -226,15 +227,19 @@ void FWidgetCatalogViewModel::BuildClassWidgetList()
 	}
 
 	// Locate all UWidget classes from code and loaded widget BPs
-	for (TObjectIterator<UClass> ClassIt; ClassIt; ++ClassIt)
+	TArray<UClass*> Widgets;
+	GetDerivedClasses(UWidget::StaticClass(), Widgets);
+	TSet<FName> ProcessedPackages;
+	for (UClass* WidgetClass : Widgets)
 	{
-		UClass* WidgetClass = *ClassIt;
 		const bool bIsSameClass = WidgetClass->GetFName() == ActiveWidgetBlueprintClassName;
 		if (bIsSameClass)
 		{
 			continue;
 		}
-
+		
+		// even if we reject the widget, consider its package processed:
+		ProcessedPackages.Add(WidgetClass->GetPackage()->GetFName());
 		if (!FWidgetBlueprintEditorUtils::IsUsableWidgetClass(WidgetClass, PinnedBPEditor.ToSharedRef()))
 		{
 			continue;
@@ -258,13 +263,24 @@ void FWidgetCatalogViewModel::BuildClassWidgetList()
 
 	// Locate all UWidget BP assets, include loaded and unloaded. Only parsed the unloaded.
 	const FAssetRegistryModule* AssetRegistryModule = FModuleManager::GetModulePtr<FAssetRegistryModule>(TEXT("AssetRegistry"));
-	TArray<FAssetData> AllBPsAssetData;
-	if (AssetRegistryModule)
+	if (!AssetRegistryModule)
 	{
-		AssetRegistryModule->Get().GetAssetsByClass(UBlueprint::StaticClass()->GetClassPathName(), AllBPsAssetData, true);
+		return;
 	}
-
-	for (const FAssetData& BPAssetData : AllBPsAssetData)
+	
+	TArray<FAssetData> AllWidgetClassAssetData;
+	UAssetRegistryHelpers::GetDerivedClassAssetData(
+		{
+			UWidget::StaticClass()->GetClassPathName()
+		},
+		AllWidgetClassAssetData
+	);
+	
+	// for a large working set with thousands of entries in AllWidgetClassAssetData
+	// this loop costs us 100ms - acceptable for editor performance. Much of the
+	// time is spent finding the native parent class, which we could probably do
+	// more efficiently.
+	for (const FAssetData& BPAssetData : AllWidgetClassAssetData)
 	{
 		const bool bIsSameClass = BPAssetData.AssetName == ActiveWidgetBlueprintClassName;
 		if (bIsSameClass)
@@ -272,12 +288,13 @@ void FWidgetCatalogViewModel::BuildClassWidgetList()
 			continue;
 		}
 
-		// Was already parsed by the TObjectIterator<UClass>
-		if (BPAssetData.IsAssetLoaded())
+		// Was already parsed by the GetDerivedClasses(UWidget::StaticClass())
+		if (ProcessedPackages.Contains(BPAssetData.PackageName))
 		{
 			continue;
 		}
 
+		ProcessedPackages.Add(BPAssetData.PackageName);
 		TValueOrError<FWidgetBlueprintEditorUtils::FUsableWidgetClassResult, void> Usable = FWidgetBlueprintEditorUtils::IsUsableWidgetClass(BPAssetData, PinnedBPEditor.ToSharedRef());
 		if (Usable.HasError())
 		{
@@ -299,12 +316,16 @@ void FWidgetCatalogViewModel::BuildClassWidgetList()
 		}
 	}
 
-
 	TArray<FAssetData> AllGeneratedBPsAssetData; // if it's a widget already compiled
 
-	if (AssetRegistryModule && FWidgetBlueprintEditorUtils::GetRelevantSettings(BlueprintEditor)->bUseEditorConfigPaletteFiltering)
+	// this logic will take >500ms in a large project, but only if bUseEditorConfigPaletteFiltering is enabled
+	// It likely can be replaced with a call to UAssetRegistryHelpers::GetDerivedClassAssetData
+	// unfortunately at least for loose cooked classes that function does not work, so i'm leaving 
+	// this in place for users relying on bUseEditorConfigPaletteFiltering
+	if (FWidgetBlueprintEditorUtils::GetRelevantSettings(BlueprintEditor)->bUseEditorConfigPaletteFiltering)
 	{
-		AssetRegistryModule->Get().GetAssetsByClass(UBlueprintGeneratedClass::StaticClass()->GetClassPathName(), AllGeneratedBPsAssetData, true);
+		const IAssetRegistry& AssetRegistry = AssetRegistryModule->Get();
+		AssetRegistry.GetAssetsByClass(UBlueprintGeneratedClass::StaticClass()->GetClassPathName(), AllGeneratedBPsAssetData, true);
 	}
 
 	for (const FAssetData& BPAssetData : AllGeneratedBPsAssetData)
@@ -315,8 +336,8 @@ void FWidgetCatalogViewModel::BuildClassWidgetList()
 			continue;
 		}
 
-		// Was already parsed by the TObjectIterator<UClass>
-		if (BPAssetData.IsAssetLoaded())
+		// Was already processed above:
+		if (ProcessedPackages.Contains(BPAssetData.PackageName))
 		{
 			continue;
 		}
diff --git a/Engine/Source/Editor/UMGEditor/UMGEditor.Build.cs b/Engine/Source/Editor/UMGEditor/UMGEditor.Build.cs
--- a/Engine/Source/Editor/UMGEditor/UMGEditor.Build.cs
+++ b/Engine/Source/Editor/UMGEditor/UMGEditor.Build.cs
@@ -19,6 +19,7 @@ public class UMGEditor : ModuleRules
 		PrivateDependencyModuleNames.AddRange(
 			new string[] {
 				"AssetDefinition",
+				"AssetRegistry",
 				"ClassViewer",
 				"Core",
 				"CoreUObject",
diff --git a/Engine/Source/Runtime/AssetRegistry/Private/AssetRegistryHelpers.cpp b/Engine/Source/Runtime/AssetRegistry/Private/AssetRegistryHelpers.cpp
--- a/Engine/Source/Runtime/AssetRegistry/Private/AssetRegistryHelpers.cpp
+++ b/Engine/Source/Runtime/AssetRegistry/Private/AssetRegistryHelpers.cpp
@@ -273,6 +273,74 @@ void UAssetRegistryHelpers::GetBlueprintAssets(const FARFilter& InFilter, TArray
 	AssetRegistry.EnumerateAssets(Filter, FilterLambda, UE::AssetRegistry::EEnumerateAssetsFlags::None);
 }
 
+void UAssetRegistryHelpers::GetDerivedClassAssetData(const TArray<FTopLevelAssetPath>& InBaseClasses, TArray<FAssetData>& OutDerivedClassAssetData)
+{
+	IAssetRegistry& AR = IAssetRegistry::GetChecked();
+	TSet<FTopLevelAssetPath> DerivedClassNames;
+	// For a large working set with millions of assets and thousands of entries in the 
+	// hierarchy, GetDerivedClassNames has been measured to cost about 2.5ms - more than 
+	// acceptable for editor performance:
+	AR.GetDerivedClassNames(InBaseClasses, TSet<FTopLevelAssetPath>(), DerivedClassNames);
+
+	// For the same working set this loop takes 4.5ms - more than acceptable for editor 
+	// performance but it could be optimized. My problems in the callers are greater ATM.
+	OutDerivedClassAssetData.Reserve(DerivedClassNames.Num());
+	for (const FTopLevelAssetPath& Path : DerivedClassNames)
+	{
+		const FName PackageName = Path.GetPackageName();
+
+		// we have no asset data for built in types
+		if (PackageName.ToString().StartsWith(TEXT("/Script/")))
+		{
+			continue;
+		}
+		// SKEL classes should not even be in the asset registry.. but we can easily filter
+		// them here:
+		FString AssetName = Path.GetAssetName().ToString();
+		if (AssetName.StartsWith(TEXT("SKEL_")))
+		{
+			continue;
+		}
+
+		// look for blueprint asset first, as the _C (class) object won't 
+		// reliably report itself as editor only:
+		if (AssetName.EndsWith(TEXT("_C")))
+		{
+			AssetName.RemoveFromEnd(TEXT("_C"));
+		}
+
+		const bool bIncludeOnlyOnDiskAssets = false;
+		const bool bSkipARFilteredAssets = false;
+		FAssetData AssetData = AR.GetAssetByObjectPath(
+			FSoftObjectPath{FTopLevelAssetPath{PackageName, FName(*AssetName)}, FString()}, 
+			bIncludeOnlyOnDiskAssets, 
+			bSkipARFilteredAssets);
+		if (AssetData.PackageName.IsNone())
+		{
+			// We didn't find an asset, but don't give up - look for an asset with _C - 
+			// the origin of this ambiguity is blueprints, where we have  two asset 
+			// objects (the generated class and the UBlueprint):
+			AssetName += TEXT("_C");
+			AssetData = AR.GetAssetByObjectPath(
+				FSoftObjectPath{FTopLevelAssetPath{PackageName, FName(*AssetName)}, FString()}, 
+				bIncludeOnlyOnDiskAssets, 
+				bSkipARFilteredAssets);
+
+			ensureMsgf(!AssetData.PackageName.IsNone(), 
+				TEXT("Failed to find asset data for package: %s.%s"), 
+				*PackageName.ToString(), *AssetName);
+		}
+
+		// AssetData very likely is set, no known cases where we fail to find asset data
+		// but callers will appreciate if we guarantee that all returned asset data
+		// is valid:
+		if(!AssetData.PackageName.IsNone()) 
+		{
+			OutDerivedClassAssetData.Emplace(MoveTemp(AssetData));
+		}
+	}
+}
+
 bool UAssetRegistryHelpers::IsAssetDataBlueprintOfClassSet(const FAssetData& AssetData, const TSet<FTopLevelAssetPath>& ClassNameSet)
 {
 	const FString ParentClassFromData = AssetData.GetTagValueRef<FString>(FBlueprintTags::ParentClassPath);
diff --git a/Engine/Source/Runtime/AssetRegistry/Public/AssetRegistry/AssetRegistryHelpers.h b/Engine/Source/Runtime/AssetRegistry/Public/AssetRegistry/AssetRegistryHelpers.h
--- a/Engine/Source/Runtime/AssetRegistry/Public/AssetRegistry/AssetRegistryHelpers.h
+++ b/Engine/Source/Runtime/AssetRegistry/Public/AssetRegistry/AssetRegistryHelpers.h
@@ -119,6 +119,14 @@ public:
 	UFUNCTION(BlueprintPure, Category = "Asset Registry", meta=(ScriptMethod))
 	static ASSETREGISTRY_API void GetBlueprintAssets(const FARFilter& InFilter, TArray<FAssetData>& OutAssetData);
 
+	/** 
+	 * Gets a list of FAssetData for all child classes, loaded or unloaded, if you have a native UClass* convert it
+	 * to a FTopLevelAssetPath using UClass::GetClassPathName. The list of derived types can be quite large, but
+	 * this function attempts to be as efficient as possible for large asset sets.
+	 */
+	UFUNCTION(BlueprintCallable, Category = "Asset Registry")
+	static ASSETREGISTRY_API void GetDerivedClassAssetData(const TArray<FTopLevelAssetPath>& InBaseClasses, TArray<FAssetData>& OutDerivedClassAssetData);
+
 	/**
 	 * Returns the first native class of the asset type that can be found.  Normally this is just the FAssetData::GetClass(),
 	 * however if the class is a blueprint generated class it may not be loaded.  In which case GetAncestorClassNames will
