Fix to in UMG Designer to add more checks around the use of GetCachedWidget which might return a null object during drag/drop and move operation.

diff --git a/Engine/Source/Editor/UMGEditor/Private/Designer/SDesignerView.cpp b/Engine/Source/Editor/UMGEditor/Private/Designer/SDesignerView.cpp
--- a/Engine/Source/Editor/UMGEditor/Private/Designer/SDesignerView.cpp
+++ b/Engine/Source/Editor/UMGEditor/Private/Designer/SDesignerView.cpp
@@ -1334,8 +1334,7 @@ bool SDesignerView::GetWidgetGeometry(const FWidgetReference& Widget, FGeometry&
 
 bool SDesignerView::GetWidgetGeometry(const UWidget* InPreviewWidget, FGeometry& Geometry) const
 {
-	TSharedPtr<SWidget> CachedPreviewWidget = InPreviewWidget->GetCachedWidget();
-	if ( CachedPreviewWidget.IsValid() )
+	if (TSharedPtr<SWidget> CachedPreviewWidget = InPreviewWidget->GetCachedWidget())
 	{
 		const FArrangedWidget* ArrangedWidget = CachedWidgetGeometry.Find(CachedPreviewWidget.ToSharedRef());
 		if ( ArrangedWidget )
@@ -2449,8 +2448,7 @@ void SDesignerView::Tick( const FGeometry& AllottedGeometry, const double InCurr
 
 	if ( PreviewWidget )
 	{
-		TSharedPtr<SWidget> CachedWidget = PreviewWidget->GetCachedWidget();
-		if ( CachedWidget.IsValid() )
+		if (TSharedPtr<SWidget> CachedWidget = PreviewWidget->GetCachedWidget())
 		{
 			CachedPreviewDesiredSize = CachedWidget->GetDesiredSize();
 		}
@@ -2551,16 +2549,22 @@ FReply SDesignerView::OnDragDetected(const FGeometry& MyGeometry, const FPointer
 		PendingSelectedWidget = FWidgetReference();
 
 		for (const FWidgetReference& SelectedWidget : SelectedWidgets)
-		{
-			// Determine The offset to keep the widget from the mouse while dragging
-			FArrangedWidget ArrangedWidget(SNullWidget::NullWidget, FGeometry());
-			FDesignTimeUtils::GetArrangedWidget(SelectedWidget.GetPreview()->GetCachedWidget().ToSharedRef(), ArrangedWidget);
-			SelectedWidgetContextMenuLocation = ArrangedWidget.Geometry.AbsoluteToLocal(DraggingStartPositionScreenSpace);
-
-			FDragWidget DraggingWidget;
-			DraggingWidget.Widget = SelectedWidget;
-			DraggingWidget.DraggedOffset = SelectedWidgetContextMenuLocation / ArrangedWidget.Geometry.GetLocalSize();
-			DraggingWidgetCandidates.Add(DraggingWidget);
+		{		
+			if (UWidget* WidgetPreview = SelectedWidget.GetPreview())
+			{
+				if (TSharedPtr<SWidget> CachedWidget = WidgetPreview->GetCachedWidget())
+				{
+					// Determine The offset to keep the widget from the mouse while dragging
+					FArrangedWidget ArrangedWidget(SNullWidget::NullWidget, FGeometry());
+					FDesignTimeUtils::GetArrangedWidget(CachedWidget.ToSharedRef(), ArrangedWidget);
+					SelectedWidgetContextMenuLocation = ArrangedWidget.Geometry.AbsoluteToLocal(DraggingStartPositionScreenSpace);
+
+					FDragWidget DraggingWidget;
+					DraggingWidget.Widget = SelectedWidget;
+					DraggingWidget.DraggedOffset = SelectedWidgetContextMenuLocation / ArrangedWidget.Geometry.GetLocalSize();
+					DraggingWidgetCandidates.Add(DraggingWidget);
+				}
+			}
 		}
 
 		TArray<FDragWidget> DraggingWidgets;
@@ -3109,41 +3113,39 @@ void SDesignerView::MoveWidgets(const FGeometry& MyGeometry, const FDragDropEven
 						{
 							if (bIsPreview)
 							{
-								CanvasSlot->SaveBaseLayout();
-
-								FArrangedWidget ArrangedWidget(SNullWidget::NullWidget, FGeometry());
-								FDesignTimeUtils::GetArrangedWidget(Widget->GetCachedWidget().ToSharedRef(), ArrangedWidget);
-
-								FVector2D Offset = FVector2D::ZeroVector;
 								if (TSharedPtr<SWidget> CachedWidget = Widget->GetCachedWidget())
 								{
-									FDesignTimeUtils::GetArrangedWidget(CachedWidget.ToSharedRef(), ArrangedWidget);
-									Offset = DraggedWidget.DraggedOffset * ArrangedWidget.Geometry.GetLocalSize();
-								}
+									CanvasSlot->SaveBaseLayout();
 
-								FVector2D NewPosition;
-								if (bAnyWidgetChangingParent || bIsChangingParent)
-								{
-									NewPosition = LocalPosition - Offset;
-								}
-								else
-								{
-									NewPosition = LocalPosition - DragDropEvent.GetCursorDelta() - Offset;
-
-								}
-								// Perform grid snapping on X and Y if we need to.
-								if (bGridSnapX)
-								{
-									NewPosition.X = ((int32)NewPosition.X) - (((int32)NewPosition.X) % DesignerSettings->GridSnapSize);
-								}
-								if (bGridSnapY)
-								{
-									NewPosition.Y = ((int32)NewPosition.Y) - (((int32)NewPosition.Y) % DesignerSettings->GridSnapSize);
+									FArrangedWidget ArrangedWidget(SNullWidget::NullWidget, FGeometry());
+																		
+									FDesignTimeUtils::GetArrangedWidget(CachedWidget.ToSharedRef(), ArrangedWidget);
+									FVector2D Offset = DraggedWidget.DraggedOffset * ArrangedWidget.Geometry.GetLocalSize();
+									
+									FVector2D NewPosition;
+									if (bAnyWidgetChangingParent || bIsChangingParent)
+									{
+										NewPosition = LocalPosition - Offset;
+									}
+									else
+									{
+										NewPosition = LocalPosition - DragDropEvent.GetCursorDelta() - Offset;
+
+									}
+									// Perform grid snapping on X and Y if we need to.
+									if (bGridSnapX)
+									{
+										NewPosition.X = ((int32)NewPosition.X) - (((int32)NewPosition.X) % DesignerSettings->GridSnapSize);
+									}
+									if (bGridSnapY)
+									{
+										NewPosition.Y = ((int32)NewPosition.Y) - (((int32)NewPosition.Y) % DesignerSettings->GridSnapSize);
+									}
+									CanvasSlot->SetDesiredPosition(NewPosition);
+
+									CanvasSlot->RebaseLayout();
+									HasChangedLayout = true;
 								}
-								CanvasSlot->SetDesiredPosition(NewPosition);
-
-								CanvasSlot->RebaseLayout();
-								HasChangedLayout = true;
 							}
 						}
 						else
diff --git a/Engine/Source/Runtime/UMG/Private/Components/Widget.cpp b/Engine/Source/Runtime/UMG/Private/Components/Widget.cpp
--- a/Engine/Source/Runtime/UMG/Private/Components/Widget.cpp
+++ b/Engine/Source/Runtime/UMG/Private/Components/Widget.cpp
@@ -1026,7 +1026,7 @@ TSharedRef<SWidget> UWidget::TakeWidget_Private(ConstructMethodType ConstructMet
 			PublicWidget = SafeGCWidget;
 		}
 	}
-	if(bWrappedByComponent)
+	if (bWrappedByComponent)
 	{
 		// Now, if the Widget instance is wrapped by a component, add it to the hierarchy
 		const UWidgetTree* WidgetTree = Cast<UWidgetTree>(GetOuter());
@@ -1038,7 +1038,7 @@ TSharedRef<SWidget> UWidget::TakeWidget_Private(ConstructMethodType ConstructMet
 				{
 					if (bNewlyCreated)
 					{
-						TSharedPtr<SWidget> SafeComponentWidget = Component->RebuildWidgetWithContent(PublicWidget.ToSharedRef());
+						TSharedRef<SWidget> SafeComponentWidget = Component->RebuildWidgetWithContent(PublicWidget.ToSharedRef());
 						// If the Widget is the same of null, we do nothing 
 						if (SafeComponentWidget != PublicWidget)
 						{
@@ -1060,7 +1060,7 @@ TSharedRef<SWidget> UWidget::TakeWidget_Private(ConstructMethodType ConstructMet
 	{
 		if (bNewlyCreated)
 		{
-			TSharedPtr<SWidget> SafeDesignWidget = RebuildDesignWidget(PublicWidget.ToSharedRef());
+			TSharedRef<SWidget> SafeDesignWidget = RebuildDesignWidget(PublicWidget.ToSharedRef());
 			if (SafeDesignWidget != PublicWidget)
 			{
 				DesignWrapperWidget = SafeDesignWidget;
