UIComponents: Navigation Component.
Listens to navigation transition events within widget hierarchy

#rnx

diff --git a/Engine/Source/Editor/UMGEditor/Private/Customizations/UIComponentCustomizationExtender.cpp b/Engine/Source/Editor/UMGEditor/Private/Customizations/UIComponentCustomizationExtender.cpp
--- a/Engine/Source/Editor/UMGEditor/Private/Customizations/UIComponentCustomizationExtender.cpp
+++ b/Engine/Source/Editor/UMGEditor/Private/Customizations/UIComponentCustomizationExtender.cpp
@@ -7,6 +7,7 @@
 #include "DetailWidgetRow.h"
 #include "Extensions/UIComponent.h"
 #include "Extensions/UIComponentUserWidgetExtension.h"
+#include "Extensions/UIComponents/NavigationUIComponentCustomization.h"
 #include "IDetailsView.h"
 #include "PropertyCustomizationHelpers.h"
 #include "PropertyHandle.h"
@@ -39,6 +40,8 @@ void FUIComponentCustomizationExtender::CustomizeDetails(IDetailLayoutBuilder& I
 	// We use the UserWidgetExtension on the preview and it will be migrated to the WBP extension in MigrateFromChain
 	if (const UUIComponentUserWidgetExtension* Extension = PreviewUserWidget->GetExtension<UUIComponentUserWidgetExtension>())
 	{
+		CustomizeComponentPropertyTypes(InDetailLayout.GetDetailsViewSharedPtr());
+
 		TArray<UUIComponent*> ComponentsOnWidget = Extension->GetComponentsFor(Widget.Get());
 		for (int32 Index = ComponentsOnWidget.Num() - 1; Index >= 0; Index--)
 		{
@@ -88,7 +91,7 @@ void FUIComponentCustomizationExtender::CustomizeDetails(IDetailLayoutBuilder& I
 					[
 						RemoveButton
 					];				
-				ComponentCategory.HeaderContent(ComponentHeaderContent, true);				
+				ComponentCategory.HeaderContent(ComponentHeaderContent, true);
 			}
 		}
 	}
@@ -110,4 +113,15 @@ void FUIComponentCustomizationExtender::CustomizeDetails(IDetailLayoutBuilder& I
 		}
 	}
 }
+
+void FUIComponentCustomizationExtender::CustomizeComponentPropertyTypes(TSharedPtr<IDetailsView> InDetailsView)
+{
+	if (!InDetailsView.IsValid())
+	{
+		return;
+	}
+
+	InDetailsView->RegisterInstancedCustomPropertyTypeLayout(TEXT("NavigationUIComponent"), FOnGetPropertyTypeCustomizationInstance::CreateStatic(&FNavigationUIComponentCustomization::MakeInstance, WidgetBlueprintEditor, Widget->GetFName()));
+}
+
 #undef LOCTEXT_NAMESPACE
\ No newline at end of file
diff --git a/Engine/Source/Editor/UMGEditor/Private/Customizations/UIComponentCustomizationExtender.h b/Engine/Source/Editor/UMGEditor/Private/Customizations/UIComponentCustomizationExtender.h
--- a/Engine/Source/Editor/UMGEditor/Private/Customizations/UIComponentCustomizationExtender.h
+++ b/Engine/Source/Editor/UMGEditor/Private/Customizations/UIComponentCustomizationExtender.h
@@ -16,6 +16,9 @@ public:
 	virtual void CustomizeDetails(IDetailLayoutBuilder& InDetailLayout, const TArrayView<UWidget*> InWidgets, const TSharedRef<FWidgetBlueprintEditor>& InWidgetBlueprintEditor) override;
 
 private:
+
+	void CustomizeComponentPropertyTypes(TSharedPtr<IDetailsView> InDetailView);
+
 	/** The selected widget in the details panel. */
 	TWeakObjectPtr<class UWidget> Widget;
 	TWeakPtr<FWidgetBlueprintEditor> WidgetBlueprintEditor;
diff --git a/Engine/Source/Editor/UMGEditor/Private/Extensions/UIComponents/NavigationUIComponentCustomization.cpp b/Engine/Source/Editor/UMGEditor/Private/Extensions/UIComponents/NavigationUIComponentCustomization.cpp
new file mode 100644
--- /dev/null
+++ b/Engine/Source/Editor/UMGEditor/Private/Extensions/UIComponents/NavigationUIComponentCustomization.cpp
@@ -0,0 +1,188 @@
+// Copyright Epic Games, Inc. All Rights Reserved.
+
+#include "Extensions/UIComponents/NavigationUIComponentCustomization.h"
+
+#include "UIComponentWidgetBlueprintExtension.h"
+#include "Extensions/UIComponentUserWidgetExtension.h"
+#include "Extensions/UIComponents/NavigationUIComponent.h"
+
+#include "DetailLayoutBuilder.h"
+#include "DetailWidgetRow.h"
+#include "Details/SFunctionSelector.h"
+#include "IDetailChildrenBuilder.h"
+#include "IDetailPropertyRow.h"
+#include "IPropertyUtilities.h"
+#include "Kismet2/BlueprintEditorUtils.h"
+#include "ScopedTransaction.h"
+#include "WidgetBlueprint.h"
+#include "Widgets/Text/STextBlock.h"
+
+#define LOCTEXT_NAMESPACE "FNavigationCustomizationUIComponent"
+
+namespace UE::NavigationUIComponent::Private
+{
+struct FScopedNavigationUIComponentTransaction
+{
+	TSharedPtr<FWidgetBlueprintEditor> Editor;
+	UNavigationUIComponent* NavigationUIComponent;
+	FScopedTransaction Transaction;
+	bool bModifiedProperty;
+
+	FScopedNavigationUIComponentTransaction(TWeakPtr<FWidgetBlueprintEditor> InEditor, UNavigationUIComponent* InNavigationUIComponent, const FText& InTransactionText)
+		: Editor(InEditor.Pin())
+		, NavigationUIComponent(InNavigationUIComponent)
+		, Transaction(InTransactionText)
+		, bModifiedProperty(false)
+	{
+		if (NavigationUIComponent)
+		{
+			NavigationUIComponent->Modify();
+		}
+	}
+
+	~FScopedNavigationUIComponentTransaction()
+	{
+		if (!bModifiedProperty)
+		{
+			Transaction.Cancel();
+		}
+	}
+
+	void MarkPropertyModified(const FName& PropertyName)
+	{
+		if (NavigationUIComponent == nullptr)
+		{
+			return;
+		}
+		
+		if (FProperty* ChangedProp = FindFProperty<FProperty>(UNavigationUIComponent::StaticClass(), PropertyName))
+		{
+			FPropertyChangedEvent PropChangedEvent(ChangedProp, EPropertyChangeType::ValueSet, MakeConstArrayView({Cast<UObject>(NavigationUIComponent)}));
+			FBlueprintEditorUtils::MarkBlueprintAsModified(Editor->GetWidgetBlueprintObj(), PropChangedEvent);
+
+			bModifiedProperty = true;
+		}
+	}
+};
+}
+
+void FNavigationUIComponentCustomization::CustomizeHeader(TSharedRef<IPropertyHandle> PropertyHandle, class FDetailWidgetRow& HeaderRow, IPropertyTypeCustomizationUtils& CustomizationUtils)
+{
+}
+
+void FNavigationUIComponentCustomization::CustomizeChildren(TSharedRef<IPropertyHandle> PropertyHandle, class IDetailChildrenBuilder& ChildBuilder, IPropertyTypeCustomizationUtils& CustomizationUtils)
+{
+	CustomizeNavigationTransitionFunction(ChildBuilder, true);
+	CustomizeNavigationTransitionFunction(ChildBuilder, false);
+}
+
+UNavigationUIComponent* FNavigationUIComponentCustomization::GetNavigationComponent() const
+{
+	TSharedPtr<FWidgetBlueprintEditor> EditorPinned = Editor.Pin();
+	if (!EditorPinned.IsValid())
+	{
+		return nullptr;
+	}
+
+	UWidgetBlueprint* WidgetBlueprint = EditorPinned->GetWidgetBlueprintObj();
+	if (WidgetBlueprint == nullptr)
+	{
+		return nullptr;
+	}
+
+	UUIComponentWidgetBlueprintExtension* Extension = UWidgetBlueprintExtension::GetExtension<UUIComponentWidgetBlueprintExtension>(WidgetBlueprint);
+	if (Extension == nullptr)
+	{
+		return nullptr;
+	}
+
+	if (UUIComponent* Component = Extension->GetComponent(UNavigationUIComponent::StaticClass(), WidgetName))
+	{
+		if (Component->IsA<UNavigationUIComponent>())
+		{
+			return Cast<UNavigationUIComponent>(Component);
+		}
+	}
+	return nullptr;
+}
+
+void FNavigationUIComponentCustomization::CustomizeNavigationTransitionFunction(class IDetailChildrenBuilder& ChildBuilder, bool bIsEnteredFunction)
+{
+	static UFunction* OnNavigationTransitionSignature = FindObject<UFunction>(FindPackage(nullptr, TEXT("/Script/UMG")), TEXT("OnNavigationTransition__DelegateSignature"));
+
+	const FText Name = bIsEnteredFunction ? LOCTEXT("OnNavigationEntered", "On Navigation Entered") : LOCTEXT("OnNavigationExited", "On Navigation Exited");
+
+	ChildBuilder.AddCustomRow(Name)
+		.NameContent()
+		[
+			SNew(STextBlock)
+			.Font(IDetailLayoutBuilder::GetDetailFont())
+			.Text(Name)
+		]
+		.ValueContent()
+		.MaxDesiredWidth(300)
+		[
+			SNew(SHorizontalBox)
+
+			+ SHorizontalBox::Slot()
+			.FillWidth(1.0f)
+			[
+				SNew(SFunctionSelector, Editor.Pin().ToSharedRef(), OnNavigationTransitionSignature)
+				.CurrentFunction(this, &FNavigationUIComponentCustomization::GetNavigationTransitionFunction, bIsEnteredFunction)
+				.OnSelectedFunction(this, &FNavigationUIComponentCustomization::HandleSelectedNavigationTransitionFunction, bIsEnteredFunction)
+				.OnResetFunction(this, &FNavigationUIComponentCustomization::HandleResetNavigationTransitionFunction, bIsEnteredFunction)
+			]
+		];
+
+}
+
+TOptional<FName> FNavigationUIComponentCustomization::GetNavigationTransitionFunction(bool bIsEnteredFunction) const
+{
+	if ( const UNavigationUIComponent* NavigationUIComponent = GetNavigationComponent())
+	{
+		if (bIsEnteredFunction)
+		{
+			return  NavigationUIComponent->OnNavigationEntered;
+		}
+		else
+		{
+			return NavigationUIComponent->OnNavigationExited;
+		}
+	}
+	return TOptional<FName>();
+}
+
+void FNavigationUIComponentCustomization::HandleSelectedNavigationTransitionFunction(FName SelectedFunction, bool bIsEnteredFunction)
+{
+	UNavigationUIComponent* NavigationUIComponent = GetNavigationComponent();
+	if (NavigationUIComponent == nullptr)
+	{
+		return;
+	}
+
+	UE::NavigationUIComponent::Private::FScopedNavigationUIComponentTransaction NavigationUIComponentTransaction(Editor, NavigationUIComponent, LOCTEXT("SetNavigationTransition", "Set  Navigation Transition"));
+
+	if (bIsEnteredFunction)
+	{
+		if (NavigationUIComponent->OnNavigationEntered != SelectedFunction)
+		{
+			NavigationUIComponent->OnNavigationEntered = SelectedFunction;
+			NavigationUIComponentTransaction.MarkPropertyModified(GET_MEMBER_NAME_CHECKED(UNavigationUIComponent, OnNavigationEntered));
+		}
+	}
+	else
+	{
+		if (NavigationUIComponent->OnNavigationExited != SelectedFunction)
+		{
+			NavigationUIComponent->OnNavigationExited = SelectedFunction;
+			NavigationUIComponentTransaction.MarkPropertyModified(GET_MEMBER_NAME_CHECKED(UNavigationUIComponent, OnNavigationExited));
+		}
+	}
+}
+
+void FNavigationUIComponentCustomization::HandleResetNavigationTransitionFunction(bool bIsEnteredFunction)
+{
+	HandleSelectedNavigationTransitionFunction(NAME_None, bIsEnteredFunction);
+}
+
+#undef LOCTEXT_NAMESPACE
\ No newline at end of file
diff --git a/Engine/Source/Editor/UMGEditor/Private/Extensions/UIComponents/NavigationUIComponentCustomization.h b/Engine/Source/Editor/UMGEditor/Private/Extensions/UIComponents/NavigationUIComponentCustomization.h
new file mode 100644
--- /dev/null
+++ b/Engine/Source/Editor/UMGEditor/Private/Extensions/UIComponents/NavigationUIComponentCustomization.h
@@ -0,0 +1,41 @@
+// Copyright Epic Games, Inc. All Rights Reserved.
+
+#include "Layout/Visibility.h"
+#include "WidgetBlueprintEditor.h"
+#include "IPropertyTypeCustomization.h"
+#include "StructUtils/InstancedStruct.h"
+#include "PropertyHandle.h"
+
+class IDetailChildrenBuilder;
+class IPropertyUtilities;
+class UNavigationUIComponent;
+
+class FNavigationUIComponentCustomization : public IPropertyTypeCustomization
+{
+public:
+	static TSharedRef<class IPropertyTypeCustomization> MakeInstance(TWeakPtr<FWidgetBlueprintEditor> InEditor, FName InWidgetName)
+	{
+		return MakeShareable(new FNavigationUIComponentCustomization(InEditor, InWidgetName));
+	}
+
+	FNavigationUIComponentCustomization(TWeakPtr<FWidgetBlueprintEditor> InEditor, FName InWidgetName)
+		: Editor(InEditor)
+		, WidgetName(InWidgetName)
+	{}
+	
+	virtual void CustomizeHeader(TSharedRef<IPropertyHandle> PropertyHandle, class FDetailWidgetRow& HeaderRow, IPropertyTypeCustomizationUtils& CustomizationUtils) override;
+	virtual void CustomizeChildren(TSharedRef<IPropertyHandle> PropertyHandle, class IDetailChildrenBuilder& ChildBuilder, IPropertyTypeCustomizationUtils& CustomizationUtils) override;
+
+private:
+
+	UNavigationUIComponent* GetNavigationComponent() const;
+
+	void CustomizeNavigationTransitionFunction(class IDetailChildrenBuilder& ChildBuilder, bool bIsEnteredFunction);
+
+	TOptional<FName> GetNavigationTransitionFunction(bool bIsEnteredFunction) const;
+	void HandleSelectedNavigationTransitionFunction(FName SelectedFunction, bool bIsEnteredFunction);
+	void HandleResetNavigationTransitionFunction(bool bIsEnteredFunction);
+	
+	TWeakPtr<FWidgetBlueprintEditor> Editor;
+	FName WidgetName;
+};
diff --git a/Engine/Source/Runtime/Slate/Private/Framework/Application/SlateApplication.cpp b/Engine/Source/Runtime/Slate/Private/Framework/Application/SlateApplication.cpp
--- a/Engine/Source/Runtime/Slate/Private/Framework/Application/SlateApplication.cpp
+++ b/Engine/Source/Runtime/Slate/Private/Framework/Application/SlateApplication.cpp
@@ -12,6 +12,7 @@
 #include "Misc/App.h"
 #include "Modules/ModuleManager.h"
 #include "InputCoreModule.h"
+#include "Input/NavigationMetadata.h"
 #include "Layout/LayoutUtils.h"
 #include "Sound/ISlateSoundDevice.h"
 #include "Sound/NullSlateSoundDevice.h"
@@ -522,6 +523,61 @@ struct FEventRouter::Translate<FPointerEvent>
 	}
 };
 
+//////////////////////////////////////////////////////////////////////////
+
+struct FScopedNavigationTransition
+{
+	FScopedNavigationTransition(EUINavigation NavigationType, TSharedPtr<FSlateUser> User)
+		: NavigationType(NavigationType)
+		, User(User)
+	{
+		OldFocusedWidget = GetFocusedWidget();
+	}
+
+	~FScopedNavigationTransition()
+	{
+		NewFocusedWidget = GetFocusedWidget();
+
+		ExecuteNavigationTransition(OldFocusedWidget, ENavigationTransitionDirection::Outgoing);
+		ExecuteNavigationTransition(NewFocusedWidget, ENavigationTransitionDirection::Incoming);
+	}
+
+	TWeakPtr<SWidget> GetFocusedWidget()
+	{
+		if (User.IsValid() && User->GetWeakFocusPath().IsValid())
+		{
+			return User->GetWeakFocusPath().GetLastWidget().Pin();
+		}
+		return nullptr;
+	}
+
+	void ExecuteNavigationTransition(TWeakPtr<SWidget> Widget, const ENavigationTransitionDirection& Direction)
+	{
+		// No transition occurred.
+		if (OldFocusedWidget == NewFocusedWidget)
+		{
+			return;
+		}
+		// Notify widget hierarchy of navigation event
+		while(Widget.IsValid())
+		{
+			TSharedPtr<SWidget> PinnedWidget = Widget.Pin();
+			if (TSharedPtr<FNavigationTransitionMetadata> Metadata = PinnedWidget->GetMetaData<FNavigationTransitionMetadata>())
+			{
+				const FNavigationTransition Transition(NavigationType, OldFocusedWidget, NewFocusedWidget, Direction);
+				Metadata->OnNavigationTransition.ExecuteIfBound(Transition);
+			}
+			Widget = PinnedWidget->GetParentWidget();
+		}
+	}
+
+	EUINavigation NavigationType;
+	TSharedPtr<FSlateUser> User;
+	TWeakPtr<SWidget> OldFocusedWidget; 
+	TWeakPtr<SWidget> NewFocusedWidget; 
+};
+
+
 DECLARE_CYCLE_STAT( TEXT("Message Tick Time"), STAT_SlateMessageTick, STATGROUP_Slate );
 DECLARE_CYCLE_STAT( TEXT("Slate App Input"), STAT_SlateApplicationInput, STATGROUP_Slate );
 DECLARE_CYCLE_STAT( TEXT("Total Slate Tick Time"), STAT_SlateTickTime, STATGROUP_Slate );
@@ -3559,6 +3615,8 @@ void FSlateApplication::ProcessReply( const FWidgetPath& CurrentEventPath, const
 			}
 			else if (TheReply.GetNavigationDestination().IsValid())
 			{
+				FScopedNavigationTransition ScopedNavigationTransition(TheReply.GetNavigationType(), SlateUser);
+
 				const bool bAlwaysHandleNavigationAttempt = false;
 				ExecuteNavigation(NavigationSource, TheReply.GetNavigationDestination(), UserIndex, bAlwaysHandleNavigationAttempt);
 			}
@@ -6221,7 +6279,7 @@ bool FSlateApplication::OnCursorSet()
 	return true;
 }
 
-void FSlateApplication::NavigateToWidget(const uint32 UserIndex, const TSharedPtr<SWidget>& NavigationDestination, ENavigationSource NavigationSource)
+void FSlateApplication::NavigateToWidget(const uint32 UserIndex, const TSharedPtr<SWidget>& NavigationDestination, ENavigationSource NavigationSource, EUINavigation NavigationType)
 {
 	if (NavigationDestination.IsValid())
 	{
@@ -6237,6 +6295,8 @@ void FSlateApplication::NavigateToWidget(const uint32 UserIndex, const TSharedPt
 
 		if (NavigationSourceWP.IsValid())
 		{
+			FScopedNavigationTransition ScopedNavigationTransition(NavigationType, GetUser(UserIndex));
+
 			bool bAlwaysHandleNavigationAttempt = false;
 			ExecuteNavigation(NavigationSourceWP, NavigationDestination, UserIndex, bAlwaysHandleNavigationAttempt);
 		}
@@ -6389,6 +6449,8 @@ bool FSlateApplication::AttemptNavigation(const FWidgetPath& NavigationSource, c
 	FSlateDebugging::BroadcastAttemptNavigation(NavigationEvent, NavigationReply, NavigationSource, DestinationWidget, NavigationMethod);
 #endif
 
+	FScopedNavigationTransition ScopedNavigationTransition(NavigationType, GetUser(NavigationEvent.GetUserIndex()));
+
 	return ExecuteNavigation(NavigationSource, DestinationWidget, NavigationEvent.GetUserIndex(), bAlwaysHandleNavigationAttempt);
 }
 
diff --git a/Engine/Source/Runtime/Slate/Private/Widgets/Views/STableViewBase.cpp b/Engine/Source/Runtime/Slate/Private/Widgets/Views/STableViewBase.cpp
--- a/Engine/Source/Runtime/Slate/Private/Widgets/Views/STableViewBase.cpp
+++ b/Engine/Source/Runtime/Slate/Private/Widgets/Views/STableViewBase.cpp
@@ -1255,9 +1255,9 @@ float STableViewBase::GetFirstLineScrollOffset() const
 	return FirstLineScrollOffset - (int64)FirstLineScrollOffset;
 }
 
-void STableViewBase::NavigateToWidget(const uint32 UserIndex, const TSharedPtr<SWidget>& NavigationDestination, ENavigationSource NavigationSource) const
+void STableViewBase::NavigateToWidget(const uint32 UserIndex, const TSharedPtr<SWidget>& NavigationDestination, ENavigationSource NavigationSource, EUINavigation NavigationType) const
 {
-	FSlateApplication::Get().NavigateToWidget(UserIndex, NavigationDestination, NavigationSource);
+	FSlateApplication::Get().NavigateToWidget(UserIndex, NavigationDestination, NavigationSource, NavigationType);
 }
 
 int32 STableViewBase::FindChildUnderPosition(FArrangedChildren& ArrangedChildren, const FVector2D& ArrangedSpacePosition) const
diff --git a/Engine/Source/Runtime/Slate/Public/Framework/Application/SlateApplication.h b/Engine/Source/Runtime/Slate/Public/Framework/Application/SlateApplication.h
--- a/Engine/Source/Runtime/Slate/Public/Framework/Application/SlateApplication.h
+++ b/Engine/Source/Runtime/Slate/Public/Framework/Application/SlateApplication.h
@@ -1813,8 +1813,9 @@ public:
 	 * @param InUserIndex The user that is doing the navigation
 	 * @param NavigationDestination The navigation destination widget
 	 * @param NavigationSource The source type of the navigation
+	 * @param NavigationType The type of the navigation transition.
 	 */
-	SLATE_API void NavigateToWidget(const uint32 UserIndex, const TSharedPtr<SWidget>& NavigationDestination, ENavigationSource NavigationSource = ENavigationSource::FocusedWidget);
+	SLATE_API void NavigateToWidget(const uint32 UserIndex, const TSharedPtr<SWidget>& NavigationDestination, ENavigationSource NavigationSource = ENavigationSource::FocusedWidget, EUINavigation NavigationType = EUINavigation::Invalid);
 
 	/**
 	 * Attempts to navigate directly to the widget currently under the given user's cursor
diff --git a/Engine/Source/Runtime/Slate/Public/Widgets/Views/SListView.h b/Engine/Source/Runtime/Slate/Public/Widgets/Views/SListView.h
--- a/Engine/Source/Runtime/Slate/Public/Widgets/Views/SListView.h
+++ b/Engine/Source/Runtime/Slate/Public/Widgets/Views/SListView.h
@@ -511,6 +511,8 @@ public:
 
 	virtual FNavigationReply OnNavigation(const FGeometry& MyGeometry, const FNavigationEvent& InNavigationEvent) override
 	{
+		NavigationTypeOnScrollIntoView = EUINavigation::Invalid;
+
 		if (this->HasValidItemsSource() && this->bHandleDirectionalNavigation && (this->bHandleGamepadEvents || InNavigationEvent.GetNavigationGenesis() != ENavigationGenesis::Controller))
 		{
 			const TArrayView<const ItemType> ItemsSourceRef = this->GetItems();
@@ -2309,7 +2311,8 @@ protected:
 				if (bNavigateOnScrollIntoView)
 				{
 					SelectorItem = NonNullItemToNotifyWhenInView;
-					NavigateToWidget(UserRequestingScrollIntoView, Widget->AsWidget());
+					NavigateToWidget(UserRequestingScrollIntoView, Widget->AsWidget(), ENavigationSource::FocusedWidget, NavigationTypeOnScrollIntoView);
+					NavigationTypeOnScrollIntoView = EUINavigation::Invalid;
 				}
 			
 				OnItemScrolledIntoView.ExecuteIfBound(NonNullItemToNotifyWhenInView, Widget);
@@ -2610,6 +2613,18 @@ protected:
 		}
 	}
 
+	/**
+	* Special handling for navigation events
+	*/
+	virtual void NavigationSelect(const ItemType& InItemToSelect, const FNavigationEvent& InNavigationEvent)
+	{
+		NavigationSelect(InItemToSelect, StaticCast<FInputEvent>(InNavigationEvent));
+
+		// If the navigation requested a scroll into view, then set navigation type
+		NavigationTypeOnScrollIntoView = bNavigateOnScrollIntoView ? InNavigationEvent.GetNavigationType() : EUINavigation::Invalid;
+	}
+
+
 protected:
 	/** A widget generator component */
 	FWidgetGenerator WidgetGenerator;
@@ -2714,6 +2729,9 @@ protected:
 	/** If true, the item currently slated to be scrolled into view will also be navigated to after being scrolled in */
 	bool bNavigateOnScrollIntoView = false;
 
+	/** Cached navigation direction to forward to navigate widget call */
+	EUINavigation NavigationTypeOnScrollIntoView = EUINavigation::Invalid;
+
 	/** Style resource for the list */
 	const FTableViewStyle* Style;
 
diff --git a/Engine/Source/Runtime/Slate/Public/Widgets/Views/STableViewBase.h b/Engine/Source/Runtime/Slate/Public/Widgets/Views/STableViewBase.h
--- a/Engine/Source/Runtime/Slate/Public/Widgets/Views/STableViewBase.h
+++ b/Engine/Source/Runtime/Slate/Public/Widgets/Views/STableViewBase.h
@@ -428,7 +428,7 @@ protected:
 	virtual void NotifyFinishedScrolling() = 0;
 
 	/** Util Function so templates classes don't need to include SlateApplication */
-	SLATE_API void NavigateToWidget(const uint32 UserIndex, const TSharedPtr<SWidget>& NavigationDestination, ENavigationSource NavigationSource = ENavigationSource::FocusedWidget) const;
+	SLATE_API void NavigateToWidget(const uint32 UserIndex, const TSharedPtr<SWidget>& NavigationDestination, ENavigationSource NavigationSource = ENavigationSource::FocusedWidget, EUINavigation NavigationType = EUINavigation::Invalid) const;
 
 	/** Util function to find the child index under the given position. */
 	SLATE_API int32 FindChildUnderPosition(FArrangedChildren& ArrangedChildren, const FVector2D& ArrangedSpacePosition) const;
diff --git a/Engine/Source/Runtime/SlateCore/Public/Input/NavigationMetadata.h b/Engine/Source/Runtime/SlateCore/Public/Input/NavigationMetadata.h
new file mode 100644
--- /dev/null
+++ b/Engine/Source/Runtime/SlateCore/Public/Input/NavigationMetadata.h
@@ -0,0 +1,46 @@
+// Copyright Epic Games, Inc. All Rights Reserved.
+
+#pragma once
+
+#include "CoreMinimal.h"
+#include "Types/ISlateMetaData.h"
+#include "Input/NavigationReply.h"
+
+/**
+* FNavigationTransition is a description of a navigation change between two widgets. 
+* When a user attempts to navigate from a widget, 2 Transitions are reported. 
+* The Outgoing Transition is executed on the OldFocusedWidget.
+* The Incoming Transition is executed on the NewFOcusedWidget.
+*/
+
+enum class ENavigationTransitionDirection
+{
+	Incoming, Outgoing
+};
+
+struct FNavigationTransition
+{
+	FNavigationTransition(EUINavigation Type, TWeakPtr<SWidget> OldFocusedWidget, TWeakPtr<SWidget> NewFocusedWidget, ENavigationTransitionDirection Direction)
+		: Type(Type)
+		, OldFocusedWidget(OldFocusedWidget)
+		, NewFocusedWidget(NewFocusedWidget)
+		, Direction(Direction)
+	{}
+
+	EUINavigation Type;
+	TWeakPtr<SWidget> OldFocusedWidget;
+	TWeakPtr<SWidget> NewFocusedWidget;
+	ENavigationTransitionDirection Direction;
+};
+
+DECLARE_DELEGATE_OneParam(FOnNavigationTransitionDelegate, const FNavigationTransition&);
+
+/**
+ * FNavigationTransitionMetadatais an optional data field on SWidget instances that provides users with additional information about navigation transitions.
+ */
+struct FNavigationTransitionMetadata : ISlateMetaData
+{
+	SLATE_METADATA_TYPE(FNavigationTransitionMetadata, ISlateMetaData);
+
+	FOnNavigationTransitionDelegate OnNavigationTransition;
+};
\ No newline at end of file
diff --git a/Engine/Source/Runtime/UMG/Private/Extensions/UIComponents/NavigationUIComponent.cpp b/Engine/Source/Runtime/UMG/Private/Extensions/UIComponents/NavigationUIComponent.cpp
new file mode 100644
--- /dev/null
+++ b/Engine/Source/Runtime/UMG/Private/Extensions/UIComponents/NavigationUIComponent.cpp
@@ -0,0 +1,113 @@
+// Copyright Epic Games, Inc. All Rights Reserved.
+
+#include "Extensions/UIComponents/NavigationUIComponent.h"
+
+#include "Blueprint/WidgetTree.h"
+#include "CoreGlobals.h"
+#include "Components/Widget.h"
+#include "Slate/SObjectWidget.h"
+#include "Types/ReflectionMetadata.h"
+
+namespace UE::NavigationUIComponent::Private
+{
+	UUserWidget* FindUserWidget(TWeakPtr<SWidget> Widget)
+	{
+		// Find the outer user widget that contains this widget. 
+		while(Widget.IsValid())
+		{
+			TSharedPtr<SWidget> PinnedWidget = Widget.Pin();
+			if (PinnedWidget->GetWidgetClass().GetWidgetType() == SObjectWidget::StaticWidgetClass().GetWidgetType())
+			{
+				TSharedPtr<SObjectWidget> ObjectWidget = StaticCastSharedPtr<SObjectWidget>(PinnedWidget);
+				return ObjectWidget->GetWidgetObject();
+			}
+		
+			Widget = PinnedWidget->GetParentWidget();
+		}
+		return nullptr;
+	}
+
+	TSharedPtr<SWidget> GetNavigationWidget(TSharedPtr<SWidget> Widget)
+	{
+		// Special condition, get internal SWidget that will be focused/navigated. 
+		// Otherwise, the outer SObject Widget will report that it is focusable, although the native widget is being focus
+		if (Widget->GetType() == FName(TEXT("SObjectWidget")))
+		{
+			TSharedPtr<SObjectWidget> ObjectWidget = StaticCastSharedPtr<SObjectWidget>(Widget);
+			if (UUserWidget* UserWidget = ObjectWidget->GetWidgetObject())
+			{
+				Widget = UserWidget->GetCachedWidget().ToSharedRef();
+
+				UWidgetTree* WidgetTree = UserWidget != nullptr ? UserWidget->WidgetTree : nullptr;
+				if(TObjectPtr<UWidget> RootWidget = WidgetTree != nullptr ? WidgetTree->RootWidget : nullptr)
+				{
+					Widget = RootWidget->GetCachedWidget().ToSharedRef();
+				}
+			}
+		}
+
+		return Widget;
+	}
+}
+
+void UNavigationUIComponent::OnConstruct()
+{
+	TWeakObjectPtr<UWidget> OwningWidget = GetOwner();
+	if (!ensure(OwningWidget.IsValid()))
+	{
+		return;
+	}
+
+	UUserWidget* OwningUserWidget = OwningWidget->GetTypedOuter<UUserWidget>();
+	if (!ensure(OwningUserWidget))
+	{
+		return;
+	}
+
+	// Bind user provided functions to component delegates
+
+	OnNavigationEnteredDelegate.BindUFunction(OwningUserWidget, OnNavigationEntered);
+	OnNavigationExitedDelegate.BindUFunction(OwningUserWidget, OnNavigationExited);
+
+	TSharedPtr<SWidget> Content = OwningUserWidget->GetCachedWidget();
+	if (!Content.IsValid())
+	{
+		return;
+	}
+
+	NavigationTransitionMetadata = MakeShared<FNavigationTransitionMetadata>();
+	NavigationTransitionMetadata->OnNavigationTransition.BindUObject(this, &UNavigationUIComponent::HandleNavigationTransition);
+
+	if (TSharedPtr<SWidget> NavigationWidget = UE::NavigationUIComponent::Private::GetNavigationWidget(Content))
+	{
+		NavigationWidget->AddMetadata(NavigationTransitionMetadata.ToSharedRef());
+	}
+}
+
+void UNavigationUIComponent::OnDestruct()
+{
+	OnNavigationEnteredDelegate.Clear();
+	OnNavigationExitedDelegate.Clear();
+}
+
+void UNavigationUIComponent::HandleNavigationTransition(const FNavigationTransition& NavigationTransition)
+{
+	// Forward slate navigation transition to user provided functions via component delegates
+
+	UUserWidget* OldFocusedWidgetObject = UE::NavigationUIComponent::Private::FindUserWidget(NavigationTransition.OldFocusedWidget);
+	UUserWidget* NewFocusedWidgetObject = UE::NavigationUIComponent::Private::FindUserWidget(NavigationTransition.NewFocusedWidget);
+
+	switch (NavigationTransition.Direction)
+	{
+	case ENavigationTransitionDirection::Incoming:
+	{
+		OnNavigationEnteredDelegate.ExecuteIfBound(NavigationTransition.Type, OldFocusedWidgetObject, NewFocusedWidgetObject);
+		break;
+	}
+	case ENavigationTransitionDirection::Outgoing:
+	{
+		OnNavigationExitedDelegate.ExecuteIfBound(NavigationTransition.Type, OldFocusedWidgetObject, NewFocusedWidgetObject);
+		break;
+	}
+	}
+}
diff --git a/Engine/Source/Runtime/UMG/Public/Extensions/UIComponents/NavigationUIComponent.h b/Engine/Source/Runtime/UMG/Public/Extensions/UIComponents/NavigationUIComponent.h
new file mode 100644
--- /dev/null
+++ b/Engine/Source/Runtime/UMG/Public/Extensions/UIComponents/NavigationUIComponent.h
@@ -0,0 +1,53 @@
+// Copyright Epic Games, Inc. All Rights Reserved.
+
+#pragma once
+
+#include "CoreMinimal.h"
+#include "Input/NavigationMetadata.h"
+#include "Extensions/UIComponent.h"
+#include "UObject/ObjectMacros.h"
+#include "Widgets/SWidget.h"
+
+#include "NavigationUIComponent.generated.h"
+
+class SWidget;
+class UUserWidget;
+
+/**
+* If Entered, Other Widget is the Widget Navigated From 
+* If Exited , Other Widget is the Widget Navigated To
+*/
+DECLARE_DYNAMIC_DELEGATE_ThreeParams(FOnNavigationTransition, EUINavigation, Type, UUserWidget*, OldFocusedWidget, UUserWidget*, NewFocusedWidget);
+
+UCLASS(BlueprintType, Experimental)
+class UNavigationUIComponent : public UUIComponent
+{
+	GENERATED_BODY()
+
+private:
+	virtual void OnConstruct() override;
+	virtual void OnDestruct() override;
+
+	void HandleNavigationTransition(const FNavigationTransition& NavigationTransition);
+public:
+	/* FNames are references to user provided functions.  */
+
+	UPROPERTY(EditAnywhere, BlueprintReadWrite, Category = Navigation)
+	FName OnNavigationEntered; 
+
+	UPROPERTY(EditAnywhere, BlueprintReadWrite, Category = Navigation)
+	FName OnNavigationExited;
+
+private:
+	/* Delegates are the concrete functions from the provided user provided function names */
+	
+	UPROPERTY()
+	FOnNavigationTransition OnNavigationEnteredDelegate; 
+
+	UPROPERTY()
+	FOnNavigationTransition OnNavigationExitedDelegate;
+
+	TSharedPtr<FNavigationTransitionMetadata> NavigationTransitionMetadata;
+};
+
+
