UIComponents: Add missing pragma once

#rnx

diff --git a/Engine/Source/Editor/UMGEditor/Private/Extensions/UIComponents/NavigationUIComponentCustomization.h b/Engine/Source/Editor/UMGEditor/Private/Extensions/UIComponents/NavigationUIComponentCustomization.h
--- a/Engine/Source/Editor/UMGEditor/Private/Extensions/UIComponents/NavigationUIComponentCustomization.h
+++ b/Engine/Source/Editor/UMGEditor/Private/Extensions/UIComponents/NavigationUIComponentCustomization.h
@@ -1,5 +1,7 @@
 // Copyright Epic Games, Inc. All Rights Reserved.
 
+#pragma once
+
 #include "Layout/Visibility.h"
 #include "WidgetBlueprintEditor.h"
 #include "IPropertyTypeCustomization.h"
