UMG: Add tooltip to designer screensize toolbar.

#rnx

diff --git a/Engine/Source/Editor/UMGEditor/Private/Designer/SDesignerView.cpp b/Engine/Source/Editor/UMGEditor/Private/Designer/SDesignerView.cpp
--- a/Engine/Source/Editor/UMGEditor/Private/Designer/SDesignerView.cpp
+++ b/Engine/Source/Editor/UMGEditor/Private/Designer/SDesignerView.cpp
@@ -601,6 +601,7 @@ TSharedRef<SWidget> SDesignerView::CreateOverlayUI()
 			[
 				SNew(STextBlock)
 				.Text(LOCTEXT("ScreenSize", "Screen Size"))
+				.ToolTipText(LOCTEXT("ScreenSize_ToolTip", "Change the size of the designer canvas to the screen resolution."))
 				.TextStyle(&ToolBarStyle.LabelStyle)
 				.ColorAndOpacity(FSlateColor::UseForeground())
 			]
@@ -620,6 +621,7 @@ TSharedRef<SWidget> SDesignerView::CreateOverlayUI()
 			[
 				SNew(STextBlock)
 				.Text(this, &SDesignerView::GetScreenSizingFillText)
+				.ToolTipText(LOCTEXT("ScreenFill_ToolTip", "Adjust how the widget content fills the selected Screen Size."))
 				.TextStyle(&ToolBarStyle.LabelStyle)
 				.ColorAndOpacity(FSlateColor::UseForeground())
 			]
