UMG Designer: Fix to a nullptr access when a property in changed and a non UWIdget was added to the SelectedObjects

diff --git a/Engine/Source/Editor/UMGEditor/Private/Details/SWidgetDetailsView.cpp b/Engine/Source/Editor/UMGEditor/Private/Details/SWidgetDetailsView.cpp
--- a/Engine/Source/Editor/UMGEditor/Private/Details/SWidgetDetailsView.cpp
+++ b/Engine/Source/Editor/UMGEditor/Private/Details/SWidgetDetailsView.cpp
@@ -579,7 +579,7 @@ void SWidgetDetailsView::HandleNameTextCommitted(const FText& Text, ETextCommit:
 			if ( HandleVerifyNameTextChanged(Text, DummyText) )
 			{
 				UWidget* Widget = Cast<UWidget>(SelectedObjects[0].Get());
-				if (!Widget->GetLabelText().EqualToCaseIgnored(Text))
+				if (Widget && !Widget->GetLabelText().EqualToCaseIgnored(Text))
 				{
 					FWidgetBlueprintEditorUtils::RenameWidget(BlueprintEditor.Pin().ToSharedRef(), Widget->GetFName(), Text.ToString());
 				}
@@ -624,12 +624,15 @@ void SWidgetDetailsView::HandleIsVariableChanged(ECheckBoxState CheckState)
 		TArray<UWidget*, TInlineAllocator<16>> WidgetToModify;
 		for (const TWeakObjectPtr<UObject>& SelectedObject : SelectedObjects)
 		{
-			UWidget* Widget = Cast<UWidget>(SelectedObject.Get());
-			FWidgetReference WidgetRef = BPEditor->GetReferenceFromTemplate(Blueprint->WidgetTree->FindWidget(Widget->GetFName()));
-			if (WidgetRef.IsValid())
+			if (UWidget* Widget = Cast<UWidget>(SelectedObject.Get()))
 			{
-				WidgetToModify.Add(WidgetRef.GetTemplate());
-				WidgetToModify.Add(WidgetRef.GetPreview());
+				check(Blueprint->WidgetTree);
+				FWidgetReference WidgetRef = BPEditor->GetReferenceFromTemplate(Blueprint->WidgetTree->FindWidget(Widget->GetFName()));
+				if (WidgetRef.IsValid())
+				{
+					WidgetToModify.Add(WidgetRef.GetTemplate());
+					WidgetToModify.Add(WidgetRef.GetPreview());
+				}
 			}
 		}
 
