Slate: Modify GetCachedWidget.
1. GetCachedWidget now returns the Widget returned by RebuildWidget.
2. GetWrappedWidget returns the Widget return by RebuildDesignWidget if available. Otherwise, returns the widget returned by UIComponents via RebuildWidgetWithContent

[RN] minor creator/ui

diff --git a/Engine/Source/Editor/UMGEditor/Private/Designer/SDesignerView.cpp b/Engine/Source/Editor/UMGEditor/Private/Designer/SDesignerView.cpp
--- a/Engine/Source/Editor/UMGEditor/Private/Designer/SDesignerView.cpp
+++ b/Engine/Source/Editor/UMGEditor/Private/Designer/SDesignerView.cpp
@@ -1336,7 +1336,7 @@ bool SDesignerView::GetWidgetGeometry(const FWidgetReference& Widget, FGeometry&
 
 bool SDesignerView::GetWidgetGeometry(const UWidget* InPreviewWidget, FGeometry& Geometry) const
 {
-	if (TSharedPtr<SWidget> CachedPreviewWidget = InPreviewWidget->GetCachedWidget())
+	if (TSharedPtr<SWidget> CachedPreviewWidget = InPreviewWidget->GetCachedWrappedWidget())
 	{
 		const FArrangedWidget* ArrangedWidget = CachedWidgetGeometry.Find(CachedPreviewWidget.ToSharedRef());
 		if ( ArrangedWidget )
@@ -2450,7 +2450,7 @@ void SDesignerView::Tick( const FGeometry& AllottedGeometry, const double InCurr
 
 	if ( PreviewWidget )
 	{
-		if (TSharedPtr<SWidget> CachedWidget = PreviewWidget->GetCachedWidget())
+		if (TSharedPtr<SWidget> CachedWidget = PreviewWidget->GetCachedWrappedWidget())
 		{
 			CachedPreviewDesiredSize = CachedWidget->GetDesiredSize();
 		}
@@ -2554,7 +2554,7 @@ FReply SDesignerView::OnDragDetected(const FGeometry& MyGeometry, const FPointer
 		{		
 			if (UWidget* WidgetPreview = SelectedWidget.GetPreview())
 			{
-				if (TSharedPtr<SWidget> CachedWidget = WidgetPreview->GetCachedWidget())
+				if (TSharedPtr<SWidget> CachedWidget = WidgetPreview->GetCachedWrappedWidget())
 				{
 					// Determine The offset to keep the widget from the mouse while dragging
 					FArrangedWidget ArrangedWidget(SNullWidget::NullWidget, FGeometry());
@@ -3115,7 +3115,7 @@ void SDesignerView::MoveWidgets(const FGeometry& MyGeometry, const FDragDropEven
 						{
 							if (bIsPreview)
 							{
-								if (TSharedPtr<SWidget> CachedWidget = Widget->GetCachedWidget())
+								if (TSharedPtr<SWidget> CachedWidget = Widget->GetCachedWrappedWidget())
 								{
 									CanvasSlot->SaveBaseLayout();
 
diff --git a/Engine/Source/Runtime/UMG/Private/Components/Widget.cpp b/Engine/Source/Runtime/UMG/Private/Components/Widget.cpp
--- a/Engine/Source/Runtime/UMG/Private/Components/Widget.cpp
+++ b/Engine/Source/Runtime/UMG/Private/Components/Widget.cpp
@@ -1108,25 +1108,29 @@ TSharedRef<SWidget> UWidget::TakeWidget_Private(ConstructMethodType ConstructMet
 }
 
 TSharedPtr<SWidget> UWidget::GetCachedWidget() const
-{
+{	
+	if (MyGCWidget.IsValid())
+	{
+		return MyGCWidget.Pin();
+	}
+
+	return MyWidget.Pin();
+}
+
+TSharedPtr<SWidget> UWidget::GetCachedWrappedWidget() const
+{	
 #if WITH_EDITOR
 	if (DesignWrapperWidget.IsValid())
 	{
 		return DesignWrapperWidget.Pin();
 	}
-#endif
-
+#endif //WITH_EDITOR
+	
 	if (ComponentWrapperWidget.IsValid())
 	{
 		return ComponentWrapperWidget.Pin();
 	}
-	
-	if (MyGCWidget.IsValid())
-	{
-		return MyGCWidget.Pin();
-	}
-
-	return MyWidget.Pin();
+	return GetCachedWidget();
 }
 
 bool UWidget::IsConstructed() const
diff --git a/Engine/Source/Runtime/UMG/Public/Components/Widget.h b/Engine/Source/Runtime/UMG/Public/Components/Widget.h
--- a/Engine/Source/Runtime/UMG/Public/Components/Widget.h
+++ b/Engine/Source/Runtime/UMG/Public/Components/Widget.h
@@ -845,9 +845,13 @@ private:
 
 public:
 
-	/** Gets the last created widget does not recreate the gc container for the widget if one is needed. */
+	/** Gets the last created widget from RebuildWidget does not recreate the gc container for the widget if one is needed. */
 	UMG_API TSharedPtr<SWidget> GetCachedWidget() const;
 
+	/** Gets the last created widget from RebuildDesignWidget does not recreate the gc container for the widget if one is needed. */
+	UMG_API TSharedPtr<SWidget> GetCachedWrappedWidget() const;
+
+
 	/** Gets the last created widget does not recreate the gc container for the widget if one is needed. */
 	UMG_API bool IsConstructed() const;
 
