UMG: Fix warning V595: The 'PropertyThatChanged' pointer was utilized before it was verified against nullptr

#rnx

diff --git a/Engine/Source/Editor/UMGEditor/Private/WidgetBlueprintEditor.cpp b/Engine/Source/Editor/UMGEditor/Private/WidgetBlueprintEditor.cpp
--- a/Engine/Source/Editor/UMGEditor/Private/WidgetBlueprintEditor.cpp
+++ b/Engine/Source/Editor/UMGEditor/Private/WidgetBlueprintEditor.cpp
@@ -1241,15 +1241,14 @@ void FWidgetBlueprintEditor::MigrateFromChain(const FPropertyChangedEvent* Prope
 				// Only do the migration if the property on the head of the linked list lives in the CDO
 				// We want to skip the migration if the property isn't leading to this CDO, such as when we call 
 				// FDetailCategoryImpl::AddExternalObjects to inject external objects in the details panel of the CDO.
-				if (SelectedObjectPtr->IsA(OwnerClass) && WidgetCDO->IsA(OwnerClass))
+				if (SelectedObjectPtr->IsA(OwnerClass) && WidgetCDO->IsA(OwnerClass) && PropertyThatChanged)
 				{
 					if (PropertyChangedEvent == nullptr)
 					{
 						WidgetCDO->SetEditChangePropagationFlags(EEditChangePropagationFlags::OnlyMarkRealignedInstancesAsDirty);
 						WidgetCDO->PreEditChange(*PropertyThatChanged);
 					}
-					
-					if (PropertyChangedEvent && PropertyThatChanged)
+					else
 					{
 						// We have to do this before we call MigratePropertyValue, because that will change the CDO, and we won't be
 						// able to cheat like we do still having a pristine CDO copy. The details panel has to do shenanigans where it copies
