UMG: Only select if the widgets aren't already selected.

Double clicking a color key in Sequencer should open the color picker. The mouse up event was causing a selection synchronization event, and causing widgets to be reselected even if they were already selected

diff --git a/Engine/Source/Editor/UMGEditor/Private/WidgetBlueprintEditor.cpp b/Engine/Source/Editor/UMGEditor/Private/WidgetBlueprintEditor.cpp
--- a/Engine/Source/Editor/UMGEditor/Private/WidgetBlueprintEditor.cpp
+++ b/Engine/Source/Editor/UMGEditor/Private/WidgetBlueprintEditor.cpp
@@ -612,6 +612,11 @@ void FWidgetBlueprintEditor::RegisterApplicationModes(const TArray<UBlueprint*>&
 
 void FWidgetBlueprintEditor::SelectWidgets(const TSet<FWidgetReference>& Widgets, bool bAppendOrToggle)
 {
+	if (!bAppendOrToggle && SelectedWidgets.Array() == Widgets.Array())
+	{
+		return;
+	}
+
 	TSet<FWidgetReference> TempSelection;
 	for ( const FWidgetReference& Widget : Widgets )
 	{
