UIComponents: Hook into graph renaming

#rnx

diff --git a/Engine/Source/Editor/UMGEditor/Private/WidgetBlueprint.cpp b/Engine/Source/Editor/UMGEditor/Private/WidgetBlueprint.cpp
--- a/Engine/Source/Editor/UMGEditor/Private/WidgetBlueprint.cpp
+++ b/Engine/Source/Editor/UMGEditor/Private/WidgetBlueprint.cpp
@@ -15,7 +15,6 @@
 #include "Animation/UMGSequencePlayer.h"
 
 #include "Kismet2/StructureEditorUtils.h"
-
 #include "Kismet2/CompilerResultsLog.h"
 #include "Binding/PropertyBinding.h"
 #include "Blueprint/WidgetBlueprintGeneratedClass.h"
@@ -29,6 +28,9 @@
 #include "WidgetGraphSchema.h"
 #include "UMGEditorProjectSettings.h"
 
+#include "Extensions/UIComponent.h"
+#include "UIComponentWidgetBlueprintExtension.h"
+
 #if WITH_EDITOR
 #include "Interfaces/ITargetPlatform.h"
 #include "Modules/ModuleManager.h"
@@ -753,14 +755,34 @@ void UWidgetBlueprint::NotifyGraphRenamed(class UEdGraph* Graph, FName OldName,
 {
 	Super::NotifyGraphRenamed(Graph, OldName, NewName);
 	
+	const TObjectPtr<UBlueprintExtension>* UIComponentExtensionPtr = GetExtensions().FindByPredicate([](const UBlueprintExtension* Other) {
+		return Other && Other->GetClass() == UUIComponentWidgetBlueprintExtension::StaticClass(); 
+	});
+
+	UUIComponentWidgetBlueprintExtension* UIComponentExtension = nullptr;
+	if (UIComponentExtensionPtr)
+	{
+		UIComponentExtension = CastChecked<UUIComponentWidgetBlueprintExtension>(*UIComponentExtensionPtr);
+	}
+
 	// Update any explicit widget bindings.
-	WidgetTree->ForEachWidget([OldName, NewName](UWidget* Widget) {
+	WidgetTree->ForEachWidget([Graph, OldName, NewName, UIComponentExtension](UWidget* Widget) {
 		if (Widget->Navigation)
 		{
 			Widget->Navigation->SetFlags(RF_Transactional);
 			Widget->Navigation->Modify();
 			Widget->Navigation->TryToRenameBinding(OldName, NewName);
 		}
+
+		if (UIComponentExtension != nullptr)
+		{
+			for (UUIComponent* Component : UIComponentExtension->GetComponentsFor(Widget))
+			{
+				Component->SetFlags(RF_Transactional);
+				Component->Modify();
+				Component->GraphRenamed(Graph, OldName, NewName);
+			}
+		}
 	});
 }
 
diff --git a/Engine/Source/Runtime/UMG/Private/Extensions/UIComponent.cpp b/Engine/Source/Runtime/UMG/Private/Extensions/UIComponent.cpp
--- a/Engine/Source/Runtime/UMG/Private/Extensions/UIComponent.cpp
+++ b/Engine/Source/Runtime/UMG/Private/Extensions/UIComponent.cpp
@@ -38,6 +38,11 @@ void UUIComponent::Destruct()
 	OnDestruct();
 }
 
+void UUIComponent::GraphRenamed(UEdGraph* Graph, const FName& OldName, const FName& NewName)
+{
+	OnGraphRenamed(Graph, OldName, NewName);
+}
+
 TWeakObjectPtr<UWidget> UUIComponent::GetOwner() const
 {
 	return Owner;
@@ -63,6 +68,11 @@ void UUIComponent::OnDestruct()
 
 }
 
+void UUIComponent::OnGraphRenamed(UEdGraph* Graph, const FName& OldName, const FName& NewName)
+{
+
+}
+
 TSharedRef<SWidget> UUIComponent::RebuildWidgetWithContent(TSharedRef<SWidget> Content)
 {
 	return Content;
diff --git a/Engine/Source/Runtime/UMG/Private/Extensions/UIComponents/NavigationUIComponent.cpp b/Engine/Source/Runtime/UMG/Private/Extensions/UIComponents/NavigationUIComponent.cpp
--- a/Engine/Source/Runtime/UMG/Private/Extensions/UIComponents/NavigationUIComponent.cpp
+++ b/Engine/Source/Runtime/UMG/Private/Extensions/UIComponents/NavigationUIComponent.cpp
@@ -107,6 +107,18 @@ void UNavigationUIComponent::OnDestruct()
 	OnNavigationExitedDelegate.Clear();
 }
 
+void UNavigationUIComponent::OnGraphRenamed(UEdGraph* Graph, const FName& OldName, const FName& NewName)
+{
+	if (OnNavigationEntered == OldName)
+	{
+		OnNavigationEntered = NewName;
+	}
+	if (OnNavigationExited == OldName)
+	{
+		OnNavigationExited = NewName;
+	}
+}
+
 void UNavigationUIComponent::HandleNavigationTransition(const FNavigationTransition& NavigationTransition)
 {
 	// Forward slate navigation transition to user provided functions via component delegates
diff --git a/Engine/Source/Runtime/UMG/Public/Extensions/UIComponent.h b/Engine/Source/Runtime/UMG/Public/Extensions/UIComponent.h
--- a/Engine/Source/Runtime/UMG/Public/Extensions/UIComponent.h
+++ b/Engine/Source/Runtime/UMG/Public/Extensions/UIComponent.h
@@ -12,6 +12,7 @@
 
 #include "UIComponent.generated.h"
 
+class UEdGraph;
 class UWidget;
 class SWidget;
 
@@ -50,6 +51,11 @@ public:
 	 */
 	UMG_API void Destruct();
 
+	/**
+	 * Called when a graph owned by the user widget is renamed.
+	 */
+	UMG_API virtual void GraphRenamed(UEdGraph* Graph, const FName& OldName, const FName& NewName);
+
 	/**
 	 * Returns the Owner Widget this component is attached to.
 	 *  @returns The owner widget
@@ -78,6 +84,7 @@ protected:
 	UMG_API virtual void OnPreConstruct(bool bIsDesignTime);
 	UMG_API virtual void OnConstruct();
 	UMG_API virtual void OnDestruct();
+	UMG_API virtual void OnGraphRenamed(UEdGraph* Graph, const FName& OldName, const FName& NewName);
 
 private:
 	UPROPERTY(transient, DuplicateTransient)
diff --git a/Engine/Source/Runtime/UMG/Public/Extensions/UIComponents/NavigationUIComponent.h b/Engine/Source/Runtime/UMG/Public/Extensions/UIComponents/NavigationUIComponent.h
--- a/Engine/Source/Runtime/UMG/Public/Extensions/UIComponents/NavigationUIComponent.h
+++ b/Engine/Source/Runtime/UMG/Public/Extensions/UIComponents/NavigationUIComponent.h
@@ -25,6 +25,7 @@ class UNavigationUIComponent : public UUIComponent
 private:
 	virtual void OnConstruct() override;
 	virtual void OnDestruct() override;
+	virtual void OnGraphRenamed(UEdGraph* Graph, const FName& OldName, const FName& NewName) override;
 
 	void HandleNavigationTransition(const FNavigationTransition& NavigationTransition);
 public:
