Fix to a crash when ReplaceWidget command from the contex menu was used with an item under User Created was Selected in Pallette.

To fix it I had to replace the code to use FWidgetTemplateClass instead of the UClass.

diff --git a/Engine/Source/Editor/UMGEditor/Private/WidgetBlueprintEditorUtils.cpp b/Engine/Source/Editor/UMGEditor/Private/WidgetBlueprintEditorUtils.cpp
--- a/Engine/Source/Editor/UMGEditor/Private/WidgetBlueprintEditorUtils.cpp
+++ b/Engine/Source/Editor/UMGEditor/Private/WidgetBlueprintEditorUtils.cpp
@@ -1166,18 +1166,39 @@ void FWidgetBlueprintEditorUtils::ReplaceWidgetWithSelectedTemplate(TSharedRef<F
 		return;
 	}
 
-	UClass* ReplacementWidgetClass = BlueprintEditor->GetSelectedTemplate().Get();
-	if (!ReplacementWidgetClass)
+	TSharedPtr<FWidgetTemplateClass> TemplateClass;
+	if (UClass* ReplacementWidgetClass = BlueprintEditor->GetSelectedTemplate().Get())
 	{
-		ReplacementWidgetClass = BlueprintEditor->GetSelectedUserWidget().GetClass(EResolveClass::Yes);
+		TemplateClass = MakeShareable(new FWidgetTemplateClass(ReplacementWidgetClass));
+	}
+	else
+	{	
+		const FAssetData SelectedUserWidget = BlueprintEditor->GetSelectedUserWidget();		
+		if (SelectedUserWidget.GetSoftObjectPath().IsValid())
+		{
+			TemplateClass = MakeShareable(new FWidgetTemplateBlueprintClass(SelectedUserWidget));
+		}	
 	}
 	
-	if (!ReplacementWidgetClass)
+	if (!TemplateClass.IsValid())
+	{
+		return;
+	}
+
+	FScopedTransaction Transaction(LOCTEXT("ReplaceWidgets", "Replace Widgets"));
+
+	TArray<UWidget*> UsedVariables;
+	TArray<FText> WidgetNames;
+	const bool bIncludeChildrenVariables = false;
+	FindUsedVariablesForWidgets({WidgetToReplace}, BP, UsedVariables, WidgetNames, bIncludeChildrenVariables);
+
+	if (UsedVariables.Num() != 0 && !ShouldContinueReplaceOperation(BP, WidgetNames))
 	{
+		Transaction.Cancel();
 		return;
 	}
 	
-	ReplaceWidgets(BP, {WidgetToReplace}, ReplacementWidgetClass, EReplaceWidgetNamingMethod::MaintainNameAndReferences);
+	ReplaceWidgetsWithTemplateClass(BP, {WidgetToReplace}, TemplateClass, EReplaceWidgetNamingMethod::MaintainNameAndReferences);
 }
 
 bool FWidgetBlueprintEditorUtils::CanBeReplacedWithTemplate(TSharedRef<FWidgetBlueprintEditor> BlueprintEditor, UWidgetBlueprint* BP, FWidgetReference Widget)
@@ -1368,7 +1389,13 @@ void FWidgetBlueprintEditorUtils::ReplaceWidgets(UWidgetBlueprint* BP, TSet<UWid
 		return;
 	}
 
-	TSharedPtr<FWidgetTemplateClass> Template = MakeShareable(new FWidgetTemplateClass(WidgetClass));
+	TSharedPtr<FWidgetTemplateClass> TemplateClass = MakeShareable(new FWidgetTemplateClass(WidgetClass));	
+	
+	ReplaceWidgetsWithTemplateClass(BP, Widgets, TemplateClass, NewWidgetNamingMethod);
+}
+
+void FWidgetBlueprintEditorUtils::ReplaceWidgetsWithTemplateClass(UWidgetBlueprint* BP, TSet<UWidget*> Widgets, TSharedPtr<FWidgetTemplateClass> TemplateClass, EReplaceWidgetNamingMethod NewWidgetNamingMethod)
+{
 	TMap<FName, FName> ReplacedWidgetMap;
 	
 	for (UWidget* Item : Widgets)
@@ -1376,7 +1403,7 @@ void FWidgetBlueprintEditorUtils::ReplaceWidgets(UWidgetBlueprint* BP, TSet<UWid
 		BP->WidgetTree->SetFlags(RF_Transactional);
 		BP->WidgetTree->Modify();
 
-		UWidget* NewReplacementWidget = Template->Create(BP->WidgetTree);
+		UWidget* NewReplacementWidget = TemplateClass->Create(BP->WidgetTree);
 		UWidget* WidgetToReplace = Item;
 
 		// If replacing a panel widget, then it must not have children or the replacement must also be a panel widget
diff --git a/Engine/Source/Editor/UMGEditor/Public/WidgetBlueprintEditorUtils.h b/Engine/Source/Editor/UMGEditor/Public/WidgetBlueprintEditorUtils.h
--- a/Engine/Source/Editor/UMGEditor/Public/WidgetBlueprintEditorUtils.h
+++ b/Engine/Source/Editor/UMGEditor/Public/WidgetBlueprintEditorUtils.h
@@ -20,6 +20,7 @@ class SWindow;
 class UWidgetEditingProjectSettings;
 class FWidgetObjectTextFactory;
 class ULocalPlayer;
+class FWidgetTemplateClass;
 
 //////////////////////////////////////////////////////////////////////////
 // FWidgetBlueprintEditorUtils
@@ -200,6 +201,8 @@ private:
 	static UE_API void ReplaceWidgetWithChildren(TSharedRef<FWidgetBlueprintEditor> BlueprintEditor, UWidgetBlueprint* BP, FWidgetReference Widget);
 
 	static UE_API void ReplaceWidgetWithNamedSlot(TSharedRef<FWidgetBlueprintEditor> BlueprintEditor, UWidgetBlueprint* BP, FWidgetReference Widget, FName NamedSlot);
+	
+	static UE_API void ReplaceWidgetsWithTemplateClass(UWidgetBlueprint* BP, TSet<UWidget*> Widgets, TSharedPtr<FWidgetTemplateClass> TemplateClass, EReplaceWidgetNamingMethod NewWidgetNamingMethod);
 
 	static UE_API FString FindNextValidName(UWidgetTree* WidgetTree, const FString& Name);
 
