Add support for Component copied at the same time as the Widget it is attached to.

#FailedValidations Preflight

diff --git a/Engine/Source/Editor/UMGEditor/Private/UIComponentClipboardExtension.cpp b/Engine/Source/Editor/UMGEditor/Private/UIComponentClipboardExtension.cpp
new file mode 100644
--- /dev/null
+++ b/Engine/Source/Editor/UMGEditor/Private/UIComponentClipboardExtension.cpp
@@ -0,0 +1,103 @@
+// Copyright Epic Games, Inc. All Rights Reserved.
+
+#include "UIComponentClipboardExtension.h"
+
+#include "Exporters/Exporter.h"
+#include "Extensions/UIComponent.h"
+#include "Kismet2/BlueprintEditorUtils.h"
+#include "UIComponentUtils.h"
+#include "UIComponentWidgetBlueprintExtension.h"
+#include "WidgetBlueprintEditorUtils.h"
+
+void FUIComponentClipboardExtension::AppendToClipboard(const UWidget* Widget, const FExportArgs& ExportArgs)
+{
+	check(ExportArgs.Out);
+
+	if (const UWidgetBlueprint* WidgetBlueprint = FWidgetBlueprintEditorUtils::GetWidgetBlueprintFromWidget(Widget))
+	{
+		if (const UUIComponentWidgetBlueprintExtension* ComponentExtension = UUIComponentWidgetBlueprintExtension::GetExtension<UUIComponentWidgetBlueprintExtension>(WidgetBlueprint))
+		{
+			for (const UUIComponent* Component : ComponentExtension->GetComponentsFor(Widget))
+			{
+				UUIComponentWidgetPair* ComponentWidgetPair = NewObject<UUIComponentWidgetPair>();
+				ComponentWidgetPair->WidgetName = Widget->GetFName();				
+				ComponentWidgetPair->Component = CastChecked<UUIComponent>(StaticDuplicateObject(Component, ComponentWidgetPair));								
+				
+				UExporter::ExportToOutputDevice(ExportArgs.Context, ComponentWidgetPair, ExportArgs.Exporter, *ExportArgs.Out, ExportArgs.FileType, ExportArgs.Indent, ExportArgs.PortFlags, ExportArgs.bSelectedOnly, ExportArgs.ExportRootScope);				
+			}		
+		}
+	}
+}
+
+bool FUIComponentClipboardExtension::CanAppendToClipboard(const UWidget* Widget) const
+{
+	if (FWidgetBlueprintEditorUtils::GetWidgetBlueprintFromWidget(Widget))
+	{
+		return true;
+	}
+	return false;
+}
+
+void FUIComponentClipboardExtension::ProcessImportedText(const UWidgetBlueprint* WidgetBlueprint, const FString& TextToImport, UPackage*& TempPackage)
+{
+	Factory.Reinit();
+	Factory.ProcessBuffer(TempPackage, RF_Transactional, TextToImport);
+}
+
+bool FUIComponentClipboardExtension::CanImportFromClipboard(const UWidget* Widget) const
+{
+	return !Factory.IsEmpty();
+}
+
+void FUIComponentClipboardExtension::ImportDataToWidget(const UWidget* Widget, FName OldWidgetName)
+{
+	Factory.AddProcessedComponentsTo(Widget, OldWidgetName);
+}
+
+bool FUIComponentClipboardExtension::CanWidgetAcceptPaste(const UWidget* Widget) const
+{
+	constexpr bool bCanWidgetAcceptPaste = true;
+	return bCanWidgetAcceptPaste;
+}
+
+FUIComponentClipboardExtension::FExtensionTextFactory::FExtensionTextFactory()
+	: FCustomizableTextObjectFactory(GWarn)
+{
+}
+
+bool FUIComponentClipboardExtension::FExtensionTextFactory::CanCreateClass(UClass* ObjectClass, bool& bOmitSubObjs) const
+{
+	const bool bIsComponentWidgetPair = ObjectClass->IsChildOf(UUIComponentWidgetPair::StaticClass());
+	return bIsComponentWidgetPair;
+}
+
+void FUIComponentClipboardExtension::FExtensionTextFactory::ProcessConstructedObject(UObject* NewObject)
+{
+	check(NewObject);
+
+	if (UUIComponentWidgetPair* ComponentWidgetPair = Cast<UUIComponentWidgetPair>(NewObject))
+	{
+		NewComponentWidgetPairs.Add(ComponentWidgetPair);
+	}
+}
+void FUIComponentClipboardExtension::FExtensionTextFactory::AddProcessedComponentsTo(const UWidget* Widget, FName OldWidgetName)
+{
+	if (IsEmpty())
+	{
+		return;
+	}
+
+	UWidgetBlueprint* WidgetBlueprint = FWidgetBlueprintEditorUtils::GetWidgetBlueprintFromWidget(Widget);	
+	// Here we call request instead of GetExtension to make sure to create one if it doesn't exist.
+	if (UUIComponentWidgetBlueprintExtension* ComponentExtension = WidgetBlueprint ? UUIComponentWidgetBlueprintExtension::RequestExtension<UUIComponentWidgetBlueprintExtension>(WidgetBlueprint) : nullptr)
+	{
+		for (const UUIComponentWidgetPair* Pair : NewComponentWidgetPairs)
+		{
+			UUIComponent* NewComponent = Pair->Component;
+			if (NewComponent && Pair->WidgetName == OldWidgetName)
+			{
+				ComponentExtension->AddOrReplaceComponent(NewComponent, Widget->GetFName());
+			}
+		}
+	}	
+}
\ No newline at end of file
diff --git a/Engine/Source/Editor/UMGEditor/Private/UIComponentClipboardExtension.h b/Engine/Source/Editor/UMGEditor/Private/UIComponentClipboardExtension.h
new file mode 100644
--- /dev/null
+++ b/Engine/Source/Editor/UMGEditor/Private/UIComponentClipboardExtension.h
@@ -0,0 +1,70 @@
+// Copyright Epic Games, Inc. All Rights Reserved.
+
+#pragma once
+
+#include "Factories.h"
+#include "IHasClipboardExtensibility.h"
+#include "UIComponentClipboardExtension.generated.h"
+
+class UWidget;
+class UWidgetBlueprint;
+class UUIComponent;
+
+UCLASS()
+class UUIComponentWidgetPair : public UObject
+{
+	GENERATED_BODY()
+	
+public:
+
+	UUIComponentWidgetPair(){};
+
+	UPROPERTY()
+	FName WidgetName;
+
+	UPROPERTY(instanced)
+	TObjectPtr<UUIComponent> Component;
+};
+
+class FUIComponentClipboardExtension
+	: public IClipboardExtension
+{
+	//~ Begin IClipboardExtension overrides
+	virtual void AppendToClipboard(const UWidget* Widget, const FExportArgs& ExportArgs) override;
+	virtual bool CanAppendToClipboard(const UWidget* Widget) const override;
+	virtual void ProcessImportedText(const UWidgetBlueprint* WidgetBlueprint, const FString& TextToImport, UPackage*& TempPackage) override;
+	virtual bool CanImportFromClipboard(const UWidget* Widget) const override;
+	virtual void ImportDataToWidget(const UWidget* Widget, FName OldWidgetName) override;
+	virtual bool CanWidgetAcceptPaste(const UWidget* Widget) const override;
+	//~ End IClipboardExtension overrides
+
+private:
+
+	class FExtensionTextFactory : public FCustomizableTextObjectFactory
+	{
+	public:
+		FExtensionTextFactory();
+
+		//~ Begin FCustomizableTextObjectFactory overrides
+		virtual bool CanCreateClass(UClass* ObjectClass, bool& bOmitSubObjs) const override;
+		virtual void ProcessConstructedObject(UObject* NewObject) override;
+		//~ End FCustomizableTextObjectFactory overrides
+
+		void AddProcessedComponentsTo(const UWidget* Widget, FName OldWidgetName);
+
+		void Reinit()
+		{
+			NewComponentWidgetPairs.Empty();
+		}
+
+		bool IsEmpty() const
+		{			
+			return NewComponentWidgetPairs.IsEmpty();	
+		}
+
+	private:
+		TArray<UUIComponentWidgetPair*> NewComponentWidgetPairs;
+	};
+
+	FExtensionTextFactory Factory;
+};
diff --git a/Engine/Source/Editor/UMGEditor/Private/UIComponentWidgetBlueprintExtension.cpp b/Engine/Source/Editor/UMGEditor/Private/UIComponentWidgetBlueprintExtension.cpp
--- a/Engine/Source/Editor/UMGEditor/Private/UIComponentWidgetBlueprintExtension.cpp
+++ b/Engine/Source/Editor/UMGEditor/Private/UIComponentWidgetBlueprintExtension.cpp
@@ -148,6 +148,18 @@ UUIComponent* UUIComponentWidgetBlueprintExtension::AddComponent( const UClass*
 	return nullptr;
 }
 
+UUIComponent* UUIComponentWidgetBlueprintExtension::AddOrReplaceComponent(UUIComponent* Component, FName OwnerName)
+{
+	// Add the Component to the WidgetBlueprint
+	check(ComponentContainer);
+	Modify();
+	// If we already have a component of that tyype we replace it
+	ComponentContainer->RemoveAllComponentsOfType(Component->GetClass(), OwnerName);
+	ComponentContainer->AddComponent(OwnerName, Component);
+	Component->Rename(nullptr, ComponentContainer, RF_ArchetypeObject);
+	return Component;
+}
+
 void UUIComponentWidgetBlueprintExtension::HandleBeginCompilation(FWidgetBlueprintCompilerContext& InCreationContext)
 {
 	Super::HandleBeginCompilation(InCreationContext);
diff --git a/Engine/Source/Editor/UMGEditor/Private/UMGEditorModule.cpp b/Engine/Source/Editor/UMGEditor/Private/UMGEditorModule.cpp
--- a/Engine/Source/Editor/UMGEditor/Private/UMGEditorModule.cpp
+++ b/Engine/Source/Editor/UMGEditor/Private/UMGEditorModule.cpp
@@ -52,6 +52,7 @@
 #include "Extensions/UIComponentContainer.h"
 #include "Extensions/UIComponentContainerDesignerExtension.h"
 #include "UIComponentUtils.h"
+#include "UIComponentClipboardExtension.h"
 
 #define LOCTEXT_NAMESPACE "UMG"
 
@@ -101,6 +102,10 @@ public:
 
 		PropertyBindingExtensibilityManager = MakeShared<FPropertyBindingExtensibilityManager>();
 		ClipboardExtensibilityManager = MakeShared<FClipboardExtensibilityManager>();
+
+		ComponentClipboardExtension = MakeShared<FUIComponentClipboardExtension>();
+		GetClipboardExtensibilityManager()->AddExtension(ComponentClipboardExtension.ToSharedRef());
+		
 		WidgetDragDropExtensibilityManager = MakeShared<FWidgetDragDropExtensibilityManager>();
 		WidgetContextMenuExtensibilityManager = MakeShared<FWidgetContextMenuExtensibilityManager >();
 
@@ -514,6 +519,7 @@ private:
 	TSharedPtr<FExtensibilityManager> MenuExtensibilityManager;
 	TSharedPtr<FExtensibilityManager> ToolBarExtensibilityManager;
 	TSharedPtr<FDesignerExtensibilityManager> DesignerExtensibilityManager;
+	TSharedPtr<FUIComponentClipboardExtension> ComponentClipboardExtension;
 	TSharedPtr<FPropertyBindingExtensibilityManager> PropertyBindingExtensibilityManager;
 	TSharedPtr<FClipboardExtensibilityManager> ClipboardExtensibilityManager;
 	TSharedPtr<FWidgetDragDropExtensibilityManager> WidgetDragDropExtensibilityManager;
diff --git a/Engine/Source/Editor/UMGEditor/Public/UIComponentWidgetBlueprintExtension.h b/Engine/Source/Editor/UMGEditor/Public/UIComponentWidgetBlueprintExtension.h
--- a/Engine/Source/Editor/UMGEditor/Public/UIComponentWidgetBlueprintExtension.h
+++ b/Engine/Source/Editor/UMGEditor/Public/UIComponentWidgetBlueprintExtension.h
@@ -30,6 +30,7 @@ public:
 	static UE_API const FName MD_ComponentVariable;
 
 	UE_API UUIComponent* AddComponent(const UClass* ComponentClass, FName OwnerName);
+	UE_API UUIComponent* AddOrReplaceComponent(UUIComponent* Component, FName OwnerName);
 	UE_API void RemoveComponent(const UClass* ComponentClass, FName OwnerName);
 	UE_API void MoveComponent(FName OwnerName, const UClass* ComponentClass, const UClass* RelativeToComponentClass, bool bMoveAfter);
 	UE_API TArray<UUIComponent*> GetComponentsFor(const UWidget* Target) const;
diff --git a/Engine/Source/Runtime/UMG/Private/Extensions/UIComponentUserWidgetExtension.cpp b/Engine/Source/Runtime/UMG/Private/Extensions/UIComponentUserWidgetExtension.cpp
--- a/Engine/Source/Runtime/UMG/Private/Extensions/UIComponentUserWidgetExtension.cpp
+++ b/Engine/Source/Runtime/UMG/Private/Extensions/UIComponentUserWidgetExtension.cpp
@@ -108,6 +108,9 @@ void UUIComponentUserWidgetExtension::CreateAndAddComponent(UUIComponent* Archet
 	{
 		ensure(ComponentContainer && ArchetypeComponent && ArchetypeComponent->HasAllFlags(RF_ArchetypeObject));
 
+		// First, remove any component already on the Widget, replace it with a new one.
+		RemoveComponent(ArchetypeComponent->GetClass(), OwnerName);
+
 		// Create the component with the WBP Component as the Archetype.
 		UObject* Outer = ComponentContainer;	
 		UUIComponent* PreviewWidgetComponent = NewObject<UUIComponent>(Outer,
