Adding an extra vertical splitter to allow layout extensions to add tabs above and below the widget details tab

diff --git a/Engine/Source/Editor/UMGEditor/Private/BlueprintModes/WidgetDesignerApplicationMode.cpp b/Engine/Source/Editor/UMGEditor/Private/BlueprintModes/WidgetDesignerApplicationMode.cpp
--- a/Engine/Source/Editor/UMGEditor/Private/BlueprintModes/WidgetDesignerApplicationMode.cpp
+++ b/Engine/Source/Editor/UMGEditor/Private/BlueprintModes/WidgetDesignerApplicationMode.cpp
@@ -128,9 +128,14 @@ FWidgetDesignerApplicationMode::FWidgetDesignerApplicationMode(TSharedPtr<FWidge
 			)
 			->Split
 			(
-				FTabManager::NewStack()
+				FTabManager::NewSplitter()
 				->SetSizeCoefficient(0.35f)
-				->AddTab(FDetailsTabSummoner::TabID, ETabState::OpenedTab)
+				->SetOrientation(Orient_Vertical)
+				->Split
+				(
+					FTabManager::NewStack()
+					->AddTab(FDetailsTabSummoner::TabID, ETabState::OpenedTab)
+				)
 			)
 		)
 		->Split
