UMG - Anchor tests added

#rnx

diff --git a/Engine/Source/Editor/UMGEditor/Private/Customizations/CanvasSlotCustomization.cpp b/Engine/Source/Editor/UMGEditor/Private/Customizations/CanvasSlotCustomization.cpp
--- a/Engine/Source/Editor/UMGEditor/Private/Customizations/CanvasSlotCustomization.cpp
+++ b/Engine/Source/Editor/UMGEditor/Private/Customizations/CanvasSlotCustomization.cpp
@@ -82,6 +82,7 @@ public:
 			.ButtonStyle(FAppStyle::Get(), "SimpleSharpButton")
 			.ButtonColorAndOpacity(FLinearColor(FColor(40, 40, 40)))
 			.OnClicked(this, &SAnchorPreviewWidget::OnAnchorClicked, AnchorsHandle, AlignmentHandle, OffsetsHandle, Anchors)
+			.AddMetaData<FTagMetaData>(FTagMetaData(FName{ LabelText.ToString() }))
 			.ContentPadding(FMargin(2.0f, 2.0f))
 			[
 				SNew(SVerticalBox)
@@ -102,6 +103,7 @@ public:
 							SNew(SBox)
 							.WidthOverride(this, &SAnchorPreviewWidget::GetCurrentWidth)
 							.HeightOverride(this, &SAnchorPreviewWidget::GetCurrentHeight)
+							.AddMetaData<FTagMetaData>(FTagMetaData(TEXT("AnimatedAnchorPreviewBox")))
 							[
 								SNew(SBorder)
 								.Padding(1)
diff --git a/Engine/Source/Editor/UMGEditor/Private/Extensions/CanvasSlotExtension.cpp b/Engine/Source/Editor/UMGEditor/Private/Extensions/CanvasSlotExtension.cpp
--- a/Engine/Source/Editor/UMGEditor/Private/Extensions/CanvasSlotExtension.cpp
+++ b/Engine/Source/Editor/UMGEditor/Private/Extensions/CanvasSlotExtension.cpp
@@ -160,6 +160,7 @@ TSharedRef<SWidget> FCanvasSlotExtension::MakeAnchorWidget(EAnchorWidget AnchorT
 		.OnMouseButtonUp(this, &FCanvasSlotExtension::HandleAnchorEndDrag, AnchorType)
 		.OnMouseMove(this, &FCanvasSlotExtension::HandleAnchorDragging, AnchorType)
 		.Visibility(this, &FCanvasSlotExtension::GetAnchorVisibility, AnchorType)
+		.AddMetaData<FTagMetaData>(FTagMetaData(TEXT("AnchorMedallion")))
 		.Padding(FMargin(0))
 		[
 			SNew(SEventShim)
