UMG: make sure EditorStyle module is loaded before UMGEditor module

diff --git a/Engine/Source/Editor/UMGEditor/Private/UMGEditorModule.cpp b/Engine/Source/Editor/UMGEditor/Private/UMGEditorModule.cpp
--- a/Engine/Source/Editor/UMGEditor/Private/UMGEditorModule.cpp
+++ b/Engine/Source/Editor/UMGEditor/Private/UMGEditorModule.cpp
@@ -80,6 +80,7 @@ public:
 	/** Called right after the module DLL has been loaded and the module object has been created */
 	virtual void StartupModule() override
 	{
+		FModuleManager::LoadModuleChecked<IUMGModule>("EditorStyle");
 		FModuleManager::LoadModuleChecked<IUMGModule>("UMG");
 
 		// Any attempt to use GEditor right now will fail as it hasn't been initialized yet. Waiting for post engine init resolves that.
