UMG: fix a bad cast

diff --git a/Engine/Source/Editor/UMGEditor/Private/UMGEditorModule.cpp b/Engine/Source/Editor/UMGEditor/Private/UMGEditorModule.cpp
--- a/Engine/Source/Editor/UMGEditor/Private/UMGEditorModule.cpp
+++ b/Engine/Source/Editor/UMGEditor/Private/UMGEditorModule.cpp
@@ -26,6 +26,7 @@
 #include "Animation/Sequencer2DTransformTrackEditor.h"
 #include "Animation/WidgetMaterialTrackEditor.h"
 #include "Animation/MovieSceneSequenceEditor_WidgetAnimation.h"
+#include "Interfaces/IEditorStyleModule.h"
 #include "IUMGModule.h"
 #include "Designer/DesignerCommands.h"
 #include "Navigation/SWidgetDesignerNavigation.h"
@@ -80,7 +81,7 @@ public:
 	/** Called right after the module DLL has been loaded and the module object has been created */
 	virtual void StartupModule() override
 	{
-		FModuleManager::LoadModuleChecked<IUMGModule>("EditorStyle");
+		FModuleManager::LoadModuleChecked<IEditorStyleModule>("EditorStyle");
 		FModuleManager::LoadModuleChecked<IUMGModule>("UMG");
 
 		// Any attempt to use GEditor right now will fail as it hasn't been initialized yet. Waiting for post engine init resolves that.
