UMG: Fix navigation customization crash on empty selection

#rnx

diff --git a/Engine/Source/Editor/UMGEditor/Private/Customizations/WidgetNavigationCustomization.cpp b/Engine/Source/Editor/UMGEditor/Private/Customizations/WidgetNavigationCustomization.cpp
--- a/Engine/Source/Editor/UMGEditor/Private/Customizations/WidgetNavigationCustomization.cpp
+++ b/Engine/Source/Editor/UMGEditor/Private/Customizations/WidgetNavigationCustomization.cpp
@@ -823,7 +823,14 @@ void FWidgetNavigationCustomization::SetNavigationMethodStruct(UWidget* Widget,
 		}
 	}
 
-	WidgetNavigation->NavigationMethod.InitializeAsScriptStruct(NavigationMethodStruct);
+	if (NavigationMethodStruct == nullptr)
+	{
+		WidgetNavigation->NavigationMethod.Reset();
+	}
+	else
+	{
+		WidgetNavigation->NavigationMethod.InitializeAsScriptStruct(NavigationMethodStruct);
+	}
 
 	if ( WidgetNavigation->IsDefaultNavigation() )
 	{
