UMG: bCanCallInitializedWithoutPlayerContext is now set to true by default for new UUserWidgets (but not editor utility widgets)

diff --git a/Engine/Source/Editor/UMGEditor/Private/Settings/UMGEditorProjectSettings.cpp b/Engine/Source/Editor/UMGEditor/Private/Settings/UMGEditorProjectSettings.cpp
--- a/Engine/Source/Editor/UMGEditor/Private/Settings/UMGEditorProjectSettings.cpp
+++ b/Engine/Source/Editor/UMGEditor/Private/Settings/UMGEditorProjectSettings.cpp
@@ -34,6 +34,10 @@ UUMGEditorProjectSettings::UUMGEditorProjectSettings()
 	bUseUserWidgetParentClassViewerSelector = true;
 	bUseUserWidgetParentDefaultClassViewerSelector = true;
 
+	// For compatibility with the UMG Widget Previewer this must be set to true.
+	// However, we do not wish to change the value on "legacy" widgets, so this will only affect newly created ones.
+	bCanCallInitializedWithoutPlayerContext = true;
+
 	bUseWidgetTemplateSelector = false;
 	CommonRootWidgetClasses = {
 		UHorizontalBox::StaticClass(),
