[WidgetBlueprint]
- Do not try to run the leak detection code outside the game thread

diff --git a/Engine/Source/Editor/UMGEditor/Private/WidgetBlueprint.cpp b/Engine/Source/Editor/UMGEditor/Private/WidgetBlueprint.cpp
--- a/Engine/Source/Editor/UMGEditor/Private/WidgetBlueprint.cpp
+++ b/Engine/Source/Editor/UMGEditor/Private/WidgetBlueprint.cpp
@@ -790,9 +790,14 @@ EDataValidationResult UWidgetBlueprint::IsDataValid(FDataValidationContext& Cont
 {
 	EDataValidationResult Result = UBlueprint::IsDataValid(Context);
 
-	const bool bFoundLeak = DetectSlateWidgetLeaks(Context);
+	if (IsInGameThread())
+	{
+		const bool bFoundLeak = DetectSlateWidgetLeaks(Context);
 
-	return bFoundLeak ? EDataValidationResult::Invalid : Result;
+		return bFoundLeak ? EDataValidationResult::Invalid : Result;
+	}
+	
+	return Result;
 }
 
 bool UWidgetBlueprint::DetectSlateWidgetLeaks(FDataValidationContext& Context) const
