[UMG + Enhnaced Input] Fix for child widgets not having bAutomaticallyRegisterInputOnConstruction set when their parent did, preventing them from processing input events that their parents were bound to.

bAutomaticallyRegisterInputOnConstruction needs to be respected in child widgets because we still want to create an input component on the child widget if the parent has any references to input nodes in its graph. Otherwise, the child widget will not receive input events.

#tests WidgetRepro test project, PIE in Lyra, PIE in QAGame test project for the associated jira.

#Synced-CL 48656164

diff --git a/Engine/Plugins/EnhancedInput/Source/InputBlueprintNodes/Private/K2Node_EnhancedInputAction.cpp b/Engine/Plugins/EnhancedInput/Source/InputBlueprintNodes/Private/K2Node_EnhancedInputAction.cpp
--- a/Engine/Plugins/EnhancedInput/Source/InputBlueprintNodes/Private/K2Node_EnhancedInputAction.cpp
+++ b/Engine/Plugins/EnhancedInput/Source/InputBlueprintNodes/Private/K2Node_EnhancedInputAction.cpp
@@ -401,11 +401,7 @@ void UK2Node_EnhancedInputAction::ExpandNode(FKismetCompilerContext& CompilerCon
 	// Widget blueprints require the bAutomaticallyRegisterInputOnConstruction to be set to true in order to receive callbacks
 	if (GetBlueprint()->IsA<UWidgetBlueprint>())
 	{
-		CompilerContext.AddPostCDOCompiledStep([](const UObject::FPostCDOCompiledContext& Context, UObject* NewCDO)
-		{
-			UUserWidget* Widget = CastChecked<UUserWidget>(NewCDO);
-			Widget->bAutomaticallyRegisterInputOnConstruction = true;
-		});	
+		CompilerContext.SetAutomaticallyRegisterInputEvents(true);
 	}
 
 	// Create temporary variables to copy ActionValue and ElapsedSeconds into
diff --git a/Engine/Plugins/EnhancedInput/Source/InputBlueprintNodes/Private/K2Node_InputDebugKey.cpp b/Engine/Plugins/EnhancedInput/Source/InputBlueprintNodes/Private/K2Node_InputDebugKey.cpp
--- a/Engine/Plugins/EnhancedInput/Source/InputBlueprintNodes/Private/K2Node_InputDebugKey.cpp
+++ b/Engine/Plugins/EnhancedInput/Source/InputBlueprintNodes/Private/K2Node_InputDebugKey.cpp
@@ -388,13 +388,7 @@ void UK2Node_InputDebugKey::ExpandNode(FKismetCompilerContext& CompilerContext,
 	}
 
 	// Widget blueprints require the bAutomaticallyRegisterInputOnConstruction to be set to true in order to receive callbacks
-	CompilerContext.AddPostCDOCompiledStep([](const UObject::FPostCDOCompiledContext& Context, UObject* NewCDO)
-	{
-		if (UUserWidget* Widget = Cast<UUserWidget>(NewCDO))
-		{
-			Widget->bAutomaticallyRegisterInputOnConstruction = true;
-		}
-	});
+	CompilerContext.SetAutomaticallyRegisterInputEvents(true);
 }
 
 void UK2Node_InputDebugKey::GetMenuActions(FBlueprintActionDatabaseRegistrar& ActionRegistrar) const
diff --git a/Engine/Source/Editor/BlueprintGraph/Private/K2Node_InputAction.cpp b/Engine/Source/Editor/BlueprintGraph/Private/K2Node_InputAction.cpp
--- a/Engine/Source/Editor/BlueprintGraph/Private/K2Node_InputAction.cpp
+++ b/Engine/Source/Editor/BlueprintGraph/Private/K2Node_InputAction.cpp
@@ -275,13 +275,7 @@ void UK2Node_InputAction::ExpandNode(FKismetCompilerContext& CompilerContext, UE
 	}
 
 	// Widget blueprints require the bAutomaticallyRegisterInputOnConstruction to be set to true in order to receive callbacks
-	CompilerContext.AddPostCDOCompiledStep([](const UObject::FPostCDOCompiledContext& Context, UObject* NewCDO)
-	{
-		if (UUserWidget* Widget = Cast<UUserWidget>(NewCDO))
-		{
-			Widget->bAutomaticallyRegisterInputOnConstruction = true;
-		}
-	});
+	CompilerContext.SetAutomaticallyRegisterInputEvents(true);
 }
 
 void UK2Node_InputAction::GetMenuActions(FBlueprintActionDatabaseRegistrar& ActionRegistrar) const
diff --git a/Engine/Source/Editor/BlueprintGraph/Private/K2Node_InputAxisEvent.cpp b/Engine/Source/Editor/BlueprintGraph/Private/K2Node_InputAxisEvent.cpp
--- a/Engine/Source/Editor/BlueprintGraph/Private/K2Node_InputAxisEvent.cpp
+++ b/Engine/Source/Editor/BlueprintGraph/Private/K2Node_InputAxisEvent.cpp
@@ -28,7 +28,6 @@
 #include "Templates/SubclassOf.h"
 #include "UObject/Class.h"
 #include "UObject/ObjectVersion.h"
-#include "Blueprint/UserWidget.h"
 
 #include UE_INLINE_GENERATED_CPP_BY_NAME(K2Node_InputAxisEvent)
 
@@ -264,13 +263,7 @@ void UK2Node_InputAxisEvent::ExpandNode(FKismetCompilerContext& CompilerContext,
 	Super::ExpandNode(CompilerContext, SourceGraph);
 
 	// Widget blueprints require the bAutomaticallyRegisterInputOnConstruction to be set to true in order to receive callbacks
-	CompilerContext.AddPostCDOCompiledStep([](const UObject::FPostCDOCompiledContext& Context, UObject* NewCDO)
-		{
-			if (UUserWidget* Widget = Cast<UUserWidget>(NewCDO))
-			{
-				Widget->bAutomaticallyRegisterInputOnConstruction = true;
-			}
-		});
+	CompilerContext.SetAutomaticallyRegisterInputEvents(true);
 }
 
 #undef LOCTEXT_NAMESPACE
diff --git a/Engine/Source/Editor/BlueprintGraph/Private/K2Node_InputAxisKeyEvent.cpp b/Engine/Source/Editor/BlueprintGraph/Private/K2Node_InputAxisKeyEvent.cpp
--- a/Engine/Source/Editor/BlueprintGraph/Private/K2Node_InputAxisKeyEvent.cpp
+++ b/Engine/Source/Editor/BlueprintGraph/Private/K2Node_InputAxisKeyEvent.cpp
@@ -216,13 +216,7 @@ void UK2Node_InputAxisKeyEvent::ExpandNode(FKismetCompilerContext& CompilerConte
 	// Widget blueprints require the bAutomaticallyRegisterInputOnConstruction to be set to true in order to receive callbacks
 	if (GetBlueprint()->ShouldAutomaticallyRegisterInputOnConstruction())
 	{
-		CompilerContext.AddPostCDOCompiledStep([](const UObject::FPostCDOCompiledContext& Context, UObject* NewCDO)
-		{
-			if (UUserWidget* Widget = Cast<UUserWidget>(NewCDO))
-			{
-				Widget->bAutomaticallyRegisterInputOnConstruction = true;	
-			}
-		});	
+		CompilerContext.SetAutomaticallyRegisterInputEvents(true);
 	}
 }
 
diff --git a/Engine/Source/Editor/BlueprintGraph/Private/K2Node_InputKey.cpp b/Engine/Source/Editor/BlueprintGraph/Private/K2Node_InputKey.cpp
--- a/Engine/Source/Editor/BlueprintGraph/Private/K2Node_InputKey.cpp
+++ b/Engine/Source/Editor/BlueprintGraph/Private/K2Node_InputKey.cpp
@@ -367,13 +367,7 @@ void UK2Node_InputKey::ExpandNode(FKismetCompilerContext& CompilerContext, UEdGr
 	}
 
 	// Widget blueprints require the bAutomaticallyRegisterInputOnConstruction to be set to true in order to receive callbacks
-	CompilerContext.AddPostCDOCompiledStep([](const UObject::FPostCDOCompiledContext& Context, UObject* NewCDO)
-	{
-		if (UUserWidget* Widget = Cast<UUserWidget>(NewCDO))
-		{
-			Widget->bAutomaticallyRegisterInputOnConstruction = true;
-		}
-	});
+	CompilerContext.SetAutomaticallyRegisterInputEvents(true);
 }
 
 void UK2Node_InputKey::GetMenuActions(FBlueprintActionDatabaseRegistrar& ActionRegistrar) const
diff --git a/Engine/Source/Editor/BlueprintGraph/Private/K2Node_InputTouch.cpp b/Engine/Source/Editor/BlueprintGraph/Private/K2Node_InputTouch.cpp
--- a/Engine/Source/Editor/BlueprintGraph/Private/K2Node_InputTouch.cpp
+++ b/Engine/Source/Editor/BlueprintGraph/Private/K2Node_InputTouch.cpp
@@ -293,11 +293,5 @@ void UK2Node_InputTouch::ExpandNode(FKismetCompilerContext& CompilerContext, UEd
 	}
 
 	// Widget blueprints require the bAutomaticallyRegisterInputOnConstruction to be set to true in order to receive callbacks
-	CompilerContext.AddPostCDOCompiledStep([](const UObject::FPostCDOCompiledContext& Context, UObject* NewCDO)
-	{
-		if (UUserWidget* Widget = Cast<UUserWidget>(NewCDO))
-		{
-			Widget->bAutomaticallyRegisterInputOnConstruction = true;
-		}
-	});
+	CompilerContext.SetAutomaticallyRegisterInputEvents(true);
 }
diff --git a/Engine/Source/Editor/KismetCompiler/Private/KismetCompiler.cpp b/Engine/Source/Editor/KismetCompiler/Private/KismetCompiler.cpp
--- a/Engine/Source/Editor/KismetCompiler/Private/KismetCompiler.cpp
+++ b/Engine/Source/Editor/KismetCompiler/Private/KismetCompiler.cpp
@@ -553,6 +553,7 @@ FKismetCompilerContext::FKismetCompilerContext(UBlueprint* SourceSketch, FCompil
 	, OldLinker(nullptr)
 	, TargetClass(nullptr)
 	, bAssignDelegateSignatureFunction(false)
+	, bAutomaticallyRegisterInputEvents(false)
 {
 	MacroRowMaxHeight = 0;
 
@@ -5444,6 +5445,16 @@ void FKismetCompilerContext::AddPostCDOCompiledStep(TFunction<void(const UObject
 	PostCDOCompileSteps.Emplace(StepFunction);
 }
 
+void FKismetCompilerContext::SetAutomaticallyRegisterInputEvents(const bool bVal)
+{
+	bAutomaticallyRegisterInputEvents = bVal;
+}
+
+bool FKismetCompilerContext::GetAutomaticallyRegisterInputEvents() const
+{
+	return bAutomaticallyRegisterInputEvents;
+}
+
 void FKismetCompilerContext::Compile()
 {
 	CompileClassLayout(EInternalCompilerFlags::None);
diff --git a/Engine/Source/Editor/KismetCompiler/Public/KismetCompiler.h b/Engine/Source/Editor/KismetCompiler/Public/KismetCompiler.h
--- a/Engine/Source/Editor/KismetCompiler/Public/KismetCompiler.h
+++ b/Engine/Source/Editor/KismetCompiler/Public/KismetCompiler.h
@@ -185,6 +185,17 @@ public:
 	// CreateClassVariablesFromBlueprint:
 	bool bAssignDelegateSignatureFunction;
 
+private:
+	/**
+	* Set to true when there is an input event the graph and we want to automatically set up input component binding.
+	*/
+	bool bAutomaticallyRegisterInputEvents;
+
+public:
+
+	UE_API void SetAutomaticallyRegisterInputEvents(const bool bVal);
+	UE_API bool GetAutomaticallyRegisterInputEvents() const;
+	
 	struct FDelegateInfo
 	{
 		FName ProxyFunctionName;
diff --git a/Engine/Source/Editor/UMGEditor/Private/WidgetBlueprint.cpp b/Engine/Source/Editor/UMGEditor/Private/WidgetBlueprint.cpp
--- a/Engine/Source/Editor/UMGEditor/Private/WidgetBlueprint.cpp
+++ b/Engine/Source/Editor/UMGEditor/Private/WidgetBlueprint.cpp
@@ -1508,7 +1508,15 @@ static bool HasLatentActions(UEdGraph* Graph)
 	return false;
 }
 
+PRAGMA_DISABLE_DEPRECATION_WARNINGS
 void UWidgetBlueprint::UpdateTickabilityStats(bool& OutHasLatentActions, bool& OutHasAnimations, bool& OutClassRequiresNativeTick)
+{
+	bool bDummyFlag = false;
+	UpdateTickabilityStatsAndInputRegistration(OutHasLatentActions, OutHasAnimations, OutClassRequiresNativeTick, bDummyFlag);
+}
+PRAGMA_ENABLE_DEPRECATION_WARNINGS
+
+void UWidgetBlueprint::UpdateTickabilityStatsAndInputRegistration(bool& OutHasLatentActions, bool& OutHasAnimations, bool& OutClassRequiresNativeTick, bool& OutClassRequiresInputReg)
 {
 	if (GeneratedClass && GeneratedClass->ClassConstructor)
 	{
@@ -1521,12 +1529,20 @@ void UWidgetBlueprint::UpdateTickabilityStats(bool& OutHasLatentActions, bool& O
 		bool bHasLatentActions = false;
 		bool bHasAnimations = false;
 		const bool bHasScriptImplementedTick = DefaultWidget->bHasScriptImplementedTick;
+		bool bParentsRequireInputReg = false;
 
 		for (UBlueprint* Blueprint : BlueprintParents)
 		{
 			UWidgetBlueprint* WidgetBP = Cast<UWidgetBlueprint>(Blueprint);
 			if (WidgetBP)
 			{
+				// Check this blueprint parent for if it requires input registration
+				UWidgetBlueprintGeneratedClass* GenClass = CastChecked<UWidgetBlueprintGeneratedClass>(WidgetBP->GeneratedClass);
+				if (const UUserWidget* WidgetBPCDO = GenClass->GetDefaultObject<UUserWidget>())
+				{
+					bParentsRequireInputReg |= WidgetBPCDO->bAutomaticallyRegisterInputOnConstruction;
+				}                           		                   	
+            	
 				bHasAnimations |= WidgetBP->Animations.Num() > 0;
 
 				if (!bHasLatentActions)
@@ -1538,8 +1554,7 @@ void UWidgetBlueprint::UpdateTickabilityStats(bool& OutHasLatentActions, bool& O
 					{
 						if (HasLatentActions(Graph))
 						{
-							bHasLatentActions = true;
-							break;
+							bHasLatentActions = true;						
 						}
 					}
 				}
@@ -1607,6 +1622,7 @@ void UWidgetBlueprint::UpdateTickabilityStats(bool& OutHasLatentActions, bool& O
 		OutHasLatentActions = bHasLatentActions;
 		OutHasAnimations = bHasAnimations;
 		OutClassRequiresNativeTick = bClassRequiresNativeTick;
+		OutClassRequiresInputReg = bParentsRequireInputReg;
 	}
 }
 
diff --git a/Engine/Source/Editor/UMGEditor/Private/WidgetBlueprintCompiler.cpp b/Engine/Source/Editor/UMGEditor/Private/WidgetBlueprintCompiler.cpp
--- a/Engine/Source/Editor/UMGEditor/Private/WidgetBlueprintCompiler.cpp
+++ b/Engine/Source/Editor/UMGEditor/Private/WidgetBlueprintCompiler.cpp
@@ -654,6 +654,8 @@ void FWidgetBlueprintCompilerContext::CopyTermDefaultsToDefaultObject(UObject* D
 	UUserWidget* DefaultWidget = CastChecked<UUserWidget>(DefaultObject);
 	UWidgetBlueprintGeneratedClass* WidgetClass = CastChecked<UWidgetBlueprintGeneratedClass>(DefaultObject->GetClass());
 
+	const bool bRequiresInputEvents = GetAutomaticallyRegisterInputEvents();
+
 	if ( DefaultWidget )
 	{
 		//TODO Once we handle multiple derived blueprint classes, we need to check parent versions of the class.
@@ -676,20 +678,21 @@ void FWidgetBlueprintCompilerContext::CopyTermDefaultsToDefaultObject(UObject* D
 		{
 			DefaultWidget->bHasScriptImplementedPaint = false;
 		}
-
-		// Reset the value of this flag, which is set on PostCDOCompiled if there are any input nodes
-		// in the widget graphs.
-		DefaultWidget->bAutomaticallyRegisterInputOnConstruction = false;
+		
+		// This "bAutomaticallyRegisterInputEvents" is set to true on node expansion of an input event. 
+		DefaultWidget->bAutomaticallyRegisterInputOnConstruction = bRequiresInputEvents;
 	}
 
-
 	bool bClassOrParentsHaveLatentActions = false;
 	bool bClassOrParentsHaveAnimations = false;
 	bool bClassRequiresNativeTick = false;
+	bool bParentRequiresInputReg = false;
 
-
-	WidgetBP->UpdateTickabilityStats(bClassOrParentsHaveLatentActions, bClassOrParentsHaveAnimations, bClassRequiresNativeTick);
+	WidgetBP->UpdateTickabilityStatsAndInputRegistration(bClassOrParentsHaveLatentActions, bClassOrParentsHaveAnimations, bClassRequiresNativeTick, bParentRequiresInputReg);
 	WidgetClass->SetClassRequiresNativeTick(bClassRequiresNativeTick);
+	
+	// If either this BP or any of its parents require the input registration then we want to use it
+	DefaultWidget->bAutomaticallyRegisterInputOnConstruction = bRequiresInputEvents || bParentRequiresInputReg;	
 
 	// If the widget is not tickable, warn the user that widgets with animations or implemented ticks will most likely not work
 	if (DefaultWidget->GetDesiredTickFrequency() == EWidgetTickFrequency::Never)
diff --git a/Engine/Source/Editor/UMGEditor/Public/WidgetBlueprint.h b/Engine/Source/Editor/UMGEditor/Public/WidgetBlueprint.h
--- a/Engine/Source/Editor/UMGEditor/Public/WidgetBlueprint.h
+++ b/Engine/Source/Editor/UMGEditor/Public/WidgetBlueprint.h
@@ -320,7 +320,11 @@ public:
 	
 	static UE_API TSharedPtr<FKismetCompilerContext> GetCompilerForWidgetBP(UBlueprint* BP, FCompilerResultsLog& InMessageLog, const FKismetCompilerOptions& InCompileOptions);
 
+	UE_DEPRECATED(5.8, "Use UpdateTickabilityStatsAndInputRegistration instead.")
 	UE_API void UpdateTickabilityStats(bool& OutHasLatentActions, bool& OutHasAnimations, bool& OutClassRequiresNativeTick);
+	
+	/** Queries the parent widget chain to determine which features we can enable/disable based on what is needed. */
+	UE_API void UpdateTickabilityStatsAndInputRegistration(bool& OutHasLatentActions, bool& OutHasAnimations, bool& OutClassRequiresNativeTick, bool& OutClassRequiresInputReg);
 
 	UE_API bool ArePropertyBindingsAllowed() const;
 
diff --git a/Engine/Source/Runtime/UMG/Private/UserWidget.cpp b/Engine/Source/Runtime/UMG/Private/UserWidget.cpp
--- a/Engine/Source/Runtime/UMG/Private/UserWidget.cpp
+++ b/Engine/Source/Runtime/UMG/Private/UserWidget.cpp
@@ -1948,13 +1948,17 @@ void UUserWidget::CreateInputComponent()
 		return;
 	}
 	
-	// The widget tree may be constructed from instanced subobjects so we need to set this value
-	// based on this widget's actual CDO.
+	// Just an editor-time sanity check for this behavior.
+	// bAutomaticallyRegisterInputOnConstruction should be true if this widget or any of its
+	// BP parent classes have input nodes in their graphs which we want to set up
+	// automatic bindings for.
+#if WITH_EDITOR
 	if (!HasAnyFlags(RF_ClassDefaultObject))
 	{
 		const UUserWidget* DefaultWidget = CastChecked<UUserWidget>(GetClass()->GetDefaultObject());
-		bAutomaticallyRegisterInputOnConstruction = DefaultWidget->bAutomaticallyRegisterInputOnConstruction;
+		ensure(bAutomaticallyRegisterInputOnConstruction == DefaultWidget->bAutomaticallyRegisterInputOnConstruction);
 	}
+#endif	// #if WITH_EDITOR
 
 	// If bAutomaticallyRegisterInputOnConstruction is false, then there are no input delegates in this widget
 	// that need to be processed. There is no need to create an input component
