UMG: fix asset references being editor only if the widget is saved but not compiled. This caused cook issues (missing references) when changes were saved but not compiled
Update the widget tree of generated class on save.
#Synced-CL 50600560

#Synced-CL 50600560

diff --git a/Engine/Source/Editor/UMGEditor/Private/WidgetBlueprint.cpp b/Engine/Source/Editor/UMGEditor/Private/WidgetBlueprint.cpp
--- a/Engine/Source/Editor/UMGEditor/Private/WidgetBlueprint.cpp
+++ b/Engine/Source/Editor/UMGEditor/Private/WidgetBlueprint.cpp
@@ -708,6 +708,11 @@ void UWidgetBlueprint::ReplaceDeprecatedNodes()
 void UWidgetBlueprint::PreSave(FObjectPreSaveContext ObjectSaveContext)
 {
 	Super::PreSave(ObjectSaveContext);
+
+	if (UWidgetBlueprintGeneratedClass* BPGC = Cast< UWidgetBlueprintGeneratedClass>(GeneratedClass))
+	{
+		FWidgetBlueprintCompilerContext::UpdateGeneratedClassWidgetTree(this, BPGC, BPGC->GetWidgetTreeArchetype());
+	}
 }
 
 void UWidgetBlueprint::GetAssetRegistryTags(TArray<FAssetRegistryTag>& OutTags) const
diff --git a/Engine/Source/Editor/UMGEditor/Private/WidgetBlueprintCompiler.cpp b/Engine/Source/Editor/UMGEditor/Private/WidgetBlueprintCompiler.cpp
--- a/Engine/Source/Editor/UMGEditor/Private/WidgetBlueprintCompiler.cpp
+++ b/Engine/Source/Editor/UMGEditor/Private/WidgetBlueprintCompiler.cpp
@@ -857,6 +857,54 @@ void FWidgetBlueprintCompilerContext::FixAbandonedWidgetTree(UWidgetBlueprint* W
 	}
 }
 
+void FWidgetBlueprintCompilerContext::UpdateGeneratedClassWidgetTree(UWidgetBlueprint* WidgetBP, UWidgetBlueprintGeneratedClass* BPGClass, UWidgetTree* OldWidgetTree)
+{
+	FixAbandonedWidgetTree(WidgetBP);
+
+	{
+		TGuardValue<uint32> DisableInitializeFromWidgetTree(UUserWidget::GetInitializingFromWidgetTree(), 0);
+
+		// Need to clear archetype flag before duplication as we check during dup to see if we should postload
+		EObjectFlags PreviousFlags = WidgetBP->WidgetTree->GetFlags();
+		WidgetBP->WidgetTree->ClearFlags(RF_ArchetypeObject);
+
+		TMap<UObject*, UObject*> DupObjectsMap;
+		FObjectDuplicationParameters DupParams(WidgetBP->WidgetTree, BPGClass);
+		DupParams.DestName = DupParams.SourceObject->GetFName();
+		DupParams.FlagMask = RF_AllFlags & ~RF_DefaultSubObject;
+		DupParams.PortFlags |= PPF_DuplicateVerbatim; // Skip resetting text IDs
+
+		// if we are recompiling the BP on load, skip post load and defer it to the loading process
+		FUObjectSerializeContext* LinkerLoadingContext = nullptr;
+		if (WidgetBP->bIsRegeneratingOnLoad)
+		{
+			FLinkerLoad* Linker = WidgetBP->GetLinker();
+			LinkerLoadingContext = Linker ? FUObjectThreadContext::Get().GetSerializeContext() : nullptr;
+			DupParams.bSkipPostLoad = true;
+			DupParams.CreatedObjects = &DupObjectsMap;
+		}
+
+		UWidgetTree* NewWidgetTree = Cast<UWidgetTree>(StaticDuplicateObjectEx(DupParams));
+
+		// if we have anything in here after duplicate, then hook them in the loading process so they get post loaded
+		if (LinkerLoadingContext)
+		{
+			TArray<UObject*> DupObjects;
+			DupObjectsMap.GenerateValueArray(DupObjects);
+			LinkerLoadingContext->AddUniqueLoadedObjects(DupObjects);
+		}
+
+		BPGClass->SetWidgetTreeArchetype(NewWidgetTree);
+		if (OldWidgetTree)
+		{
+			FLinkerLoad::PRIVATE_PatchNewObjectIntoExport(OldWidgetTree, NewWidgetTree);
+		}
+		OldWidgetTree = nullptr;
+
+		WidgetBP->WidgetTree->SetFlags(PreviousFlags);
+	}
+}
+
 void FWidgetBlueprintCompilerContext::FinishCompilingClass(UClass* Class)
 {
 	if (Class == nullptr)
@@ -887,52 +935,7 @@ void FWidgetBlueprintCompilerContext::FinishCompilingClass(UClass* Class)
 			UBlueprint::ForceLoadMembers(WidgetBP->WidgetTree, WidgetBP);
 		}
 
-		FixAbandonedWidgetTree(WidgetBP);
-
-		{
-			TGuardValue<uint32> DisableInitializeFromWidgetTree(UUserWidget::GetInitializingFromWidgetTree(), 0);
-
-			// Need to clear archetype flag before duplication as we check during dup to see if we should postload
-			EObjectFlags PreviousFlags = WidgetBP->WidgetTree->GetFlags();
-			WidgetBP->WidgetTree->ClearFlags(RF_ArchetypeObject);
-
-			TMap<UObject*, UObject*> DupObjectsMap;
-			FObjectDuplicationParameters DupParams(WidgetBP->WidgetTree, BPGClass);
-			DupParams.DestName = DupParams.SourceObject->GetFName();
-			DupParams.FlagMask = RF_AllFlags & ~RF_DefaultSubObject;
-			DupParams.PortFlags |= PPF_DuplicateVerbatim; // Skip resetting text IDs
-
-			// if we are recompiling the BP on load, skip post load and defer it to the loading process
-			FUObjectSerializeContext* LinkerLoadingContext = nullptr;
-			if (WidgetBP->bIsRegeneratingOnLoad)
-			{
-				FLinkerLoad* Linker = WidgetBP->GetLinker();
-				LinkerLoadingContext = Linker ? FUObjectThreadContext::Get().GetSerializeContext() : nullptr;
-				DupParams.bSkipPostLoad = true;
-				DupParams.CreatedObjects = &DupObjectsMap;
-			}
-
-			UWidgetTree* NewWidgetTree = Cast<UWidgetTree>(StaticDuplicateObjectEx(DupParams));
-
-			// if we have anything in here after duplicate, then hook them in the loading process so they get post loaded
-			if (LinkerLoadingContext)
-			{
-				TArray<UObject*> DupObjects;
-				DupObjectsMap.GenerateValueArray(DupObjects);
-				LinkerLoadingContext->AddUniqueLoadedObjects(DupObjects);
-			}
-
-			//WidgetBP->IsWidgetFreeFromCircularReferences();
-
-			BPGClass->SetWidgetTreeArchetype(NewWidgetTree);
-			if (OldWidgetTree)
-			{
-				FLinkerLoad::PRIVATE_PatchNewObjectIntoExport(OldWidgetTree, NewWidgetTree);
-			}
-			OldWidgetTree = nullptr;
-
-			WidgetBP->WidgetTree->SetFlags(PreviousFlags);
-		}
+		UpdateGeneratedClassWidgetTree(WidgetBP,  BPGClass, OldWidgetTree);
 
 		{
 			TValueOrError<void, UWidget*> HasReference = WidgetBP->HasCircularReferences();
diff --git a/Engine/Source/Editor/UMGEditor/Public/WidgetBlueprintCompiler.h b/Engine/Source/Editor/UMGEditor/Public/WidgetBlueprintCompiler.h
--- a/Engine/Source/Editor/UMGEditor/Public/WidgetBlueprintCompiler.h
+++ b/Engine/Source/Editor/UMGEditor/Public/WidgetBlueprintCompiler.h
@@ -143,8 +143,11 @@ public:
 		FWidgetBlueprintCompilerContext& Context;
 	};
 
+	/** Duplicate the widget tree of the passed widget blueprint and assign it as the widget tree of the passed generated class. */
+	UE_API void static UpdateGeneratedClassWidgetTree(UWidgetBlueprint* WidgetBP, UWidgetBlueprintGeneratedClass* BPGClass, UWidgetTree* OldWidgetTree);
+
 protected:
-	UE_API void FixAbandonedWidgetTree(UWidgetBlueprint* WidgetBP);
+	UE_API void static FixAbandonedWidgetTree(UWidgetBlueprint* WidgetBP);
 
 	UWidgetBlueprintGeneratedClass* NewWidgetBlueprintClass;
 
