UMG: refactor generated class widget tree clean up to use in both pre-save and compile
follow up on changes done in 50789578 to update widget trees on save
#Synced-CL 50600560

diff --git a/Engine/Source/Editor/UMGEditor/Private/WidgetBlueprint.cpp b/Engine/Source/Editor/UMGEditor/Private/WidgetBlueprint.cpp
--- a/Engine/Source/Editor/UMGEditor/Private/WidgetBlueprint.cpp
+++ b/Engine/Source/Editor/UMGEditor/Private/WidgetBlueprint.cpp
@@ -711,7 +711,8 @@ void UWidgetBlueprint::PreSave(FObjectPreSaveContext ObjectSaveContext)
 
 	if (UWidgetBlueprintGeneratedClass* BPGC = Cast< UWidgetBlueprintGeneratedClass>(GeneratedClass))
 	{
-		FWidgetBlueprintCompilerContext::UpdateGeneratedClassWidgetTree(this, BPGC, BPGC->GetWidgetTreeArchetype());
+		FWidgetBlueprintCompilerContext::ClearGeneratedClassWidgetTree(BPGC, false);
+		FWidgetBlueprintCompilerContext::UpdateGeneratedClassWidgetTree(this, BPGC, BPGC->GetWidgetTreeArchetype(), nullptr);
 	}
 }
 
diff --git a/Engine/Source/Editor/UMGEditor/Private/WidgetBlueprintCompiler.cpp b/Engine/Source/Editor/UMGEditor/Private/WidgetBlueprintCompiler.cpp
--- a/Engine/Source/Editor/UMGEditor/Private/WidgetBlueprintCompiler.cpp
+++ b/Engine/Source/Editor/UMGEditor/Private/WidgetBlueprintCompiler.cpp
@@ -350,23 +350,7 @@ void FWidgetBlueprintCompilerContext::CleanAndSanitizeClass(UBlueprintGeneratedC
 	{
 		if (UWidgetBlueprintGeneratedClass* WBC_ToClean = Cast<UWidgetBlueprintGeneratedClass>(ClassToClean))
 		{
-			if (UWidgetTree* OldArchetype = WBC_ToClean->GetWidgetTreeArchetype())
-			{
-				FString TransientArchetypeString = FString::Printf(TEXT("OLD_TEMPLATE_TREE%s"), *OldArchetype->GetName());
-				RenameObjectToTransientPackage(OldArchetype, *TransientArchetypeString, true);
-
-				TArray<UObject*> Children;
-				ForEachObjectWithOuter(OldArchetype, [&Children] (UObject* Child) {
-					Children.Add(Child);
-				}, false);
-
-				for ( UObject* Child : Children )
-				{
-					RenameObjectToTransientPackage(Child, FName(), false);
-				}
-
-				WBC_ToClean->SetWidgetTreeArchetype(nullptr);
-			}
+			ClearGeneratedClassWidgetTree(WBC_ToClean, bRecompilingOnLoad);
 		}
 	}
 
@@ -857,7 +841,7 @@ void FWidgetBlueprintCompilerContext::FixAbandonedWidgetTree(UWidgetBlueprint* W
 	}
 }
 
-void FWidgetBlueprintCompilerContext::UpdateGeneratedClassWidgetTree(UWidgetBlueprint* WidgetBP, UWidgetBlueprintGeneratedClass* BPGClass, UWidgetTree* OldWidgetTree)
+void FWidgetBlueprintCompilerContext::UpdateGeneratedClassWidgetTree(UWidgetBlueprint* WidgetBP, UWidgetBlueprintGeneratedClass* BPGClass, UWidgetTree* OldWidgetTree, FCompilerResultsLog* MessageLog)
 {
 	FixAbandonedWidgetTree(WidgetBP);
 
@@ -903,6 +887,103 @@ void FWidgetBlueprintCompilerContext::UpdateGeneratedClassWidgetTree(UWidgetBlue
 
 		WidgetBP->WidgetTree->SetFlags(PreviousFlags);
 	}
+
+	{
+		TValueOrError<void, UWidget*> HasReference = WidgetBP->HasCircularReferences();
+		if (HasReference.HasError())
+		{
+			if (UWidget* FoundCircularWidget = BPGClass->GetWidgetTreeArchetype()->FindWidget(HasReference.GetError()->GetFName()))
+			{
+				BPGClass->GetWidgetTreeArchetype()->RemoveWidget(FoundCircularWidget);
+			}
+			if (MessageLog)
+			{
+				MessageLog->Error(*FText::Format(LOCTEXT("WidgetTreeCircularReference", "The WidgetTree '{0}' Contains circular references. See widget '{1}'"),
+					FText::FromString(WidgetBP->WidgetTree->GetPathName()),
+					FText::FromString(HasReference.GetError()->GetName())
+				).ToString());
+			}
+		}
+	}
+
+	{
+#if WITH_EDITOR
+		BPGClass->NameClashingInHierarchy.Reset();
+#endif
+		TValueOrError<void, TSet<UWidget*>> HasConflictingWidgetNames = WidgetBP->HasConflictingWidgetNamesFromInheritance();
+		if (HasConflictingWidgetNames.HasError())
+		{
+			TSet<UWidget*>& ConflictingNames = HasConflictingWidgetNames.GetError();
+			for (UWidget* ConflictingWidget : ConflictingNames)
+			{
+#if WITH_EDITOR
+				FName ConflictingWidgetName = ConflictingWidget->GetFName();
+				BPGClass->NameClashingInHierarchy.Add(ConflictingWidgetName);
+#endif
+				if (UWidget* FoundConflictingWidget = BPGClass->GetWidgetTreeArchetype()->FindWidget(ConflictingWidget->GetFName()))
+				{
+					BPGClass->GetWidgetTreeArchetype()->RemoveWidget(FoundConflictingWidget);
+				}
+
+				if (MessageLog)
+				{
+					MessageLog->Error(*FText::Format(LOCTEXT("WidgetTreeDuplicateNames", "The WidgetTree '{0}' already contains a widget named '{1}'."),
+						FText::FromString(WidgetBP->WidgetTree->GetPathName()),
+						FText::FromString(ConflictingWidget->GetName())
+					).ToString());
+				}
+			}
+		}
+	}
+}
+
+void FWidgetBlueprintCompilerContext::ClearGeneratedClassWidgetTree(UWidgetBlueprintGeneratedClass* WBC_ToClean, bool bRecompilingOnLoad)
+{
+	auto RenameObjectToTransientPackage = [bRecompilingOnLoad](UObject* ObjectToRename, const FName BaseName, bool bClearFlags)
+		{
+			ObjectToRename->SetFlags(RF_Transient);
+
+			if (bClearFlags)
+			{
+				ObjectToRename->ClearFlags(RF_Public | RF_Standalone | RF_ArchetypeObject);
+			}
+
+			// Rename will remove the renamed object's linker when moving to a new package so invalidate the export beforehand
+			FLinkerLoad::InvalidateExport(ObjectToRename);
+
+			const ERenameFlags RenFlags = REN_DontCreateRedirectors | REN_NonTransactional | REN_DoNotDirty;
+
+			if (BaseName.IsNone())
+			{
+				ObjectToRename->Rename(nullptr, GetTransientPackage(), RenFlags);
+			}
+			else
+			{
+				FName TransientArchetypeName = MakeUniqueObjectName(GetTransientPackage(), ObjectToRename->GetClass(), BaseName);
+				ObjectToRename->Rename(*TransientArchetypeName.ToString(), GetTransientPackage(), RenFlags);
+			}
+		};
+
+	if (WBC_ToClean)
+	{
+		if (UWidgetTree* OldArchetype = WBC_ToClean->GetWidgetTreeArchetype())
+		{
+			FString TransientArchetypeString = FString::Printf(TEXT("OLD_TEMPLATE_TREE%s"), *OldArchetype->GetName());
+			RenameObjectToTransientPackage(OldArchetype, *TransientArchetypeString, true);
+
+			TArray<UObject*> Children;
+			ForEachObjectWithOuter(OldArchetype, [&Children](UObject* Child) {
+				Children.Add(Child);
+				}, false);
+
+			for (UObject* Child : Children)
+			{
+				RenameObjectToTransientPackage(Child, FName(), false);
+			}
+
+			WBC_ToClean->SetWidgetTreeArchetype(nullptr);
+		}
+	}
 }
 
 void FWidgetBlueprintCompilerContext::FinishCompilingClass(UClass* Class)
@@ -935,50 +1016,7 @@ void FWidgetBlueprintCompilerContext::FinishCompilingClass(UClass* Class)
 			UBlueprint::ForceLoadMembers(WidgetBP->WidgetTree, WidgetBP);
 		}
 
-		UpdateGeneratedClassWidgetTree(WidgetBP,  BPGClass, OldWidgetTree);
-
-		{
-			TValueOrError<void, UWidget*> HasReference = WidgetBP->HasCircularReferences();
-			if (HasReference.HasError())
-			{
-				if (UWidget* FoundCircularWidget = BPGClass->GetWidgetTreeArchetype()->FindWidget(HasReference.GetError()->GetFName()))
-				{
-					BPGClass->GetWidgetTreeArchetype()->RemoveWidget(FoundCircularWidget);
-				}
-				MessageLog.Error(*FText::Format(LOCTEXT("WidgetTreeCircularReference", "The WidgetTree '{0}' Contains circular references. See widget '{1}'"),
-					FText::FromString(WidgetBP->WidgetTree->GetPathName()),
-					FText::FromString(HasReference.GetError()->GetName())
-					).ToString());
-			}
-		}
-
-		{
-#if WITH_EDITOR
-			BPGClass->NameClashingInHierarchy.Reset();
-#endif
-			TValueOrError<void, TSet<UWidget*>> HasConflictingWidgetNames = WidgetBP->HasConflictingWidgetNamesFromInheritance();
-			if (HasConflictingWidgetNames.HasError())
-			{
-				TSet<UWidget*>& ConflictingNames = HasConflictingWidgetNames.GetError();
-				for (UWidget* ConflictingWidget : ConflictingNames)
-				{
-#if WITH_EDITOR
-					FName ConflictingWidgetName = ConflictingWidget->GetFName();
-					BPGClass->NameClashingInHierarchy.Add(ConflictingWidgetName);
-#endif
-					if (UWidget* FoundConflictingWidget = BPGClass->GetWidgetTreeArchetype()->FindWidget(ConflictingWidget->GetFName()))
-					{
-						BPGClass->GetWidgetTreeArchetype()->RemoveWidget(FoundConflictingWidget);
-					}
-
-					MessageLog.Error(*FText::Format(LOCTEXT("WidgetTreeDuplicateNames", "The WidgetTree '{0}' already contains a widget named '{1}'."),
-						FText::FromString(WidgetBP->WidgetTree->GetPathName()),
-						FText::FromString(ConflictingWidget->GetName())
-					).ToString());
-				}
-			}
-		}
-
+		UpdateGeneratedClassWidgetTree(WidgetBP,  BPGClass, OldWidgetTree, &MessageLog);
 
 
 		int32 AnimIndex = 0;
diff --git a/Engine/Source/Editor/UMGEditor/Public/WidgetBlueprintCompiler.h b/Engine/Source/Editor/UMGEditor/Public/WidgetBlueprintCompiler.h
--- a/Engine/Source/Editor/UMGEditor/Public/WidgetBlueprintCompiler.h
+++ b/Engine/Source/Editor/UMGEditor/Public/WidgetBlueprintCompiler.h
@@ -144,7 +144,8 @@ public:
 	};
 
 	/** Duplicate the widget tree of the passed widget blueprint and assign it as the widget tree of the passed generated class. */
-	UE_API void static UpdateGeneratedClassWidgetTree(UWidgetBlueprint* WidgetBP, UWidgetBlueprintGeneratedClass* BPGClass, UWidgetTree* OldWidgetTree);
+	UE_API void static UpdateGeneratedClassWidgetTree(UWidgetBlueprint* WidgetBP, UWidgetBlueprintGeneratedClass* BPGClass, UWidgetTree* OldWidgetTree, FCompilerResultsLog* MessageLog);
+	UE_API void static ClearGeneratedClassWidgetTree(UWidgetBlueprintGeneratedClass* WBC_ToClean, bool bRecompilingOnLoad);
 
 protected:
 	UE_API void static FixAbandonedWidgetTree(UWidgetBlueprint* WidgetBP);
