UMG - Basics test added

#rnx

#Synced-CL 50845945

diff --git a/Engine/Source/Editor/UMGEditor/Private/Details/SWidgetDetailsView.cpp b/Engine/Source/Editor/UMGEditor/Private/Details/SWidgetDetailsView.cpp
--- a/Engine/Source/Editor/UMGEditor/Private/Details/SWidgetDetailsView.cpp
+++ b/Engine/Source/Editor/UMGEditor/Private/Details/SWidgetDetailsView.cpp
@@ -77,6 +77,7 @@ void SWidgetDetailsView::Construct(const FArguments& InArgs, TSharedPtr<FWidgetB
 		.OnCheckStateChanged(this, &SWidgetDetailsView::HandleIsVariableChanged)
 		.IsEnabled(this, &SWidgetDetailsView::IsDetailsPanelEditingAllowed)
 		.Padding(FMargin(3.0f, 1.0f, 18.0f, 1.0f))
+		.AddMetaData<FTagMetaData>(TEXT("Details.IsVariable"))
 		[
 			SNew(STextBlock)
 			.Text(LOCTEXT("IsVariable", "Is Variable"))
@@ -164,6 +165,7 @@ void SWidgetDetailsView::Construct(const FArguments& InArgs, TSharedPtr<FWidgetB
 							.IsEnabled(this, &SWidgetDetailsView::IsWidgetNameFieldEnabled)
 							.OnTextChanged(this, &SWidgetDetailsView::HandleNameTextChanged)
 							.OnTextCommitted(this, &SWidgetDetailsView::HandleNameTextCommitted)
+							.AddMetaData<FTagMetaData>(TEXT("Details.WidgetName"))
 						]
 					]
 
