Fix for crash due to Error Reporting adding or removing a widget during Prepass.
`HandleNameTextChanged` can be triggered during attribute processing which occurs in Slate Prepass. During Prepass, we cannot add or remove widgets and calling `SetError` here can end up either creating or dismissing a popup if the name of the BP variable already exists, which would trigger a crash in `FChildren::ForEachWidget`. To avoid this, we delay the error reporting until it is safe to do so.

#Synced-CL 51454850

diff --git a/Engine/Source/Editor/UMGEditor/Private/Details/SWidgetDetailsView.cpp b/Engine/Source/Editor/UMGEditor/Private/Details/SWidgetDetailsView.cpp
--- a/Engine/Source/Editor/UMGEditor/Private/Details/SWidgetDetailsView.cpp
+++ b/Engine/Source/Editor/UMGEditor/Private/Details/SWidgetDetailsView.cpp
@@ -536,14 +536,24 @@ bool SWidgetDetailsView::IsSingleObjectSelected() const
 void SWidgetDetailsView::HandleNameTextChanged(const FText& Text)
 {
 	FText OutErrorMessage;
-	if ( !HandleVerifyNameTextChanged(Text, OutErrorMessage) )
+	if ( HandleVerifyNameTextChanged(Text, OutErrorMessage) )
 	{
-		NameTextBox->SetError(OutErrorMessage);
-	}
-	else
-	{
-		NameTextBox->SetError(FText::GetEmpty());
+		OutErrorMessage = FText::GetEmpty();
 	}
+
+	// We delay the error message to be sure we are outside of attribute processing.
+	// Otherwise, we might add or remove a widget (the error reporting one) which is not allowed during Prepass
+	TWeakPtr<SEditableTextBox> WeakNameTextBox = NameTextBox;
+	RegisterActiveTimer(0.f, FWidgetActiveTimerDelegate::CreateLambda(
+		[OutErrorMessage = MoveTemp(OutErrorMessage), WeakNameTextBox = MoveTemp(WeakNameTextBox)](double InCurrentTime, float InDeltaTime)
+		{
+			if (TSharedPtr<SEditableTextBox> NameTextBox = WeakNameTextBox.Pin())
+			{
+				NameTextBox->SetError(OutErrorMessage);
+			}
+			return EActiveTimerReturnType::Stop;
+		})
+	);
 }
 
 bool SWidgetDetailsView::HandleVerifyNameTextChanged(const FText& InText, FText& OutErrorMessage)
