Fix to DesiredFocus and Bindings to support selecting a Widget in a Inherited NamedSlot

#Synced-CL 52059832

diff --git a/Engine/Source/Editor/UMGEditor/Private/Hierarchy/SReadOnlyHierarchyView.cpp b/Engine/Source/Editor/UMGEditor/Private/Hierarchy/SReadOnlyHierarchyView.cpp
--- a/Engine/Source/Editor/UMGEditor/Private/Hierarchy/SReadOnlyHierarchyView.cpp
+++ b/Engine/Source/Editor/UMGEditor/Private/Hierarchy/SReadOnlyHierarchyView.cpp
@@ -84,6 +84,12 @@ FName SReadOnlyHierarchyView::GetItemName(const TSharedPtr<FItem>& Item) const
 		return Widget->GetFName();
 	}
 
+	const FName NamedSlotName = Item->InheritedNamedSlotName;
+	if (!NamedSlotName.IsNone())
+	{
+		return NamedSlotName;
+	}
+
 	return FName();
 }
 
@@ -139,7 +145,11 @@ FText SReadOnlyHierarchyView::GetItemText(TSharedPtr<FItem> Item) const
 
 			static const FText RootWidgetFormat = LOCTEXT("WidgetNameFormat", "[{0}]");
 			return FText::Format(RootWidgetFormat, RootWidgetName);
-
+		}
+		const FName NamedSlotName = Item->InheritedNamedSlotName;
+		if (!NamedSlotName.IsNone())
+		{
+			return FText::FromName(NamedSlotName);
 		}
 	}
 	return FText::GetEmpty();
@@ -199,6 +209,13 @@ TSharedRef<ITableRow> SReadOnlyHierarchyView::GenerateRow(TSharedPtr<FItem> Item
 			bIsEnabled = false;
 		}
 	}
+
+	const FName NamedSlotName = Item->InheritedNamedSlotName;
+	if (!NamedSlotName.IsNone())
+	{
+		bIsEnabled = false;
+	}
+
 	return SNew(STableRow<TSharedPtr<FItem>>, OwnerTable)
 		   .IsEnabled(bIsEnabled)
 		[
@@ -236,9 +253,17 @@ void SReadOnlyHierarchyView::GetFilterStringsForItem(TSharedPtr<FItem> Item, TAr
 		OutStrings.Add(Widget->GetName());
 		OutStrings.Add(Widget->GetLabelTextWithMetadata().ToString());
 	}
+	else if (const UWidgetBlueprint* Blueprint = Item->WidgetBlueprint.Get())
+	{
+		OutStrings.Add(Blueprint->GetName());
+	}
 	else
 	{
-		OutStrings.Add(WidgetBlueprint->GetName());
+		const FName NamedSlotName = Item->InheritedNamedSlotName;
+		if (!NamedSlotName.IsNone())
+		{
+			OutStrings.Add(NamedSlotName.ToString());
+		}
 	}
 }
 
@@ -309,7 +334,7 @@ void SReadOnlyHierarchyView::BuildWidgetChildren(const UWidget* Widget, TSharedP
 
 		for (FName SlotName : SlotNames)
 		{
-			if (UWidget* SlotContent = NamedSlotHost->GetContentForSlot(SlotName))
+			if (const UWidget* SlotContent = NamedSlotHost->GetContentForSlot(SlotName))
 			{
 				BuildWidgetChildren(SlotContent, Parent);
 			}
@@ -348,6 +373,26 @@ void SReadOnlyHierarchyView::RebuildTree()
 		BuildWidgetChildren(RootWidget, RootItem);
 	}
 
+	// Also show content placed in named slots inherited from a parent blueprint,
+	// regardless of whether this blueprint has its own root widget.
+	for (const FName& SlotName : WidgetBP->GetInheritedAvailableNamedSlots())
+	{
+		TSharedRef<FItem> NamedSlotItem = MakeShared<FItem>(SlotName);
+		if (RootItem.IsValid())
+		{
+			RootItem->Children.Add(NamedSlotItem);
+		}
+		else
+		{
+			RootWidgets.Add(NamedSlotItem);
+		}
+
+		if (const UWidget* SlotContent = WidgetBP->WidgetTree->GetContentForSlot(SlotName))
+		{
+			BuildWidgetChildren(SlotContent, NamedSlotItem);
+		}
+	}
+
 	if (bExpandAll)
 	{
 		ExpandAll();
@@ -385,6 +430,14 @@ TSharedPtr<SReadOnlyHierarchyView::FItem> SReadOnlyHierarchyView::FindItem(const
 				return Item;
 			}
 		}
+		else 
+		{
+			const FName NamedSlotName = Item->InheritedNamedSlotName;
+			if (!NamedSlotName.IsNone() && NamedSlotName == Name)
+			{
+				return Item;
+			}
+		}
 
 		GetItemChildren(Item, Items);
 	}
diff --git a/Engine/Source/Editor/UMGEditor/Private/WidgetBlueprintCompiler.cpp b/Engine/Source/Editor/UMGEditor/Private/WidgetBlueprintCompiler.cpp
--- a/Engine/Source/Editor/UMGEditor/Private/WidgetBlueprintCompiler.cpp
+++ b/Engine/Source/Editor/UMGEditor/Private/WidgetBlueprintCompiler.cpp
@@ -1454,11 +1454,26 @@ void FWidgetBlueprintCompilerContext::ValidateDesiredFocusWidgetName()
 		UUserWidget* UserWidgetCDO = WidgetClass->GetDefaultObject<UUserWidget>();
 		if (WidgetBP && UserWidgetCDO)
 		{
-			UWidgetTree* LatestWidgetTree = FWidgetBlueprintEditorUtils::FindLatestWidgetTree(WidgetBP, UserWidgetCDO);
 			FName DesiredFocusWidgetName = UserWidgetCDO->GetDesiredFocusWidgetName();
 
-			if (!DesiredFocusWidgetName.IsNone() && !LatestWidgetTree->FindWidget(DesiredFocusWidgetName))
+			if (!DesiredFocusWidgetName.IsNone())
 			{
+				// Check the current tree first as it could includes named slot content
+				if (WidgetBP->WidgetTree->FindWidget(DesiredFocusWidgetName))
+				{
+					return;
+				}
+
+				// If there is a root widget, also check the latest merged tree (which may include parent-inherited widgets)
+				if (WidgetBP->WidgetTree->RootWidget)
+				{
+					const UWidgetTree* LatestWidgetTree = FWidgetBlueprintEditorUtils::FindLatestWidgetTree(WidgetBP, UserWidgetCDO);
+					if (LatestWidgetTree->FindWidget(DesiredFocusWidgetName))
+					{
+						return;
+					}
+				}
+
 				FBlueprintCompilerLog BlueprintLog(MessageLog, WidgetClass);
 				// Notify that the desired focus widget is not found in the Widget tree, so it's invalid.
 				const FText InvalidDesiredFocusWidgetNameMessage = LOCTEXT("InvalidDesiredFocusWidgetName", "User Widget '{0}' Desired Focus is set to a non-existent widget '{1}'. Select a valid desired focus for this User Widget.");
diff --git a/Engine/Source/Editor/UMGEditor/Public/Hierarchy/SReadOnlyHierarchyView.h b/Engine/Source/Editor/UMGEditor/Public/Hierarchy/SReadOnlyHierarchyView.h
--- a/Engine/Source/Editor/UMGEditor/Public/Hierarchy/SReadOnlyHierarchyView.h
+++ b/Engine/Source/Editor/UMGEditor/Public/Hierarchy/SReadOnlyHierarchyView.h
@@ -43,9 +43,11 @@ private:
 	{
 		FItem(const UWidget* InWidget) : Widget(InWidget) {}
 		FItem(const UWidgetBlueprint* InWidgetBlueprint) : WidgetBlueprint(InWidgetBlueprint) {}
+		FItem(const FName InInheritedNamedSlotName) : InheritedNamedSlotName(InInheritedNamedSlotName) {}
 
 		TWeakObjectPtr<const UWidget> Widget;
 		TWeakObjectPtr<const UWidgetBlueprint> WidgetBlueprint;
+		const FName InheritedNamedSlotName;
 
 		TArray<TSharedPtr<FItem>> Children;
 	};
