[UMG] GetDesignTimeInfo follow-up to review feedback + fix ensure in SDesignerView orientation toggle

- Move the design-time validation check outside of WITH_EDITOR so
GetDesignTimeInfo() logs an error in all Development builds, not just
editor. GConfig reads remain editor-only.
- Update header comment to reflect that error logging now applies in
all builds, not just editor.
- Fix ensure(WidgetTree) failure in SDesignerView when toggling
orientation or changing resolution: remove immediate
BroadcastDesignerChanged() after RefreshPreview(), as the newly
created preview's child widget trees are not yet initialized.
UpdatePreviewWidget on the next Tick calls BroadcastDesignerChanged
after the preview is fully set up ? matching the existing pattern
used by blueprint compilation and other preview rebuild paths.

#rnx

#Synced-CL 52147079

diff --git a/Engine/Source/Editor/UMGEditor/Private/Designer/SDesignerView.cpp b/Engine/Source/Editor/UMGEditor/Private/Designer/SDesignerView.cpp
--- a/Engine/Source/Editor/UMGEditor/Private/Designer/SDesignerView.cpp
+++ b/Engine/Source/Editor/UMGEditor/Private/Designer/SDesignerView.cpp
@@ -3563,16 +3563,14 @@ void SDesignerView::HandleOnCommonResolutionSelected(const FPlayScreenResolution
 		MarkDesignModifed(/*bRequiresRecompile*/ false);
 	}
 
-	// Kick off a preview rebuild so PreConstruct fires with the new preview dimensions.
-	// Note: RefreshPreview() destroys the current preview synchronously; the new preview
-	// is not available until UpdatePreviewWidget() runs on the next Tick.
+	// Recreate the preview widget so PreConstruct fires with the new preview dimensions.
+	// UpdatePreviewWidget (on next Tick) will call BroadcastDesignerChanged once the
+	// new preview and its child widget trees are fully initialized.
 	if (TSharedPtr<FWidgetBlueprintEditor> Editor = BlueprintEditor.Pin())
 	{
 		Editor->RefreshPreview();
 	}
 
-	BroadcastDesignerChanged();
-
 	ResolutionTextFade.Play(this->AsShared());
 }
 
@@ -3862,16 +3860,14 @@ FReply SDesignerView::HandleSwapAspectRatioClicked()
 		MarkDesignModifed(/*bRequiresRecompile*/ false);
 	}
 
-	// Kick off a preview rebuild so PreConstruct fires with the new preview dimensions.
-	// Note: RefreshPreview() destroys the current preview synchronously; the new preview
-	// is not available until UpdatePreviewWidget() runs on the next Tick.
+	// Recreate the preview widget so PreConstruct fires with the new preview dimensions.
+	// UpdatePreviewWidget (on next Tick) will call BroadcastDesignerChanged once the
+	// new preview and its child widget trees are fully initialized.
 	if (TSharedPtr<FWidgetBlueprintEditor> Editor = BlueprintEditor.Pin())
 	{
 		Editor->RefreshPreview();
 	}
 
-	BroadcastDesignerChanged();
-
 	ResolutionTextFade.Play(this->AsShared());
 
 	return FReply::Handled();
diff --git a/Engine/Source/Runtime/UMG/Public/Blueprint/WidgetLayoutLibrary.h b/Engine/Source/Runtime/UMG/Public/Blueprint/WidgetLayoutLibrary.h
--- a/Engine/Source/Runtime/UMG/Public/Blueprint/WidgetLayoutLibrary.h
+++ b/Engine/Source/Runtime/UMG/Public/Blueprint/WidgetLayoutLibrary.h
@@ -13,8 +13,8 @@
  * Contains the current designer preview dimensions and derived orientation.
  * Only populated at design time — all values are zeroed/false at runtime.
  */
-USTRUCT(BlueprintType)
-struct UMG_API FWidgetDesignTimeInfo
+USTRUCT(BlueprintType, MinimalAPI)
+struct FWidgetDesignTimeInfo
 {
 	GENERATED_BODY()
 
@@ -237,9 +237,8 @@ public:
 	 * Use this in PreConstruct to adapt widget layout to the designer's current preview
 	 * dimensions (e.g., switching between portrait and landscape layouts).
 	 *
-	 * In editor builds, logs an error and returns zeroed values if called outside of a
-	 * design-time preview (e.g., at runtime or from a non-preview context). In non-editor
-	 * builds, returns zeroed values silently. Use viewport/display system
+	 * Logs an error and returns zeroed values if called outside of a design-time preview
+	 * (e.g., at runtime or from a non-preview context). Use viewport/display system
 	 * queries instead for runtime orientation.
 	 *
 	 * @param Widget	The calling widget (auto-filled via DefaultToSelf).
