Removed duplicated hanging if statement in widget blueprint compiler animation handling.

PR #14575

#tests created an animation in a widget blueprint and compiled it.
#Synced-CL 52063442

diff --git a/Engine/Source/Editor/UMGEditor/Private/WidgetBlueprintCompiler.cpp b/Engine/Source/Editor/UMGEditor/Private/WidgetBlueprintCompiler.cpp
--- a/Engine/Source/Editor/UMGEditor/Private/WidgetBlueprintCompiler.cpp
+++ b/Engine/Source/Editor/UMGEditor/Private/WidgetBlueprintCompiler.cpp
@@ -1026,8 +1026,6 @@ void FWidgetBlueprintCompilerContext::FinishCompilingClass(UClass* Class)
 			//ClonedAnimation->SetFlags(RF_Public); // Needs to be marked public so that it can be referenced from widget instances.
 			
 			if (OldWidgetAnimations.IsValidIndex(AnimIndex) && OldWidgetAnimations[AnimIndex])
-
-			if ((AnimIndex < OldWidgetAnimations.Num()) && OldWidgetAnimations[AnimIndex])
 			{
 				FLinkerLoad::PRIVATE_PatchNewObjectIntoExport(OldWidgetAnimations[AnimIndex], ClonedAnimation);
 			}
