Sync the flow direction in the UMG preview from the game preview language (if any)

#Synced-CL 52194623

diff --git a/Engine/Source/Editor/UMGEditor/Private/Designer/SDesignerView.cpp b/Engine/Source/Editor/UMGEditor/Private/Designer/SDesignerView.cpp
--- a/Engine/Source/Editor/UMGEditor/Private/Designer/SDesignerView.cpp
+++ b/Engine/Source/Editor/UMGEditor/Private/Designer/SDesignerView.cpp
@@ -70,6 +70,7 @@
 
 #include "Math/TransformCalculus2D.h"
 #include "Input/HittestGrid.h"
+#include "Internationalization/Culture.h"
 
 #include "Fonts/FontMeasure.h"
 #include "WidgetEditingProjectSettings.h"
@@ -2395,6 +2396,19 @@ void SDesignerView::Tick( const FGeometry& AllottedGeometry, const double InCurr
 	const bool bForceUpdate = false;
 	UpdatePreviewWidget(bForceUpdate);
 
+	// Update the preview flow direction based on the active game preview language (if any)
+	{
+		TOptional<EFlowDirectionPreference> FlowDirectionOverride;
+		if (FTextLocalizationManager::Get().IsGameLocalizationPreviewEnabled())
+		{
+			if (FCulturePtr PreviewLanguage = FInternationalization::Get().GetCulture(FTextLocalizationManager::Get().GetConfiguredGameLocalizationPreviewLanguage()))
+			{
+				FlowDirectionOverride = PreviewLanguage->IsRightToLeft() ? EFlowDirectionPreference::RightToLeft : EFlowDirectionPreference::LeftToRight;
+			}
+		}
+		PreviewSizeConstraint->SetFlowDirectionPreference(FlowDirectionOverride.Get(EFlowDirectionPreference::Inherit));
+	}
+
 	// Perform an arrange children pass to cache the geometry of all widgets so that we can query it later.
 	{
 		CachedWidgetGeometry.Reset();
