PR #14217: Fixed widgets not being renamed if new name only differs by case
#Synced-CL 52272722

#Synced-CL 52272722

diff --git a/Engine/Source/Editor/UMGEditor/Private/Details/SWidgetDetailsView.cpp b/Engine/Source/Editor/UMGEditor/Private/Details/SWidgetDetailsView.cpp
--- a/Engine/Source/Editor/UMGEditor/Private/Details/SWidgetDetailsView.cpp
+++ b/Engine/Source/Editor/UMGEditor/Private/Details/SWidgetDetailsView.cpp
@@ -587,7 +587,7 @@ void SWidgetDetailsView::HandleNameTextCommitted(const FText& Text, ETextCommit:
 			if ( HandleVerifyNameTextChanged(Text, DummyText) )
 			{
 				UWidget* Widget = Cast<UWidget>(SelectedObjects[0].Get());
-				if (Widget && !Widget->GetLabelText().EqualToCaseIgnored(Text))
+				if (Widget && !Widget->GetLabelText().EqualTo(Text))
 				{
 					FWidgetBlueprintEditorUtils::RenameWidget(BlueprintEditor.Pin().ToSharedRef(), Widget->GetFName(), Text.ToString());
 				}
