PR #12008: UMG: Expose the default Frame Rate for Widget Animations in the project Widget Designer settings

#Synced-CL 52392772
#Synced-CL 52392772

diff --git a/Engine/Source/Editor/UMGEditor/Private/Settings/WidgetEditingProjectSettings.cpp b/Engine/Source/Editor/UMGEditor/Private/Settings/WidgetEditingProjectSettings.cpp
--- a/Engine/Source/Editor/UMGEditor/Private/Settings/WidgetEditingProjectSettings.cpp
+++ b/Engine/Source/Editor/UMGEditor/Private/Settings/WidgetEditingProjectSettings.cpp
@@ -32,6 +32,8 @@ UWidgetEditingProjectSettings::UWidgetEditingProjectSettings()
 	bUseUserWidgetParentClassViewerSelector = true;
 	bUseUserWidgetParentDefaultClassViewerSelector = true;
 
+	DefaultWidgetAnimationFrameRate = 20;
+
 	bUseWidgetTemplateSelector = false;
 }
 
diff --git a/Engine/Source/Editor/UMGEditor/Private/TabFactory/AnimationTabSummoner.cpp b/Engine/Source/Editor/UMGEditor/Private/TabFactory/AnimationTabSummoner.cpp
--- a/Engine/Source/Editor/UMGEditor/Private/TabFactory/AnimationTabSummoner.cpp
+++ b/Engine/Source/Editor/UMGEditor/Private/TabFactory/AnimationTabSummoner.cpp
@@ -22,6 +22,7 @@
 #include "Blueprint/WidgetTree.h"
 
 #include "UMGEditorActions.h"
+#include "UMGEditorProjectSettings.h"
 #include "UMGStyle.h"
 #include "Widgets/Input/SSearchBox.h"
 #include "Input/DragAndDrop.h"
@@ -639,8 +640,8 @@ private:
 
 		NewAnimation->MovieScene = NewObject<UMovieScene>(NewAnimation, NewFName, RF_Transactional);
 
-		// Default to 20 fps display rate (as was the previous default in USequencerSettings)
-		NewAnimation->MovieScene->SetDisplayRate(FFrameRate(20, 1));
+		const UUMGEditorProjectSettings* UMGEditorSettings = GetDefault<UUMGEditorProjectSettings>();
+		NewAnimation->MovieScene->SetDisplayRate(FFrameRate(UMGEditorSettings->DefaultWidgetAnimationFrameRate, 1));
 
 		const FFrameTime InFrame  = InTime  * NewAnimation->MovieScene->GetTickResolution();
 		const FFrameTime OutFrame = OutTime * NewAnimation->MovieScene->GetTickResolution();
diff --git a/Engine/Source/Editor/UMGEditor/Public/WidgetEditingProjectSettings.h b/Engine/Source/Editor/UMGEditor/Public/WidgetEditingProjectSettings.h
--- a/Engine/Source/Editor/UMGEditor/Public/WidgetEditingProjectSettings.h
+++ b/Engine/Source/Editor/UMGEditor/Public/WidgetEditingProjectSettings.h
@@ -220,6 +220,10 @@ public:
 	/** Use compact link button (just display the link) or complete one (with info about the current binding in the extension part of the detail view. */
 	bool bUseCompactLinkButtonInExtension = false;
 
+	/** Set the default widget animation display frame rate's numerator. */
+	UPROPERTY(EditAnywhere, config, Category = "Designer|Animation", meta = (ClampMin = "1", UIMin = "1"))
+	int32 DefaultWidgetAnimationFrameRate;
+
 	/**
 	 * The list of parent classes to choose from for newly constructed widget blueprints.
 	 * The classes must have empty widget hierarchies.
