[UMG] Fix compile error when renaming a widget targeted by a UI Component while the component is referenced in the nodegraph (UE-372796)

- UI Component BP variables are generated at compile time with the name "ComponentName_WidgetName". When renaming the Widget we didn't properly renamed the BP variable.
"ComponentName_NewWidgetName". However, UK2Node_Variable nodes in the graph still
reference the old compound name "ComponentName_OldWidgetName".

#Synced-CL 52276699

diff --git a/Engine/Source/Editor/UMGEditor/Private/UIComponentUtils.cpp b/Engine/Source/Editor/UMGEditor/Private/UIComponentUtils.cpp
--- a/Engine/Source/Editor/UMGEditor/Private/UIComponentUtils.cpp
+++ b/Engine/Source/Editor/UMGEditor/Private/UIComponentUtils.cpp
@@ -45,6 +45,14 @@ void FUIComponentUtils::OnWidgetRenamed(const TSharedRef<FWidgetBlueprintEditor>
 	}	
 }
 
+void FUIComponentUtils::ReplaceComponentVariableReferences(UWidgetBlueprint* WidgetBlueprint, const FName& OldWidgetName, const FName& NewWidgetName)
+{
+	if (UUIComponentWidgetBlueprintExtension* Extension = UWidgetBlueprintExtension::GetExtension<UUIComponentWidgetBlueprintExtension>(WidgetBlueprint))
+	{
+		Extension->ReplaceComponentVariableReferences(OldWidgetName, NewWidgetName);
+	}
+}
+
 void FUIComponentUtils::AddComponent(const TSharedRef<FWidgetBlueprintEditor>& BlueprintEditor, const UClass* ComponentClass, const FName WidgetName)
 {	
 	UWidgetBlueprint* WidgetBlueprint = BlueprintEditor->GetWidgetBlueprintObj();
diff --git a/Engine/Source/Editor/UMGEditor/Private/UIComponentWidgetBlueprintExtension.cpp b/Engine/Source/Editor/UMGEditor/Private/UIComponentWidgetBlueprintExtension.cpp
--- a/Engine/Source/Editor/UMGEditor/Private/UIComponentWidgetBlueprintExtension.cpp
+++ b/Engine/Source/Editor/UMGEditor/Private/UIComponentWidgetBlueprintExtension.cpp
@@ -84,6 +84,29 @@ void UUIComponentWidgetBlueprintExtension::RenameWidget(const FName& OldVarName,
 	}
 }
 
+void UUIComponentWidgetBlueprintExtension::ReplaceComponentVariableReferences(const FName& OldWidgetName, const FName& NewWidgetName)
+{
+	// Component variables are generated at compile time with a derived name "ComponentName_WidgetName".
+	// When a widget is renamed, the generated variable name changes, but graph node references still use
+	// the old compound name. This must be called after MarkBlueprintAsStructurallyModified to avoid
+	// BroadcastChanged() reverting the graph fixup.
+	if (ensure(ComponentContainer))
+	{
+		if (UWidgetBlueprint* WidgetBlueprint = GetWidgetBlueprint())
+		{
+			ComponentContainer->ForEachComponentTarget([&](const FUIComponentTarget& ComponentTarget)
+			{
+				if (ComponentTarget.GetTargetName() == NewWidgetName && ComponentTarget.GetComponent())
+				{
+					const FName OldCompVarName = UUIComponentContainer::GetPropertyNameForComponent(ComponentTarget.GetComponent(), OldWidgetName);
+					const FName NewCompVarName = UUIComponentContainer::GetPropertyNameForComponent(ComponentTarget.GetComponent(), NewWidgetName);
+					FBlueprintEditorUtils::ReplaceVariableReferences(WidgetBlueprint, OldCompVarName, NewCompVarName);
+				}
+			});
+		}
+	}
+}
+
 bool UUIComponentWidgetBlueprintExtension::VerifyContainer(UUserWidget* UserWidget) const
 {
 	UUIComponentUserWidgetExtension* UserWidgetExtension = UserWidget->GetExtension<UUIComponentUserWidgetExtension>();
diff --git a/Engine/Source/Editor/UMGEditor/Private/WidgetBlueprintEditorUtils.cpp b/Engine/Source/Editor/UMGEditor/Private/WidgetBlueprintEditorUtils.cpp
--- a/Engine/Source/Editor/UMGEditor/Private/WidgetBlueprintEditorUtils.cpp
+++ b/Engine/Source/Editor/UMGEditor/Private/WidgetBlueprintEditorUtils.cpp
@@ -460,6 +460,10 @@ bool FWidgetBlueprintEditorUtils::RenameWidget(TSharedRef<FWidgetBlueprintEditor
 		// Update Event References to member variables
 		FBlueprintEditorUtils::ReplaceVariableReferences(Blueprint, OldObjectName, NewFName);
 
+		// Update graph references of omponent variable names ("ComponentName_WidgetName").
+		// Must run after MarkBlueprintAsStructurallyModified to avoid BroadcastChanged() reverting changes.
+		FUIComponentUtils::ReplaceComponentVariableReferences(Blueprint, OldObjectName, NewFName);
+
 		bRenamed = true;
 	}
 
diff --git a/Engine/Source/Editor/UMGEditor/Public/UIComponentUtils.h b/Engine/Source/Editor/UMGEditor/Public/UIComponentUtils.h
--- a/Engine/Source/Editor/UMGEditor/Public/UIComponentUtils.h
+++ b/Engine/Source/Editor/UMGEditor/Public/UIComponentUtils.h
@@ -34,7 +34,8 @@ public:
 	static UE_API FClassViewerInitializationOptions CreateClassViewerInitializationOptions();
 
 	static UE_API void OnWidgetRenamed(const TSharedRef<FWidgetBlueprintEditor>& BlueprintEditor, UWidgetBlueprint* WidgetBlueprint, const FName& OldVarName, const FName& NewVarName);
-	static UE_API void AddComponent(const TSharedRef<FWidgetBlueprintEditor>& BlueprintEditor, const UClass* ComponentClass, const FName WidgetName);	
+	static UE_API void ReplaceComponentVariableReferences(UWidgetBlueprint* WidgetBlueprint, const FName& OldWidgetName, const FName& NewWidgetName);
+	static UE_API void AddComponent(const TSharedRef<FWidgetBlueprintEditor>& BlueprintEditor, const UClass* ComponentClass, const FName WidgetName);
 	static UE_API void RemoveComponent(const TSharedRef<FWidgetBlueprintEditor>& BlueprintEditor, const UClass* ComponentClass, const FName WidgetName);
 	static UE_API void MoveComponent(const TSharedRef<FWidgetBlueprintEditor>& BlueprintEditor, const UClass* ComponentClassToMove, const UClass* RelativeToComponentClass, const FName WidgetName, bool bMoveAfter);
 	
diff --git a/Engine/Source/Editor/UMGEditor/Public/UIComponentWidgetBlueprintExtension.h b/Engine/Source/Editor/UMGEditor/Public/UIComponentWidgetBlueprintExtension.h
--- a/Engine/Source/Editor/UMGEditor/Public/UIComponentWidgetBlueprintExtension.h
+++ b/Engine/Source/Editor/UMGEditor/Public/UIComponentWidgetBlueprintExtension.h
@@ -36,8 +36,9 @@ public:
 	UE_API TArray<UUIComponent*> GetComponentsFor(const UWidget* Target) const;
 	UE_API UUIComponent* GetComponent(const UClass* ComponentClass, FName OwnerName) const;
 
-	UE_API UUIComponentUserWidgetExtension* GetOrCreateExtension(UUserWidget* PreviewWidget);	
+	UE_API UUIComponentUserWidgetExtension* GetOrCreateExtension(UUserWidget* PreviewWidget);
 	UE_API void RenameWidget(const FName& OldVarName, const FName& NewVarName);
+	UE_API void ReplaceComponentVariableReferences(const FName& OldWidgetName, const FName& NewWidgetName);
 
 	UE_API bool VerifyContainer(UUserWidget* UserWidget) const;
 
