Fox to SReadOnlyHierarchyView to support enabling or disabling inherited NamedSlot in the tree view.

#Synced-CL 52579006

diff --git a/Engine/Source/Editor/UMGEditor/Private/Customizations/WidgetChildTypeCustomization.cpp b/Engine/Source/Editor/UMGEditor/Private/Customizations/WidgetChildTypeCustomization.cpp
--- a/Engine/Source/Editor/UMGEditor/Private/Customizations/WidgetChildTypeCustomization.cpp
+++ b/Engine/Source/Editor/UMGEditor/Private/Customizations/WidgetChildTypeCustomization.cpp
@@ -135,9 +135,10 @@ TSharedRef<SWidget> FWidgetChildTypeCustomization::GetPopupContent()
 			.MinDesiredHeight(MinPopupHeight)
 			[
 				SNew(SReadOnlyHierarchyView, WidgetBlueprint)
-				.OnSelectionChanged(this, &FWidgetChildTypeCustomization::OnWidgetSelectionChanged)
-				.ShowSearch(true)
-				.RootSelectionMode(ERootSelectionMode::Self)
+					.OnSelectionChanged(this, &FWidgetChildTypeCustomization::OnWidgetSelectionChanged)
+					.ShowSearch(true)
+					.RootSelectionMode(ERootSelectionMode::Self)
+					.InheritedNamedSlotSelectionMode(EInheritedNamedSlotSelectionMode::Disabled)
 			]
 		];
 }
diff --git a/Engine/Source/Editor/UMGEditor/Private/Hierarchy/SReadOnlyHierarchyView.cpp b/Engine/Source/Editor/UMGEditor/Private/Hierarchy/SReadOnlyHierarchyView.cpp
--- a/Engine/Source/Editor/UMGEditor/Private/Hierarchy/SReadOnlyHierarchyView.cpp
+++ b/Engine/Source/Editor/UMGEditor/Private/Hierarchy/SReadOnlyHierarchyView.cpp
@@ -18,6 +18,7 @@ void SReadOnlyHierarchyView::Construct(const FArguments& InArgs, const UWidgetBl
 	WidgetBlueprint = InWidgetBlueprint;
 	ShowOnly = InArgs._ShowOnly;
 	RootSelectionMode = InArgs._RootSelectionMode;
+	InheritedNamedSlotSelectionMode = InArgs._InheritedNamedSlotSelectionMode;
 	bExpandAll = InArgs._ExpandAll;
 
 	SearchFilter = MakeShared<FTextFilter>(FTextFilter::FItemToStringArray::CreateSP(this, &SReadOnlyHierarchyView::GetFilterStringsForItem));
@@ -211,7 +212,7 @@ TSharedRef<ITableRow> SReadOnlyHierarchyView::GenerateRow(TSharedPtr<FItem> Item
 	}
 
 	const FName NamedSlotName = Item->InheritedNamedSlotName;
-	if (!NamedSlotName.IsNone())
+	if (!NamedSlotName.IsNone() && InheritedNamedSlotSelectionMode == EInheritedNamedSlotSelectionMode::Disabled)
 	{
 		bIsEnabled = false;
 	}
diff --git a/Engine/Source/Editor/UMGEditor/Public/Hierarchy/SReadOnlyHierarchyView.h b/Engine/Source/Editor/UMGEditor/Public/Hierarchy/SReadOnlyHierarchyView.h
--- a/Engine/Source/Editor/UMGEditor/Public/Hierarchy/SReadOnlyHierarchyView.h
+++ b/Engine/Source/Editor/UMGEditor/Public/Hierarchy/SReadOnlyHierarchyView.h
@@ -35,6 +35,16 @@ enum class ERootSelectionMode : uint8
 	Self
 };
 
+UENUM()
+enum class EInheritedNamedSlotSelectionMode : uint8
+{
+	/** Inherited Named slots are not selectable */
+	Disabled,
+
+	/** Inherited Named slots are selectable */
+	Enabled,
+};
+
 class SReadOnlyHierarchyView : public SCompoundWidget
 {
 private:
@@ -59,6 +69,7 @@ public:
 	SLATE_BEGIN_ARGS(SReadOnlyHierarchyView) {}
 		SLATE_ARGUMENT_DEFAULT(bool, ShowSearch) = true;
 		SLATE_ARGUMENT_DEFAULT(ERootSelectionMode, RootSelectionMode) = ERootSelectionMode::Enabled;
+		SLATE_ARGUMENT_DEFAULT(EInheritedNamedSlotSelectionMode, InheritedNamedSlotSelectionMode) = EInheritedNamedSlotSelectionMode::Enabled;		
 		SLATE_ARGUMENT_DEFAULT(ESelectionMode::Type, SelectionMode) = ESelectionMode::Single;
 		SLATE_EVENT(FOnSelectionChanged, OnSelectionChanged)
 		SLATE_ARGUMENT(TArray<FName>, ShowOnly)
@@ -106,6 +117,7 @@ private:
 	TArray<TSharedPtr<FItem>> FilteredRootWidgets;
 
 	ERootSelectionMode RootSelectionMode;
+	EInheritedNamedSlotSelectionMode InheritedNamedSlotSelectionMode;
 
 	FOnSelectionChanged OnSelectionChangedDelegate;
 
