[ContentBrowser] New Add Menu UserInterface Menu

#Synced-CL 52619962

diff --git a/Engine/Plugins/Editor/EngineAssetDefinitions/Source/Private/AssetDefinition_Font.h b/Engine/Plugins/Editor/EngineAssetDefinitions/Source/Private/AssetDefinition_Font.h
--- a/Engine/Plugins/Editor/EngineAssetDefinitions/Source/Private/AssetDefinition_Font.h
+++ b/Engine/Plugins/Editor/EngineAssetDefinitions/Source/Private/AssetDefinition_Font.h
@@ -19,7 +19,7 @@ public:
 	virtual TSoftClassPtr<UObject> GetAssetClass() const override { return UFont::StaticClass(); }
 	virtual TConstArrayView<FAssetCategoryPath> GetAssetCategories() const override
 	{
-		static const auto Categories = { EAssetCategoryPaths::UI };
+		static const auto Categories = { FAssetCategoryPath(EAssetCategoryPaths::UI, NSLOCTEXT("AssetDefinition", "Font_SubMenu", "Font"), ECategoryMenuType::Section) };
 		return Categories;
 	}
 	
diff --git a/Engine/Plugins/Editor/EngineAssetDefinitions/Source/Private/AssetDefinition_SlateBrush.h b/Engine/Plugins/Editor/EngineAssetDefinitions/Source/Private/AssetDefinition_SlateBrush.h
--- a/Engine/Plugins/Editor/EngineAssetDefinitions/Source/Private/AssetDefinition_SlateBrush.h
+++ b/Engine/Plugins/Editor/EngineAssetDefinitions/Source/Private/AssetDefinition_SlateBrush.h
@@ -19,7 +19,7 @@ public:
 	virtual TSoftClassPtr<UObject> GetAssetClass() const override { return USlateBrushAsset::StaticClass(); }
 	virtual TConstArrayView<FAssetCategoryPath> GetAssetCategories() const override
 	{
-		static const auto Categories = { EAssetCategoryPaths::UI };
+		static const auto Categories = { FAssetCategoryPath(EAssetCategoryPaths::UI, NSLOCTEXT("AssetDefinition", "SlateBrush_SubMenu", "Slate"), ECategoryMenuType::Section) };
 		return Categories;
 	}
 };
diff --git a/Engine/Plugins/Editor/EngineAssetDefinitions/Source/Private/AssetDefinition_SlateWidgetStyle.h b/Engine/Plugins/Editor/EngineAssetDefinitions/Source/Private/AssetDefinition_SlateWidgetStyle.h
--- a/Engine/Plugins/Editor/EngineAssetDefinitions/Source/Private/AssetDefinition_SlateWidgetStyle.h
+++ b/Engine/Plugins/Editor/EngineAssetDefinitions/Source/Private/AssetDefinition_SlateWidgetStyle.h
@@ -19,7 +19,7 @@ public:
 	virtual TSoftClassPtr<UObject> GetAssetClass() const override { return USlateWidgetStyleAsset::StaticClass(); }
 	virtual TConstArrayView<FAssetCategoryPath> GetAssetCategories() const override
 	{
-		static const auto Categories = { EAssetCategoryPaths::UI };
+		static const auto Categories = { FAssetCategoryPath(EAssetCategoryPaths::UI, NSLOCTEXT("AssetDefinition", "SlateStyle_SubMenu", "Slate"), ECategoryMenuType::Section) };
 		return Categories;
 	}
 	
diff --git a/Engine/Plugins/Editor/UMGWidgetPreview/Source/UMGWidgetPreview/Private/AssetDefinition_WidgetPreview.cpp b/Engine/Plugins/Editor/UMGWidgetPreview/Source/UMGWidgetPreview/Private/AssetDefinition_WidgetPreview.cpp
--- a/Engine/Plugins/Editor/UMGWidgetPreview/Source/UMGWidgetPreview/Private/AssetDefinition_WidgetPreview.cpp
+++ b/Engine/Plugins/Editor/UMGWidgetPreview/Source/UMGWidgetPreview/Private/AssetDefinition_WidgetPreview.cpp
@@ -32,7 +32,7 @@ TSoftClassPtr<> UAssetDefinition_WidgetPreview::GetAssetClass() const
 
 TConstArrayView<FAssetCategoryPath> UAssetDefinition_WidgetPreview::GetAssetCategories() const
 {
-	static const TArray<FAssetCategoryPath, TFixedAllocator<1>> Categories = { EAssetCategoryPaths::UI };
+	static const auto Categories = { FAssetCategoryPath(EAssetCategoryPaths::UI, LOCTEXT("WidgetPreview_SubMenu", "Widget"), ECategoryMenuType::Section) };
 	return Categories;
 }
 
diff --git a/Engine/Source/Editor/UMGEditor/Private/AssetDefinition_SlateVectorArtData.h b/Engine/Source/Editor/UMGEditor/Private/AssetDefinition_SlateVectorArtData.h
new file mode 100644
--- /dev/null
+++ b/Engine/Source/Editor/UMGEditor/Private/AssetDefinition_SlateVectorArtData.h
@@ -0,0 +1,23 @@
+// Copyright Epic Games, Inc. All Rights Reserved.
+
+#pragma once
+#include "AssetDefinitionDefault.h"
+#include "Slate/SlateVectorArtData.h"
+#include "AssetDefinition_SlateVectorArtData.generated.h"
+
+UCLASS()
+class UAssetDefinition_SlateVectorArtData : public UAssetDefinitionDefault
+{
+	GENERATED_BODY()
+public:
+	// UAssetDefinition Begin
+	virtual FText GetAssetDisplayName() const override { return NSLOCTEXT("AssetDefinition", "AssetDefinition_SlateVectorArtData", "Slate Vector Art Data"); }
+	virtual FLinearColor GetAssetColor() const override { return FColor(105, 165, 60); }
+	virtual TSoftClassPtr<UObject> GetAssetClass() const override { return USlateVectorArtData::StaticClass(); }
+	virtual TConstArrayView<FAssetCategoryPath> GetAssetCategories() const override
+	{
+		static const auto Categories = { FAssetCategoryPath(EAssetCategoryPaths::UI, NSLOCTEXT("AssetDefinition", "AssetDefinition_SlateVectorArtDataSubMenu", "Slate"), ECategoryMenuType::Section) };
+		return Categories;
+	}
+	// UAssetDefinition End
+};
diff --git a/Engine/Source/Editor/UMGEditor/Private/AssetDefinition_WidgetBlueprint.cpp b/Engine/Source/Editor/UMGEditor/Private/AssetDefinition_WidgetBlueprint.cpp
--- a/Engine/Source/Editor/UMGEditor/Private/AssetDefinition_WidgetBlueprint.cpp
+++ b/Engine/Source/Editor/UMGEditor/Private/AssetDefinition_WidgetBlueprint.cpp
@@ -31,7 +31,7 @@ TSoftClassPtr<> UAssetDefinition_WidgetBlueprint::GetAssetClass() const
 
 TConstArrayView<FAssetCategoryPath> UAssetDefinition_WidgetBlueprint::GetAssetCategories() const
 {
-	static const TArray<FAssetCategoryPath, TFixedAllocator<1>> Categories = { EAssetCategoryPaths::UI };
+	static const auto Categories = { FAssetCategoryPath(EAssetCategoryPaths::UI, LOCTEXT("WidgetBlueprint_SubMenu", "Widget"), ECategoryMenuType::Section) };
 	return Categories;
 }
 
diff --git a/Engine/Source/Editor/UMGEditor/Private/AssetTypeActions_SlateVectorArtData.h b/Engine/Source/Editor/UMGEditor/Private/AssetTypeActions_SlateVectorArtData.h
deleted file mode 100644
diff --git a/Engine/Source/Editor/UMGEditor/Private/UMGEditorModule.cpp b/Engine/Source/Editor/UMGEditor/Private/UMGEditorModule.cpp
--- a/Engine/Source/Editor/UMGEditor/Private/UMGEditorModule.cpp
+++ b/Engine/Source/Editor/UMGEditor/Private/UMGEditorModule.cpp
@@ -17,7 +17,6 @@
 
 #include "AssetToolsModule.h"
 #include "IAssetTypeActions.h"
-#include "AssetTypeActions_SlateVectorArtData.h"
 #include "KismetCompilerModule.h"
 #include "WidgetBlueprintCompiler.h"
 
@@ -130,10 +129,6 @@ public:
 			ModuleChangedHandle = FModuleManager::Get().OnModulesChanged().AddRaw(this, &FUMGEditorModule::HandleModuleChanged);
 		}
 
-		// Register asset types
-		IAssetTools& AssetTools = FModuleManager::LoadModuleChecked<FAssetToolsModule>("AssetTools").Get();
-		RegisterAssetTypeAction(AssetTools, MakeShareable(new FAssetTypeActions_SlateVectorArtData()));
-		
 		FKismetCompilerContext::RegisterCompilerForBP(UWidgetBlueprint::StaticClass(), &UWidgetBlueprint::GetCompilerForWidgetBP );
 
 		// Register with the sequencer module that we provide auto-key handlers.
